Clean up typos and phrasing in Interchange modules - contribution from Robb Surridge

diff --git a/Engine/Plugins/Interchange/Runtime/Source/Nodes/Public/InterchangeMeshNode.h b/Engine/Plugins/Interchange/Runtime/Source/Nodes/Public/InterchangeMeshNode.h
--- a/Engine/Plugins/Interchange/Runtime/Source/Nodes/Public/InterchangeMeshNode.h
+++ b/Engine/Plugins/Interchange/Runtime/Source/Nodes/Public/InterchangeMeshNode.h
@@ -77,17 +77,17 @@ public:
 #endif //WITH_EDITOR
 
 	/**
-	 * Return the node type name of the class, we use this when reporting error
+	 * Return the node type name of the class. This is used when reporting errors.
 	 */
 	virtual FString GetTypeName() const override;
 
 	/**
-	 * Icon name are simply create by adding "InterchangeIcon_" in front of the specialized type. If there is no special type the function will return NAME_None which will use the default icon.
+	 * Icon names are created by adding "InterchangeIcon_" in front of the specialized type. If there is no special type, the function will return NAME_None, which will use the default icon.
 	 */
 	virtual FName GetIconName() const override;
 	
 	/**
-	 * Override serialize to restore SlotMaterialDependencies on load.
+	 * Override Serialize() to restore SlotMaterialDependencies on load.
 	 */
 	virtual void Serialize(FArchive& Ar) override
 	{
@@ -100,224 +100,223 @@ public:
 	}
 
 	/**
-	 * Return true if this node represent a skinned mesh
+	 * Return true if this node represents a skinned mesh.
 	 */
 	UFUNCTION(BlueprintCallable, Category = "Interchange | Node | Mesh")
 	bool IsSkinnedMesh() const;
 
 	/**
-	 * Set the IsSkinnedMesh attribute to determine if this node represent a skinned mesh.
+	 * Set the IsSkinnedMesh attribute to determine whether this node represents a skinned mesh.
 	 */
 	UFUNCTION(BlueprintCallable, Category = "Interchange | Node | Mesh")
 	bool SetSkinnedMesh(const bool bIsSkinnedMesh);
 
 	/**
-	 * Return true if this node represent a morph target
+	 * Return true if this node represents a morph target.
 	 */
 	UFUNCTION(BlueprintCallable, Category = "Interchange | Node | Mesh")
 	bool IsMorphTarget() const;
 
 	/**
-	 * Set the IsMorphTarget attribute to determine if this node represent a morph target.
+	 * Set the IsMorphTarget attribute to determine whether this node represents a morph target.
 	 */
 	UFUNCTION(BlueprintCallable, Category = "Interchange | Node | Mesh")
 	bool SetMorphTarget(const bool bIsMorphTarget);
 
 	/**
 	 * Get the morph target name.
-	 * Return true if we successfully query the MorphTargetName attribute
+	 * Return true if we successfully retrieved the MorphTargetName attribute.
 	 */
 	UFUNCTION(BlueprintCallable, Category = "Interchange | Node | Mesh")
 	bool GetMorphTargetName(FString& OutMorphTargetName) const;
 
 	/**
-	 * Set the MorphTargetName attribute to determine if this node represent a morph target.
+	 * Set the MorphTargetName attribute to determine the name of the morph target.
 	 */
 	UFUNCTION(BlueprintCallable, Category = "Interchange | Node | Mesh")
 	bool SetMorphTargetName(const FString& MorphTargetName);
 
-	/** Mesh node Interface Begin */
 	virtual const TOptional<FInterchangeMeshPayLoadKey> GetPayLoadKey() const;
 
 	UFUNCTION(BlueprintCallable, Category = "Interchange | Node | Mesh")
 	virtual void SetPayLoadKey(const FString& PayLoadKey, const EInterchangeMeshPayLoadType& PayLoadType);
 	
-	/** Query this mesh vertices count. Return false if the attribute was not set.*/
+	/** Query the vertex count of this mesh. Return false if the attribute was not set.*/
 	UFUNCTION(BlueprintCallable, Category = "Interchange | Node | Mesh")
 	bool GetCustomVertexCount(int32& AttributeValue) const;
 
-	/** Set this mesh vertices count. Return false if the attribute cannot be set.*/
+	/** Set the vertex count of this mesh. Return false if the attribute could not be set.*/
 	UFUNCTION(BlueprintCallable, Category = "Interchange | Node | Mesh")
 	bool SetCustomVertexCount(const int32& AttributeValue);
 
-	/** Query this mesh polygon count. Return false if the attribute was not set.*/
+	/** Query the polygon count of this mesh. Return false if the attribute was not set.*/
 	UFUNCTION(BlueprintCallable, Category = "Interchange | Node | Mesh")
 	bool GetCustomPolygonCount(int32& AttributeValue) const;
 
-	/** Set this mesh polygon count. Return false if the attribute cannot be set.*/
+	/** Set the polygon count of this mesh. Return false if the attribute could not be set.*/
 	UFUNCTION(BlueprintCallable, Category = "Interchange | Node | Mesh")
 	bool SetCustomPolygonCount(const int32& AttributeValue);
 
-	/** Query this mesh bounding box. Return false if the attribute was not set.*/
+	/** Query the bounding box of this mesh. Return false if the attribute was not set.*/
 	UFUNCTION(BlueprintCallable, Category = "Interchange | Node | Mesh")
 	bool GetCustomBoundingBox(FBox& AttributeValue) const;
 
-	/** Set this mesh bounding box. Return false if the attribute cannot be set.*/
+	/** Set the bounding box of this mesh. Return false if the attribute could not be set.*/
 	UFUNCTION(BlueprintCallable, Category = "Interchange | Node | Mesh")
 	bool SetCustomBoundingBox(const FBox& AttributeValue);
 
-	/** Query if this mesh has vertex normal. Return false if the attribute was not set.*/
+	/** Query whether this mesh has vertex normals. Return false if the attribute was not set.*/
 	UFUNCTION(BlueprintCallable, Category = "Interchange | Node | Mesh")
 	bool GetCustomHasVertexNormal(bool& AttributeValue) const;
 
-	/** Set this mesh has vertex normal attribute. Return false if the attribute cannot be set.*/
+	/** Set the vertex normal attribute of this mesh. Return false if the attribute could not be set.*/
 	UFUNCTION(BlueprintCallable, Category = "Interchange | Node | Mesh")
 	bool SetCustomHasVertexNormal(const bool& AttributeValue);
 
-	/** Query if this mesh has vertex bi-normal. Return false if the attribute was not set.*/
+	/** Query whether this mesh has vertex bi-normals. Return false if the attribute was not set.*/
 	UFUNCTION(BlueprintCallable, Category = "Interchange | Node | Mesh")
 	bool GetCustomHasVertexBinormal(bool& AttributeValue) const;
 
-	/** Set this mesh has vertex bi-normal attribute. Return false if the attribute cannot be set.*/
+	/** Set the vertex bi-normal attribute of this mesh. Return false if the attribute could not be set.*/
 	UFUNCTION(BlueprintCallable, Category = "Interchange | Node | Mesh")
 	bool SetCustomHasVertexBinormal(const bool& AttributeValue);
 
-	/** Query if this mesh has vertex tangent. Return false if the attribute was not set.*/
+	/** Query whether this mesh has vertex tangents. Return false if the attribute was not set.*/
 	UFUNCTION(BlueprintCallable, Category = "Interchange | Node | Mesh")
 	bool GetCustomHasVertexTangent(bool& AttributeValue) const;
 
-	/** Set this mesh has vertex tangent attribute. Return false if the attribute cannot be set.*/
+	/** Set the vertex tangent attribute of this mesh. Return false if the attribute could not be set.*/
 	UFUNCTION(BlueprintCallable, Category = "Interchange | Node | Mesh")
 	bool SetCustomHasVertexTangent(const bool& AttributeValue);
 
-	/** Query if this mesh has smooth group. Return false if the attribute was not set.*/
+	/** Query whether this mesh has smoothing groups. Return false if the attribute was not set.*/
 	UFUNCTION(BlueprintCallable, Category = "Interchange | Node | Mesh")
 	bool GetCustomHasSmoothGroup(bool& AttributeValue) const;
 
-	/** Set this mesh has smooth group attribute. Return false if the attribute cannot be set.*/
+	/** Set the smoothing group attribute of this mesh. Return false if the attribute could not be set.*/
 	UFUNCTION(BlueprintCallable, Category = "Interchange | Node | Mesh")
 	bool SetCustomHasSmoothGroup(const bool& AttributeValue);
 
-	/** Query if this mesh has vertex color. Return false if the attribute was not set.*/
+	/** Query whether this mesh has vertex colors. Return false if the attribute was not set.*/
 	UFUNCTION(BlueprintCallable, Category = "Interchange | Node | Mesh")
 	bool GetCustomHasVertexColor(bool& AttributeValue) const;
 
-	/** Set this mesh has vertex color attribute. Return false if the attribute cannot be set.*/
+	/** Set the vertex color attribute of this mesh. Return false if the attribute could not be set.*/
 	UFUNCTION(BlueprintCallable, Category = "Interchange | Node | Mesh")
 	bool SetCustomHasVertexColor(const bool& AttributeValue);
 
-	/** Query this mesh UV count. Return false if the attribute was not set.*/
+	/** Query the UV count of this mesh. Return false if the attribute was not set.*/
 	UFUNCTION(BlueprintCallable, Category = "Interchange | Node | Mesh")
 	bool GetCustomUVCount(int32& AttributeValue) const;
 
-	/** Set this mesh UV count attribute. Return false if the attribute cannot be set.*/
+	/** Set the UV count attribute of this mesh. Return false if the attribute could not be set.*/
 	UFUNCTION(BlueprintCallable, Category = "Interchange | Node | Mesh")
 	bool SetCustomUVCount(const int32& AttributeValue);
 
 	/**
-	 * This function allow to retrieve the number of skeleton dependencies for this object.
+	 * Retrieve the number of skeleton dependencies for this object.
 	 */
 	UFUNCTION(BlueprintCallable, Category = "Interchange | Node | Mesh")
 	int32 GetSkeletonDependeciesCount() const;
 
 	/**
-	 * This function allow to retrieve the skeleton dependency for this object.
+	 * Retrieve the skeleton dependency for this object.
 	 */
 	UFUNCTION(BlueprintCallable, Category = "Interchange | Node | Mesh")
 	void GetSkeletonDependencies(TArray<FString>& OutDependencies) const;
 
 	/**
-	 * This function allow to retrieve one skeleton dependency for this object.
+	 * Retrieve the specified skeleton dependency for this object.
 	 */
 	UFUNCTION(BlueprintCallable, Category = "Interchange | Node | Mesh")
 	void GetSkeletonDependency(const int32 Index, FString& OutDependency) const;
 
 	/**
-	 * Add one skeleton dependency to this object.
+	 * Add the specified skeleton dependency to this object.
 	 */
 	UFUNCTION(BlueprintCallable, Category = "Interchange | Node | Mesh")
 	bool SetSkeletonDependencyUid(const FString& DependencyUid);
 
 	/**
-	 * Remove one skeleton dependency from this object.
+	 * Remove the specified skeleton dependency from this object.
 	 */
 	UFUNCTION(BlueprintCallable, Category = "Interchange | Node | Mesh")
 	bool RemoveSkeletonDependencyUid(const FString& DependencyUid);
 
 	/**
-	 * This function allow to retrieve the number of morph target dependencies for this object.
+	 * Retrieve the number of morph target dependencies for this object.
 	 */
 	UFUNCTION(BlueprintCallable, Category = "Interchange | Node | Mesh")
 	int32 GetMorphTargetDependeciesCount() const;
 
 	/**
-	 * This function allow to retrieve the morph target dependency for this object.
+	 * Retrieve all morph target dependencies for this object.
 	 */
 	UFUNCTION(BlueprintCallable, Category = "Interchange | Node | Mesh")
 	void GetMorphTargetDependencies(TArray<FString>& OutDependencies) const;
 
 	/**
-	 * This function allow to retrieve one morph target dependency for this object.
+	 * Retrieve the specified morph target dependency for this object.
 	 */
 	UFUNCTION(BlueprintCallable, Category = "Interchange | Node | Mesh")
 	void GetMorphTargetDependency(const int32 Index, FString& OutDependency) const;
 
 	/**
-	 * Add one morph target dependency to this object.
+	 * Add the specified morph target dependency to this object.
 	 */
 	UFUNCTION(BlueprintCallable, Category = "Interchange | Node | Mesh")
 	bool SetMorphTargetDependencyUid(const FString& DependencyUid);
 
 	/**
-	 * Remove one morph target dependency from this object.
+	 * Remove the specified morph target dependency from this object.
 	 */
 	UFUNCTION(BlueprintCallable, Category = "Interchange | Node | Mesh")
 	bool RemoveMorphTargetDependencyUid(const FString& DependencyUid);
 
 	/**
-	 * This function allow to retrieve the number of scene node instancing this mesh.
+	 * Retrieve the number of scene nodes instancing this mesh.
 	 */
 	UFUNCTION(BlueprintCallable, Category = "Interchange | Node | Mesh")
 	int32 GetSceneInstanceUidsCount() const;
 
 	/**
-	 * This function allow to retrieve the asset instances this scene node is refering.
+	 * Retrieve the asset instances this scene node refers to.
 	 */
 	UFUNCTION(BlueprintCallable, Category = "Interchange | Node | Mesh")
 	void GetSceneInstanceUids(TArray<FString>& OutDependencies) const;
 
 	/**
-	 * This function allow to retrieve an asset instance this scene node is refering.
+	 * Retrieve the asset instance this scene node refers to.
 	 */
 	UFUNCTION(BlueprintCallable, Category = "Interchange | Node | Mesh")
 	void GetSceneInstanceUid(const int32 Index, FString& OutDependency) const;
 
 	/**
-	 * Add one asset instance this scene node is refering.
+	 * Add the specified asset instance this scene node refers to.
 	 */
 	UFUNCTION(BlueprintCallable, Category = "Interchange | Node | Mesh")
 	bool SetSceneInstanceUid(const FString& DependencyUid);
 
 	/**
-	 * Remove one asset instance this scene node is refering.
+	 * Remove the specified asset instance this scene node refers to.
 	 */
 	UFUNCTION(BlueprintCallable, Category = "Interchange | Node | Mesh")
 	bool RemoveSceneInstanceUid(const FString& DependencyUid);
 
 	/**
-	 * Allow to retrieve the correspondence table between slot names and assigned materials for this object.
+	 * Retrieve the correspondence table between slot names and assigned materials for this object.
 	 */
 	UFUNCTION(BlueprintCallable, Category = "Interchange | Node | StaticMesh")
 	void GetSlotMaterialDependencies(TMap<FString, FString>& OutMaterialDependencies) const;
 
 	/**
-	 * Allow to retrieve one Material dependency for a given slot of this object.
+	 * Retrieve the specified Material dependency for a given slot of this object.
 	 */
 	UFUNCTION(BlueprintCallable, Category = "Interchange | Node | StaticMesh")
 	bool GetSlotMaterialDependencyUid(const FString& SlotName, FString& OutMaterialDependency) const;
 
 	/**
-	 * Add one Material dependency to a specific slot name of this object.
+	 * Add the specified Material dependency to a specific slot name of this object.
 	 */
 	UFUNCTION(BlueprintCallable, Category = "Interchange | Node | StaticMesh")
 	bool SetSlotMaterialDependencyUid(const FString& SlotName, const FString& MaterialDependencyUid);
diff --git a/Engine/Plugins/Interchange/Runtime/Source/Nodes/Public/InterchangeSceneNode.h b/Engine/Plugins/Interchange/Runtime/Source/Nodes/Public/InterchangeSceneNode.h
--- a/Engine/Plugins/Interchange/Runtime/Source/Nodes/Public/InterchangeSceneNode.h
+++ b/Engine/Plugins/Interchange/Runtime/Source/Nodes/Public/InterchangeSceneNode.h
@@ -30,8 +30,8 @@ namespace UE
 }//ns UE
 
 /**
- * The scene node represent a transform node in the scene
- * Scene node can have user defined attribute. Use UInterchangeUserDefinedAttributesAPI to get\set user define attribute data
+ * The scene node represents a transform node in the scene.
+ * Scene nodes can have user-defined attributes. Use UInterchangeUserDefinedAttributesAPI to get and set user-defined attribute data.
  */
 UCLASS(BlueprintType)
 class INTERCHANGENODES_API UInterchangeSceneNode : public UInterchangeBaseNode
@@ -42,7 +42,7 @@ public:
 	UInterchangeSceneNode();
 
 	/**
-	 * Override serialize to restore SlotMaterialDependencies on load.
+	 * Override Serialize() to restore SlotMaterialDependencies on load.
 	 */
 	virtual void Serialize(FArchive& Ar) override
 	{
@@ -55,7 +55,7 @@ public:
 	}
 
 	/**
-	 * Return the node type name of the class, we use this when reporting error
+	 * Return the node type name of the class. This is used when reporting errors.
 	 */
 	virtual FString GetTypeName() const override;
 
@@ -65,14 +65,14 @@ public:
 	virtual FString GetAttributeCategory(const UE::Interchange::FAttributeKey& NodeAttributeKey) const override;
 #endif //WITH_EDITOR
 	/**
-	 * Icon name are simply create by adding "InterchangeIcon_" in front of the specialized type. If there is no special type the function will return NAME_None which will use the default icon.
+	 * Icon names are created by adding "InterchangeIcon_" in front of the specialized type. If there is no special type, the function will return NAME_None, which will use the default icon.
 	 */
 	virtual FName GetIconName() const override;
 
 	UFUNCTION(BlueprintCallable, Category = "Interchange | Node | Scene")
 	bool IsSpecializedTypeContains(const FString& SpecializedType) const;
 
-	/** Get the Specialized type this scene node represent (Joint, LODGroup, ...).*/
+	/** Get the specialized type this scene node represents (for example, Joint or LODGroup). */
 	UFUNCTION(BlueprintCallable, Category = "Interchange | Node | Scene")
 	int32 GetSpecializedTypeCount() const;
 
@@ -92,35 +92,35 @@ public:
 	UFUNCTION(BlueprintCallable, Category = "Interchange | Node | Scene")
 	bool GetCustomAssetInstanceUid(FString& AttributeValue) const;
 
-	/** Add asset this scene node is instantiating */
+	/** Add an asset for this scene node to instantiate. */
 	UFUNCTION(BlueprintCallable, Category = "Interchange | Node | Scene")
 	bool SetCustomAssetInstanceUid(const FString& AttributeValue);
 
 
 	/**
 	 * Get the default scene node local transform.
-	 * Default transform is the local transform we have in the node(no bind pose, no time evaluation).
+	 * The default transform is the local transform of the node (no bind pose, no time evaluation).
 	 */
 	UFUNCTION(BlueprintCallable, Category = "Interchange | Node | Scene")
 	bool GetCustomLocalTransform(FTransform& AttributeValue) const;
 
 	/**
 	 * Set the default scene node local transform.
-	 * Default transform is the local transform we have in the node(no bind pose, no time evaluation).
+	 * The default transform is the local transform of the node (no bind pose, no time evaluation).
 	 */
 	UFUNCTION(BlueprintCallable, Category = "Interchange | Node | Scene")
 	bool SetCustomLocalTransform(const UInterchangeBaseNodeContainer* BaseNodeContainer, const FTransform& AttributeValue, bool bResetCache = true);
 
-	/** Get the default scene node global transform. This value is computed with all parent local transform. */
+	/** Get the default scene node global transform. This value is computed from the local transforms of all parent scene nodes. */
 	UFUNCTION(BlueprintCallable, Category = "Interchange | Node | Scene")
 	bool GetCustomGlobalTransform(const UInterchangeBaseNodeContainer* BaseNodeContainer, const FTransform& GlobalOffsetTransform, FTransform& AttributeValue, bool bForceRecache = false) const;
 
 
-	/** Get the geometric offset. Any mesh attach to this scene node will be offset using this transform. */
+	/** Get the geometric offset. Any mesh attached to this scene node will be offset using this transform. */
 	UFUNCTION(BlueprintCallable, Category = "Interchange | Node | Scene")
 	bool GetCustomGeometricTransform(FTransform& AttributeValue) const;
 
-	/** Set the geometric offset. Any mesh attach to this scene node will be offset using this transform. */
+	/** Set the geometric offset. Any mesh attached to this scene node will be offset using this transform. */
 	UFUNCTION(BlueprintCallable, Category = "Interchange | Node | Scene")
 	bool SetCustomGeometricTransform(const FTransform& AttributeValue);
 
@@ -130,36 +130,36 @@ public:
 	* Bind pose transform is the transform of the joint when the binding with the mesh was done.
 	* This attribute should be set only if this scene node represent a joint.
 	* 
-	* Time zero transform is the transform of the node at time zero.
-	* Pipeline often have the option to evaluate the joint at time zero to create the bind pose.
-	* Time zero bind pose is also use if the translator did not find any bind pose or if we import
-	* unskinned mesh has a skeletalmesh (rigid mesh)
+	* Time-zero transform is the transform of the node at time zero.
+	* Pipelines often have the option to evaluate the joint at time zero to create the bind pose.
+	* Time-zero bind pose is also used if the translator did not find any bind pose, or if we import
+	* an unskinned mesh as a skeletal smesh (rigid mesh).
 	*/
 
-	/** Get the bind pose scene node local transform. */
+	/** Get the local transform of the bind pose scene node. */
 	UFUNCTION(BlueprintCallable, Category = "Interchange | Node | Joint")
 	bool GetCustomBindPoseLocalTransform(FTransform& AttributeValue) const;
 
-	/** Set the bind pose scene node local transform. */
+	/** Set the local transform of the bind pose scene node. */
 	UFUNCTION(BlueprintCallable, Category = "Interchange | Node | Joint")
 	bool SetCustomBindPoseLocalTransform(const UInterchangeBaseNodeContainer* BaseNodeContainer, const FTransform& AttributeValue, bool bResetCache = true);
 
-	/** Get the bind pose scene node global transform. This value is computed with all parent bind pose local transform. */
+	/** Get the global transform of the bind pose scene node. This value is computed from the local transforms of all parent bind poses. */
 	UFUNCTION(BlueprintCallable, Category = "Interchange | Node | Joint")
 	bool GetCustomBindPoseGlobalTransform(const UInterchangeBaseNodeContainer* BaseNodeContainer, const FTransform& GlobalOffsetTransform, FTransform& AttributeValue, bool bForceRecache = false) const;
 
 	//Time zero transform is the transform of the node at time zero.
 	//This is useful when there is no bind pose or when we import rigid mesh.
 
-	/** Get the time zero scene node local transform. */
+	/** Get the local transform of the time-zero scene node. */
 	UFUNCTION(BlueprintCallable, Category = "Interchange | Node | Joint")
 	bool GetCustomTimeZeroLocalTransform(FTransform& AttributeValue) const;
 
-	/** Set the time zero scene node local transform. */
+	/** Set the local transform of the time-zero scene node. */
 	UFUNCTION(BlueprintCallable, Category = "Interchange | Node | Joint")
 	bool SetCustomTimeZeroLocalTransform(const UInterchangeBaseNodeContainer* BaseNodeContainer, const FTransform& AttributeValue, bool bResetCache = true);
 
-	/** Get the time zero scene node global transform. This value is computed with all parent timezero local transform. */
+	/** Get the global transform of the time-zero scene node. This value is computed from the local transforms of all parent time-zero scene nodes. */
 	UFUNCTION(BlueprintCallable, Category = "Interchange | Node | Joint")
 	bool GetCustomTimeZeroGlobalTransform(const UInterchangeBaseNodeContainer* BaseNodeContainer, const FTransform& GlobalOffsetTransform, FTransform& AttributeValue, bool bForceRecache = false) const;
 
@@ -167,26 +167,26 @@ public:
 	* Skeleton bone API End
 	***********************************************************************************************/
 
-	/** This static function make sure all the global transform caches are reset for all the UInterchangeSceneNode nodes in the UInterchangeBaseNodeContainer */
+	/** This static function ensures all the global transform caches are reset for all the UInterchangeSceneNode nodes in the UInterchangeBaseNodeContainer. */
 	static void ResetAllGlobalTransformCaches(const UInterchangeBaseNodeContainer* BaseNodeContainer);
 
-	/** This static function make sure all the global transform caches are reset for all the UInterchangeSceneNode nodes children in the UInterchangeBaseNodeContainer */
+	/** This static function ensures all the global transform caches are reset for all the UInterchangeSceneNode nodes children in the UInterchangeBaseNodeContainer. */
 	static void ResetGlobalTransformCachesOfNodeAndAllChildren(const UInterchangeBaseNodeContainer* BaseNodeContainer, const UInterchangeBaseNode* ParentNode);
 
 	/**
-	 * Allow to retrieve the correspondence table between slot names and assigned materials for this object.
+	 * Retrieve the correspondence table between slot names and assigned materials for this object.
 	 */
 	UFUNCTION(BlueprintCallable, Category = "Interchange | Node | Meshes")
 	void GetSlotMaterialDependencies(TMap<FString, FString>& OutMaterialDependencies) const;
 
 	/**
-	 * Allow to retrieve one Material dependency for a given slot of this object.
+	 * Retrieve the Material dependency for a given slot of this object.
 	 */
 	UFUNCTION(BlueprintCallable, Category = "Interchange | Node | Meshes")
 	bool GetSlotMaterialDependencyUid(const FString& SlotName, FString& OutMaterialDependency) const;
 
 	/**
-	 * Add one Material dependency to a specific slot name of this object.
+	 * Add the specified Material dependency to a specific slot name of this object.
 	 */
 	UFUNCTION(BlueprintCallable, Category = "Interchange | Node | Meshes")
 	bool SetSlotMaterialDependencyUid(const FString& SlotName, const FString& MaterialDependencyUid);
@@ -206,7 +206,7 @@ public:
 	void GetMorphTargetCurveWeights(TMap<FString, float>& OutMorphTargetCurveWeights) const;
 
 
-	/** Set the Animation Asset To Play by this Scene Node. (only relevant for SkeletalMeshActors (SceneNodes that are instantiating Skeletal Meshes)) */
+	/** Set the Animation Asset To Play by this Scene Node. Only relevant for SkeletalMeshActors (that is, SceneNodes that are instantiating Skeletal Meshes). */
 	UFUNCTION(BlueprintCallable, Category = "Interchange | Node | SkeletalMesh")
 	bool SetCustomAnimationAssetUidToPlay(const FString& AttributeValue);
 	/** Get the Animation Asset To Play by this Scene Node. */
diff --git a/Engine/Plugins/Interchange/Runtime/Source/Nodes/Public/InterchangeShaderGraphNode.h b/Engine/Plugins/Interchange/Runtime/Source/Nodes/Public/InterchangeShaderGraphNode.h
--- a/Engine/Plugins/Interchange/Runtime/Source/Nodes/Public/InterchangeShaderGraphNode.h
+++ b/Engine/Plugins/Interchange/Runtime/Source/Nodes/Public/InterchangeShaderGraphNode.h
@@ -20,26 +20,26 @@ class INTERCHANGENODES_API UInterchangeShaderPortsAPI : public UObject
  
 public:
 	/**
-	 * Makes an attribute key to represent a node being connected to an input (ie: Inputs:InputName:Connect).
+	 * Makes an attribute key to represent a node being connected to an input (that is, Inputs:InputName:Connect).
 	 */
 	UFUNCTION(BlueprintCallable, Category = "Interchange | Node | Material")
 	static FString MakeInputConnectionKey(const FString& InputName);
  
 	/**
-	 * Makes an attribute key to represent a value being given to an input (ie: Inputs:InputName:Value).
+	 * Makes an attribute key to represent a value being given to an input (that is, Inputs:InputName:Value).
 	 */
 	UFUNCTION(BlueprintCallable, Category = "Interchange | Node | Material")
 	static FString MakeInputValueKey(const FString& InputName);
 
 	/**
-	 * Makes an attribute key to represent a parameter being given to an input (ie: Inputs:InputName:Parameter).
+	 * Makes an attribute key to represent a parameter being given to an input (that is, Inputs:InputName:Parameter).
 	 * This is more relevant to Materials, but could be used to differentiate between constant values and parameters.
 	 */
 	UFUNCTION(BlueprintCallable, Category = "Interchange | Node | Material")
 	static FString MakeInputParameterKey(const FString& InputName);
 
 	/**
-	 * From an attribute key associated with an input (ie: Inputs:InputName:Value), retrieves the input name from it.
+	 * From an attribute key associated with an input (that is, Inputs:InputName:Value), retrieves the input name.
 	 */
 	UFUNCTION(BlueprintCallable, Category = "Interchange | Node | Material")
 	static FString MakeInputName(const FString& InputKey);
@@ -51,19 +51,19 @@ public:
 	static bool IsAnInput(const FString& AttributeKey);
 
 	/**
-	 * Returns true if the attribute key is an input that represents parameters(ends with ":Parameter").
+	 * Returns true if the attribute key is an input that represents parameters (ends with ":Parameter").
 	 */
 	UFUNCTION(BlueprintCallable, Category = "Interchange | Node | Material")
 	static bool IsAParameter(const FString& AttributeKey);
  
 	/**
-	 * Checks if a particular input exists on a given node.
+	 * Checks whether a particular input exists on a given node.
 	 */
 	UFUNCTION(BlueprintCallable, Category = "Interchange | Node | Material")
 	static bool HasInput(const UInterchangeBaseNode* InterchangeNode, const FName& InInputName);
  
 	/**
-	 * Checks if a particular input exists as a parameter on a given node.
+	 * Checks whether a particular input exists as a parameter on a given node.
 	 */
 	UFUNCTION(BlueprintCallable, Category = "Interchange | Node | Material")
 	static bool HasParameter(const UInterchangeBaseNode* InterchangeNode, const FName& InInputName);
@@ -76,21 +76,21 @@ public:
  
 	/**
 	 * Adds an input connection attribute.
-	 * @param InterchangeNode	The Node to create the input on.
+	 * @param InterchangeNode	The node to create the input on.
 	 * @param InputName			The name to give to the input.
-	 * @param ExpressionUid		The unique id of the node to connect to the input.
-	 * @return					true if the input connection was successfully added to the node.
+	 * @param ExpressionUid		The unique ID of the node to connect to the input.
+	 * @return					True if the input connection was successfully added to the node.
 	 */
 	UFUNCTION(BlueprintCallable, Category = "Interchange | Node | Material")
 	static bool ConnectDefaultOuputToInput(UInterchangeBaseNode* InterchangeNode, const FString& InputName, const FString& ExpressionUid);
  
 	/**
 	 * Adds an input connection attribute.
-	 * @param InterchangeNode	The Node to create the input on.
+	 * @param InterchangeNode	The node to create the input on.
 	 * @param InputName			The name to give to the input.
-	 * @param ExpressionUid		The unique id of the node to connect to the input.
+	 * @param ExpressionUid		The unique ID of the node to connect to the input.
 	 * @param OutputName		The name of the output from ExpressionUid to connect to the input.
-	 * @return					true if the input connection was successfully added to the node.
+	 * @return					True if the input connection was successfully added to the node.
 	 */
 	UFUNCTION(BlueprintCallable, Category = "Interchange | Node | Material")
 	static bool ConnectOuputToInputByName(UInterchangeBaseNode* InterchangeNode, const FString& InputName, const FString& ExpressionUid, const FString& OutputName);
@@ -102,13 +102,13 @@ public:
 
 	/**
 	 * Adds an input connection attribute.
-	 * @param InterchangeNode	The Node to create the input on.
+	 * @param InterchangeNode	The node to create the input on.
 	 * @param InputName			The name to give to the input.
-	 * @param ExpressionUid		The unique id of the node to connect to the input.
+	 * @param ExpressionUid		The unique ID of the node to connect to the input.
 	 * @param OutputIndex		The index of the output from ExpressionUid to connect to the input.
-	 * @return					true if the input connection was successfully added to the node.
-	 * OutputIndex is encoded in a string in the following pattern  ExpressionUid:OutputByIndex:FString::FromInt(OutputIndex)
-	 * The index should be retrieved using UInterchangeShaderPortsAPI::GetOutputIndexFromName
+	 * @return					True if the input connection was successfully added to the node.
+	 * OutputIndex is encoded in a string in the following pattern: ExpressionUid:OutputByIndex:FString::FromInt(OutputIndex)
+	 * The index should be retrieved using UInterchangeShaderPortsAPI::GetOutputIndexFromName().
 	 */
 	UFUNCTION(BlueprintCallable, Category = "Interchange | Node | Material")
 	static bool ConnectOuputToInputByIndex(UInterchangeBaseNode* InterchangeNode, const FString& InputName, const FString& ExpressionUid, int32 OutputIndex);
@@ -161,7 +161,7 @@ public:
 public:
 
 	/**
-	 * Sets the Float Attribute on the Shader Node. If bIsAParameter is set to true, it would be treated as a ScalarParameter
+	 * Set the Float Attribute on the Shader Node. If bIsAParameter is set to true, it would be treated as a ScalarParameter
 	 * when the Material Pipeline creates the materials. Otherwise it would be a constant expression in the shader graph.
 	 * Note: It is assumed that the input name would be the parameter name when bIsAParameter is true.
 	 */
@@ -169,7 +169,7 @@ public:
 	bool AddFloatInput(const FString& InputName, const float& AttributeValue, bool bIsAParameter = false);
 
 	/**
-	 * Sets the Linear Color Attribute on the Shader Node. If bIsAParameter is set to true, it would be treated as a VectorParameter
+	 * Set the Linear Color Attribute on the Shader Node. If bIsAParameter is set to true, it would be treated as a VectorParameter
 	 * when the Material Pipeline creates the materials. Otherwise it would be a constant 3 vector expression in the shader graph.
 	 * Note: It is assumed that the input name would be the parameter name when bIsAParameter is true.
 	 */
@@ -177,7 +177,7 @@ public:
 	bool AddLinearColorInput(const FString& InputName, const FLinearColor& AttributeValue, bool bIsAParameter = false);
 
 	/**
-	 * Sets the String Attribute on the Shader Node. If bIsAParameter is set to true, it would be treated as a overridable Texture
+	 * Set the String Attribute on the Shader Node. If bIsAParameter is set to true, it would be treated as a overridable Texture
 	 * or else it should be treated as a LUT Texture.
 	 * Note: It is assumed that the input name would be the parameter name when bIsAParameter is true.
 	 */
@@ -215,7 +215,7 @@ public:
 	bool GetCustomMaterialFunction(FString& AttributeValue) const;
 
 	/**
-	 * Sets the unique id of the material function referenced by the function call expression.
+	 * Set the unique id of the material function referenced by the function call expression.
 	 */
 	UFUNCTION(BlueprintCallable, Category = "Interchange | Node | Material")
 	bool SetCustomMaterialFunction(const FString& AttributeValue);
@@ -239,12 +239,12 @@ public:
 	static FString MakeNodeUid(const FStringView NodeName);
 
 	/**
-	 * Creates a new UInterchangeShaderGraphNode and adds it to NodeContainer as a translated node.
+	 * Create a new UInterchangeShaderGraphNode and add it to NodeContainer as a translated node.
 	 */
 	static UInterchangeShaderGraphNode* Create(UInterchangeBaseNodeContainer* NodeContainer, const FStringView NodeName);
 
 	/**
-	 * Return the node type name of the class, we use this when reporting error
+	 * Return the node type name of the class. This is used when reporting errors.
 	 */
 	virtual FString GetTypeName() const override;
 
@@ -253,13 +253,13 @@ public:
 	bool GetCustomTwoSided(bool& AttributeValue) const;
  
 	/**
-	 * Sets if this shader graph should be rendered two sided or not. Defaults to off.
+	 * Set if this shader graph should be rendered two-sided or not. Defaults to off.
 	 */
 	UFUNCTION(BlueprintCallable, Category = "Interchange | Node | Material")
 	bool SetCustomTwoSided(const bool& AttributeValue);
 
 	/**
-	* Forces two sided even for Transmission materials.
+	* Forces two-sided even for Transmission materials.
 	*/
 	UFUNCTION(BlueprintCallable, Category = "Interchange | Node | Material")
 	bool GetCustomTwoSidedTransmission(bool& AttributeValue) const;
@@ -271,7 +271,7 @@ public:
 	bool GetCustomOpacityMaskClipValue(float& AttributeValue) const;
 
 	/**
-	 * Shader is transparent or opaque if it's alpha is lower or higher than the clip value.
+	 * The shader is transparent if its alpha value is lower than the clip value, or opaque if it is higher.
 	 */
 	UFUNCTION(BlueprintCallable, Category = "Interchange | Node | Material")
 	bool SetCustomOpacityMaskClipValue(const float& AttributeValue, bool bAddApplyDelegate = true);
@@ -280,7 +280,7 @@ public:
 	bool GetCustomIsAShaderFunction(bool& AttributeValue) const;
 
 	/**
-	 * Sets if this shader graph should be considered as a material, false, or a material function, true.
+	 * Set whether this shader graph should be considered as a material (false), or a material function (true).
 	 */
 	UFUNCTION(BlueprintCallable, Category = "Interchange | Node | Material")
 	bool SetCustomIsAShaderFunction(const bool& AttributeValue);
@@ -292,7 +292,7 @@ public:
 	bool SetCustomScreenSpaceReflections(const bool& AttributeValue);
 
 	/**
-	 * Set the Blend Mode using EBlendMode, avoid a dependency on the Engine
+	 * Set the Blend Mode using EBlendMode to avoid a dependency on the Engine.
 	 */
 	UFUNCTION(BlueprintCallable, Category = "Interchange | Node | Material")
 	bool GetCustomBlendMode(int& AttributeValue) const;
diff --git a/Engine/Plugins/Interchange/Runtime/Source/Nodes/Public/InterchangeTexture2DNode.h b/Engine/Plugins/Interchange/Runtime/Source/Nodes/Public/InterchangeTexture2DNode.h
--- a/Engine/Plugins/Interchange/Runtime/Source/Nodes/Public/InterchangeTexture2DNode.h
+++ b/Engine/Plugins/Interchange/Runtime/Source/Nodes/Public/InterchangeTexture2DNode.h
@@ -46,7 +46,7 @@ public:
 	}
 
 	/**
-	 * Return the node type name of the class, we use this when reporting error
+	 * Return the node type name of the class. This is used when reporting errors.
 	 */
 	virtual FString GetTypeName() const override
 	{
@@ -55,7 +55,7 @@ public:
 	}
 
 	/**
-	 * Override serialize to restore SourceBlocks on load.
+	 * Override Serialize() to restore SourceBlocks on load.
 	 */
 	virtual void Serialize(FArchive& Ar) override
 	{
@@ -72,8 +72,8 @@ public:
 	// UDIM base texture use a different model for the source data
 
 	/**
-	 * Get the source blocks for the texture
-	 * If the map is empty then the texture will be simply be imported as normal texture using the payload key
+	 * Get the source blocks for the texture.
+	 * If the map is empty, the texture is imported as a normal texture using the payload key.
 	 */
 	UFUNCTION(BlueprintCallable, Category = "Interchange | Node | Texture | UDIMs")
 	TMap<int32, FString> GetSourceBlocks() const
@@ -82,11 +82,11 @@ public:
 	}
 
 	/**
-	 * Set the source blocks
-	 * Using this will suggest the pipeline to consider this texture as UDIM and it can chose to pass or not these block to the texture factory node.
+	 * Set the source blocks for the texture.
+	 * Using this suggests to the pipeline to consider this texture as a UDIM. The pipeline can choose whether to pass these blocks to the texture factory node.
 	 * @param InSourceBlocks The blocks and their source image that compose the whole texture.
-	 * The textures must be of the same format and use the same pixel format
-	 * The first block in the map is used to determine the accepted texture format and pixel format
+	 * The textures must be of the same format and use the same pixel format.
+	 * The first block in the map is used to determine the accepted texture format and pixel format.
 	 */
 	void SetSourceBlocks(TMap<int32, FString>&& InSourceBlocks)
 	{
@@ -95,11 +95,11 @@ public:
 
 	
 	/**
-	 * Set the source blocks
-	 * Using this will suggest the pipeline to consider this texture as UDIM and it can chose to pass or not these block to the texture factory node.
+	 * Set the source blocks for the texture.
+	 * Using this suggests to the pipeline to consider this texture as a UDIM. The pipeline can choose whether to pass these blocks to the texture factory node.
 	 * @param InSourceBlocks The blocks and their source image that compose the whole texture.
-	 * The textures must be of the same format and use the same pixel format
-	 * The first block in the map is used to determine the accepted texture format and pixel format
+	 * The textures must be of the same format and use the same pixel format.
+	 * The first block in the map is used to determine the accepted texture format and pixel format.
 	 */
 	void SetSourceBlocks(const TMap<int32, FString>& InSourceBlocks)
 	{
diff --git a/Engine/Plugins/Interchange/Runtime/Source/Nodes/Public/InterchangeVariantSetNode.h b/Engine/Plugins/Interchange/Runtime/Source/Nodes/Public/InterchangeVariantSetNode.h
--- a/Engine/Plugins/Interchange/Runtime/Source/Nodes/Public/InterchangeVariantSetNode.h
+++ b/Engine/Plugins/Interchange/Runtime/Source/Nodes/Public/InterchangeVariantSetNode.h
@@ -8,7 +8,7 @@
 #include "InterchangeVariantSetNode.generated.h"
 
 /**
- * Class to represent a set of variants
+ * Class to represent a set of variants.
  */
 UCLASS(BlueprintType)
 class INTERCHANGENODES_API UInterchangeVariantSetNode : public UInterchangeBaseNode
@@ -24,7 +24,7 @@ public:
 	}
 
 	/**
-	 * Return the node type name of the class, we use this when reporting errors
+	 * Return the node type name of the class. This is used when reporting errors.
 	 */
 	virtual FString GetTypeName() const override
 	{
@@ -33,13 +33,13 @@ public:
 	}
 
 	/**
-	 * This function allow to retrieve the text which is displayed in the UI for this VariantSet.
+	 * Retrieve the text that is displayed in the UI for this VariantSet.
 	 */
 	UFUNCTION(BlueprintCallable, Category = "Interchange | Node | VariantSet")
 	bool GetCustomDisplayText(FString& AttributeValue) const;
 
 	/**
-	 * This function allow to set the text to be displayed in the UI for this VariantSet.
+	 * Set the text to be displayed in the UI for this VariantSet.
 	 */
 	UFUNCTION(BlueprintCallable, Category = "Interchange | Node | VariantSet")
 	bool SetCustomDisplayText(const FString& AttributeValue);
@@ -57,31 +57,31 @@ public:
 	bool SetCustomVariantsPayloadKey(const FString& PayloadKey);
 
 	/**
-	 * This function allow to retrieve the number of translated node's unique ids for this VariantSet.
+	 * Retrieve the number of translated node's unique IDs for this VariantSet.
 	 */
 	UFUNCTION(BlueprintCallable, Category = "Interchange | Node | SceneVariantSet")
 	int32 GetCustomDependencyUidCount() const;
 
 	/**
-	 * This function allow to retrieve all the translated node's unique ids for this VariantSet.
+	 * Retrieve all the translated node's unique IDs for this VariantSet.
 	 */
 	UFUNCTION(BlueprintCallable, Category = "Interchange | Node | SceneVariantSet")
 	void GetCustomDependencyUids(TArray<FString>& OutDependencyUids) const;
 
 	/**
-	 * This function allow to retrieve a specific translated node's unique id for this VariantSet.
+	 * Retrieve the specified translated node's unique ID for this VariantSet.
 	 */
 	UFUNCTION(BlueprintCallable, Category = "Interchange | Node | SceneVariantSet")
 	void GetCustomDependencyUid(const int32 Index, FString& OutDependencyUid) const;
 
 	/**
-	 * Add one translated node's unique id to this VariantSet.
+	 * Add the specified translated node's unique ID to this VariantSet.
 	 */
 	UFUNCTION(BlueprintCallable, Category = "Interchange | Node | SceneVariantSet")
 	bool AddCustomDependencyUid(const FString& DependencyUid);
 
 	/**
-	 * Remove one translated node's unique id from this VariantSet.
+	 * Remove the specified translated node's unique ID from this VariantSet.
 	 */
 	UFUNCTION(BlueprintCallable, Category = "Interchange | Node | SceneVariantSet")
 	bool RemoveCustomDependencyUid(const FString& DependencyUid);
@@ -109,7 +109,7 @@ public:
 	}
 
 	/**
-	* Return the node type name of the class, we use this when reporting errors
+	* Return the node type name of the class. This is used when reporting errors.
 	*/
 	virtual FString GetTypeName() const override
 	{
@@ -118,31 +118,31 @@ public:
 	}
 
 	/**
-	 * This function allow to retrieve the number of VariantSets for this object.
+	 * Retrieve the number of VariantSets for this object.
 	 */
 	UFUNCTION(BlueprintCallable, Category = "Interchange | Node | SceneVariantSet")
 	int32 GetCustomVariantSetUidCount() const;
 
 	/**
-	 * This function allow to retrieve all the VariantSets' unique ids for this object.
+	 * Retrieve all the VariantSets' unique ids for this object.
 	 */
 	UFUNCTION(BlueprintCallable, Category = "Interchange | Node | SceneVariantSet")
 	void GetCustomVariantSetUids(TArray<FString>& OutVariantUids) const;
 
 	/**
-	 * This function allow to retrieve one VariantSet's unique id for this object.
+	 * Retrieve the specified VariantSet's unique id for this object.
 	 */
 	UFUNCTION(BlueprintCallable, Category = "Interchange | Node | SceneVariantSet")
 	void GetCustomVariantSetUid(const int32 Index, FString& OutVariantUid) const;
 
 	/**
-	 * Add one VariantSet's unique id to this object.
+	 * Add the specified VariantSet's unique id to this object.
 	 */
 	UFUNCTION(BlueprintCallable, Category = "Interchange | Node | SceneVariantSet")
 	bool AddCustomVariantSetUid(const FString& VariantUid);
 
 	/**
-	 * Remove one VariantSet's unique id from this object.
+	 * Remove the specified VariantSet's unique id from this object.
 	 */
 	UFUNCTION(BlueprintCallable, Category = "Interchange | Node | SceneVariantSet")
 	bool RemoveCustomVariantSetUid(const FString& VariantUid);
diff --git a/Engine/Plugins/Interchange/Runtime/Source/Parsers/Fbx/Private/FbxMaterial.h b/Engine/Plugins/Interchange/Runtime/Source/Parsers/Fbx/Private/FbxMaterial.h
--- a/Engine/Plugins/Interchange/Runtime/Source/Parsers/Fbx/Private/FbxMaterial.h
+++ b/Engine/Plugins/Interchange/Runtime/Source/Parsers/Fbx/Private/FbxMaterial.h
@@ -30,24 +30,24 @@ namespace UE
 				{}
 
 				/**
-				 * Create a UInterchangeFextureNode and add it to the NodeContainer for all texture of type FbxFileTexture the fbx file contain.
+				 * Create a UInterchangeFextureNode and add it to the NodeContainer for each texture of type FbxFileTexture that the FBX file contains.
 				 *
-				 * @note - Any node that already exist in the NodeContainer will not be created or modified.
+				 * @note - Any node that already exists in the NodeContainer will not be created or modified.
 				 */
 				void AddAllTextures(FbxScene* SDKScene, UInterchangeBaseNodeContainer& NodeContainer);
 				
 				/**
-				 * Create a UInterchangeMaterialNode and add it to the NodeContainer for all material of type FbxSurfaceMaterial the fbx file contain.
+				 * Create a UInterchangeMaterialNode and add it to the NodeContainer for each material of type FbxSurfaceMaterial that the FBX file contains.
 				 * 
-				 * @note - Any node that already exist in the NodeContainer will not be created or modified.
+				 * @note - Any node that already exists in the NodeContainer will not be created or modified.
 				 */
 				void AddAllMaterials(FbxScene* SDKScene, UInterchangeBaseNodeContainer& NodeContainer);
 
 				/**
-				 * Create a UInterchangeMaterialNode and add it to the NodeContainer for all material of type FbxSurfaceMaterial the fbx ParentFbxNode contain.
-				 * It also set the dependencies of the node materials on the interchange ParentNode.
+				 * Create a UInterchangeMaterialNode and add it to the NodeContainer for each material of type FbxSurfaceMaterial that the FBX ParentFbxNode contains.
+				 * Also set the dependencies of the node materials on the Interchange ParentNode.
 				 * 
-				 * @note - Any material node that already exist in the NodeContainer will simply be add has a dependency.
+				 * @note - Any material node that already exists in the NodeContainer will be added as a dependency.
 				 */
 				void AddAllNodeMaterials(UInterchangeSceneNode* SceneNode, FbxNode* ParentFbxNode, UInterchangeBaseNodeContainer& NodeContainer);
 
diff --git a/Engine/Plugins/Interchange/Runtime/Source/Parsers/Fbx/Private/FbxMesh.h b/Engine/Plugins/Interchange/Runtime/Source/Parsers/Fbx/Private/FbxMesh.h
--- a/Engine/Plugins/Interchange/Runtime/Source/Parsers/Fbx/Private/FbxMesh.h
+++ b/Engine/Plugins/Interchange/Runtime/Source/Parsers/Fbx/Private/FbxMesh.h
@@ -98,7 +98,7 @@ namespace UE
 				bool FillStaticMeshDescriptionFromFbxMesh(FbxMesh* Mesh, const FTransform& MeshGlobalTransform);
 				
 				/*
-				 * Fill the mesh description using the Mesh parameter and also fill the OutJointNodeUniqueIDs so the MeshDescription bone Index can be map to the correct interchange joint scene node.
+				 * Fill the mesh description using the Mesh parameter, and also fill the OutJointNodeUniqueIDs so the MeshDescription bone Index can be mapped to the correct Interchange joint scene node.
 				 */
 				bool FillSkinnedMeshDescriptionFromFbxMesh(FbxMesh* Mesh, const FTransform& MeshGlobalTransform, TArray<FString>& OutJointUniqueNames);
 
@@ -108,7 +108,7 @@ namespace UE
 				bool FillMeshDescriptionFromFbxShape(FbxShape* Shape, const FTransform& MeshGlobalTransform);
 
 				/**
-				 * Add messages to the message log
+				 * Add messages to the message log.
 				 */
 				template <typename T>
 				T* AddMessage(FbxGeometryBase* FbxNode) const
@@ -200,7 +200,7 @@ namespace UE
 
 				TArray<FMorphTargetAnimationBuildingData>& GetMorphTargetAnimationsBuildingData() { return MorphTargetAnimationsBuildingData; }
 			protected:
-				/** Add joint to the interchange mesh node joint dependencies. Return false if there is no valid joint (not a valid skinned mesh) */
+				/** Add joints to the Interchange mesh node joint dependencies. Return false if there is no valid joint (not a valid skinned mesh). */
 				bool ExtractSkinnedMeshNodeJoints(FbxScene* SDKScene, UInterchangeBaseNodeContainer& NodeContainer, FbxMesh* Mesh, UInterchangeMeshNode* MeshNode);
 				UInterchangeMeshNode* CreateMeshNode(UInterchangeBaseNodeContainer& NodeContainer, const FString& NodeName, const FString& NodeUniqueID);
 
diff --git a/Engine/Plugins/Interchange/Runtime/Source/Pipelines/Private/InterchangeGenericAnimationPipeline.cpp b/Engine/Plugins/Interchange/Runtime/Source/Pipelines/Private/InterchangeGenericAnimationPipeline.cpp
--- a/Engine/Plugins/Interchange/Runtime/Source/Pipelines/Private/InterchangeGenericAnimationPipeline.cpp
+++ b/Engine/Plugins/Interchange/Runtime/Source/Pipelines/Private/InterchangeGenericAnimationPipeline.cpp
@@ -85,7 +85,7 @@ void UInterchangeGenericAnimationPipeline::ExecutePipeline(UInterchangeBaseNodeC
 {
 	if (!InBaseNodeContainer)
 	{
-		UE_LOG(LogInterchangePipeline, Warning, TEXT("UInterchangeGenericAnimationPipeline: Cannot execute pre-import pipeline because InBaseNodeContrainer is null"));
+		UE_LOG(LogInterchangePipeline, Warning, TEXT("UInterchangeGenericAnimationPipeline: Cannot execute pre-import pipeline because InBaseNodeContainer is null."));
 		return;
 	}
 
@@ -309,7 +309,7 @@ void UInterchangeGenericAnimationPipeline::ExecutePipeline(UInterchangeBaseNodeC
 
 	if (CommonSkeletalMeshesAndAnimationsProperties->bImportOnlyAnimations && !CommonSkeletalMeshesAndAnimationsProperties->Skeleton.IsValid())
 	{
-		UE_LOG(LogInterchangePipeline, Warning, TEXT("UInterchangeGenericAnimationPipeline: Cannot execute pre-import pipeline because we cannot import animation only but not specify any valid skeleton"));
+		UE_LOG(LogInterchangePipeline, Warning, TEXT("UInterchangeGenericAnimationPipeline: Cannot execute pre-import pipeline because importing animation only requires a valid skeleton."));
 		return;
 	}
 	SourceDatas.Empty(InSourceDatas.Num());
@@ -562,7 +562,7 @@ void UInterchangeGenericAnimationPipeline::CreateAnimSequenceFactoryNode(UInterc
 				Message->SourceAssetName = SourceDatas[0]->GetFilename();
 				Message->DestinationAssetName = TrackNode.GetDisplayLabel();
 				Message->AssetType = UAnimSequence::StaticClass();
-				Message->Text = FText::Format(NSLOCTEXT("UInterchangeGenericAnimationPipeline", "WrongSequenceLength", "Animation length {0} is not compatible with import frame-rate {1} (sub frame {2}), animation has to be frame-border aligned."),
+				Message->Text = FText::Format(NSLOCTEXT("UInterchangeGenericAnimationPipeline", "WrongSequenceLength", "Animation length {0} is not compatible with import frame-rate {1} (sub frame {2}). The animation has to be frame-border aligned."),
 					FText::AsNumber(SequenceLength), FrameRate.ToPrettyText(), FText::AsNumber(SubFrame));
 				//Skip this anim sequence factory node
 				return;
diff --git a/Engine/Plugins/Interchange/Runtime/Source/Pipelines/Private/InterchangeGenericMeshPipeline.cpp b/Engine/Plugins/Interchange/Runtime/Source/Pipelines/Private/InterchangeGenericMeshPipeline.cpp
--- a/Engine/Plugins/Interchange/Runtime/Source/Pipelines/Private/InterchangeGenericMeshPipeline.cpp
+++ b/Engine/Plugins/Interchange/Runtime/Source/Pipelines/Private/InterchangeGenericMeshPipeline.cpp
@@ -230,7 +230,7 @@ void UInterchangeGenericMeshPipeline::ExecutePipeline(UInterchangeBaseNodeContai
 {
 	if (!InBaseNodeContainer)
 	{
-		UE_LOG(LogInterchangePipeline, Warning, TEXT("UInterchangeGenericMeshPipeline: Cannot execute pre-import pipeline because InBaseNodeContrainer is null"));
+		UE_LOG(LogInterchangePipeline, Warning, TEXT("UInterchangeGenericMeshPipeline: Cannot execute pre-import pipeline because InBaseNodeContrainer is null."));
 		return;
 	}
 	
@@ -348,7 +348,7 @@ bool UInterchangeGenericMeshPipeline::DoClassesIncludeAllEditableStructPropertie
 			//Ensure to notify
 			if (!bFindProperty)
 			{
-				UE_LOG(LogInterchangePipeline, Log, TEXT("Interchange mesh pipeline do not include build property %s"), *PropertyName.ToString());
+				UE_LOG(LogInterchangePipeline, Log, TEXT("The Interchange mesh pipeline does not include build property %s."), *PropertyName.ToString());
 				bResult = false;
 			}
 		}
diff --git a/Engine/Plugins/Interchange/Runtime/Source/Pipelines/Public/InterchangeGenericAnimationPipeline.h b/Engine/Plugins/Interchange/Runtime/Source/Pipelines/Public/InterchangeGenericAnimationPipeline.h
--- a/Engine/Plugins/Interchange/Runtime/Source/Pipelines/Public/InterchangeGenericAnimationPipeline.h
+++ b/Engine/Plugins/Interchange/Runtime/Source/Pipelines/Public/InterchangeGenericAnimationPipeline.h
@@ -16,7 +16,7 @@ class UInterchangeAnimationTrackSetNode;
 class UInterchangeMeshNode;
 class UInterchangeSkeletalAnimationTrackNode;
 
-/** Animation length type when importing */
+/** Animation length type when importing. */
 UENUM(BlueprintType)
 enum class EInterchangeAnimationRange : uint8
 {
@@ -50,7 +50,7 @@ public:
 	TWeakObjectPtr<UInterchangeGenericCommonMeshesProperties> CommonMeshesProperties;
 
 	//////	ANIMATION_CATEGORY Properties //////
-	/** If enable, import all animation assets find in the sources. */
+	/** If enabled, import all animation assets found in the source. */
 	UPROPERTY(EditAnywhere, BlueprintReadWrite, Category = "Animations")
 	bool bImportAnimations = true;
 
@@ -58,15 +58,15 @@ public:
 	UPROPERTY(EditAnywhere, BlueprintReadWrite, Category = "Animations", meta = (EditCondition = "bImportAnimations"))
 	bool bImportBoneTracks = true;
 
-	/** Which animation range to import. The one defined at Exported, at Animated time or define a range manually */
+	/** Determines which animation range to import: the range defined at export, the range of frames with animation, or a manually defined range. */
 	UPROPERTY(EditAnywhere, BlueprintReadWrite, Category = "Animations", meta = (EditCondition = "bImportAnimations && bImportBoneTracks", DisplayName = "Animation Length"))
 	EInterchangeAnimationRange AnimationRange = EInterchangeAnimationRange::Timeline;
 
-	/** Frame range used when Set Range is used in Animation Length */
+	/** The frame range used when the Animation Length setting is set to Set Range. */
 	UPROPERTY(EditAnywhere, BlueprintReadWrite, Category = "Animations", meta = (UIMin = 0, ClampMin = 0))
 	FInt32Interval FrameImportRange = FInt32Interval(0, 0);
 
-	/** Enable this option to use default sample rate for the imported animation at 30 frames per second */
+	/** If enabled, samples all imported animation data at the default rate of 30 FPS. */
 	UPROPERTY(EditAnywhere, BlueprintReadWrite, Category = "Animations", meta = (EditCondition = "bImportAnimations && bImportBoneTracks", ToolTip = "If enabled, samples all animation curves to 30 FPS"))
 	bool bUse30HzToBakeBoneAnimation = false;
 
@@ -74,47 +74,47 @@ public:
 	UPROPERTY(EditAnywhere, BlueprintReadWrite, Category = "Animations", meta = (EditCondition = "bImportAnimations && bImportBoneTracks && !bUse30HzToBakeBoneAnimation", ToolTip = "Sample fbx animation data at the specified sample rate, 0 find automaticaly the best sample rate", ClampMin = 0, UIMin = 0, ClampMax = 48000, UIMax = 60))
 	int32 CustomBoneAnimationSampleRate = 0;
 
-	/** If enabled, snaps the animation to the closest frame boundary using the import sampling rate */
+	/** If enabled, snaps the animation to the closest frame boundary using the import sampling rate. */
 	UPROPERTY(EditAnywhere, BlueprintReadWrite, Category = "Animations", meta = (EditCondition = "bImportAnimations && bImportBoneTracks"))
 	bool bSnapToClosestFrameBoundary = false;
 
-	/** If true, import node attributes as either Animation Curves or Animation Attributes */
+	/** If enabled, import node attributes as either Animation Curves or Animation Attributes. */
 	UPROPERTY(EditAnywhere, BlueprintReadWrite, Category = "Animations", meta = (SubCategory = "Curves", EditCondition = "bImportAnimations", DisplayName = "Import Attributes as Curves or Animation Attributes"))
 	bool bImportCustomAttribute = true;
 
-	/** Whether to automatically add curve metadata to an animation's skeleton. If this is disabled, curve metadata will be added to skeletal meshes for morph targets, but no metadata entry will be created for general curves. */
+	/** Determines whether to automatically add curve metadata to an animation's skeleton. If this setting is disabled, curve metadata will be added to skeletal meshes for morph targets, but no metadata entry will be created for general curves. */
 	UPROPERTY(EditAnywhere, BlueprintReadWrite, Category = "Animations", meta = (SubCategory = "Curves", EditCondition = "bImportAnimations && bImportCustomAttribute"))
 	bool bAddCurveMetadataToSkeleton = true;
 
-	/** Set Material Curve Type for all custom attributes that exists */
+	/** Set the material curve type for all custom attributes. */
 	UPROPERTY(EditAnywhere, BlueprintReadWrite, Category = "Animations", meta = (SubCategory = "Curves", EditCondition = "bImportAnimations && bImportCustomAttribute", DisplayName = "Set Material Curve Type"))
 	bool bSetMaterialDriveParameterOnCustomAttribute = false;
 
-	/** Set Material Curve Type for the custom attribute with the following suffixes. This doesn't matter if Set Material Curve Type is true  */
+	/** Set the Material Curve Type for custom attributes that have the specified suffixes. This setting is not used if the Set Material Curve Type setting is enabled.  */
 	UPROPERTY(EditAnywhere, BlueprintReadWrite, Category = "Animations", meta = (SubCategory = "Curves", EditCondition = "bImportAnimations && bImportCustomAttribute && bSetMaterialDriveParameterOnCustomAttribute", DisplayName = "Material Curve Suffixes"))
 	TArray<FString> MaterialCurveSuffixes = {TEXT("_mat")};
 
-	/** When importing custom attribute as curve, remove redundant keys */
+	/** When importing custom attributes as curves, remove redundant keys. */
 	UPROPERTY(EditAnywhere, BlueprintReadWrite, Category = "Animations", meta = (SubCategory = "Curves", EditCondition = "bImportAnimations && bImportCustomAttribute", DisplayName = "Remove Redundant Keys"))
 	bool bRemoveCurveRedundantKeys = false;
 
-	/** When importing custom attribute or morphtarget as curve, do not import if it doesn't have any value other than zero. This is to avoid adding extra curves to evaluate */
+	/** When importing a custom attribute or morph target as a curve, only import if it has a value other than zero. This avoids adding extra curves to evaluate. */
 	UPROPERTY(EditAnywhere, BlueprintReadWrite, Category = "Animations", meta = (SubCategory = "Curves", EditCondition = "bImportAnimations && bImportCustomAttribute", DisplayName = "Do not import curves with only 0 values"))
 	bool bDoNotImportCurveWithZero = false;
 
-	/** If true, all previous node attributes imported as Animation Attributes will be deleted when doing a re-import. */
+	/** If enabled, all previous node attributes imported as Animation Attributes will be deleted when doing a reimport. */
 	UPROPERTY(EditAnywhere, BlueprintReadWrite, Category = "Animations", meta = (SubCategory = "Curves", EditCondition = "bImportAnimations && bImportCustomAttribute", DisplayName = "Delete existing Animation Attributes"))
 	bool bDeleteExistingNonCurveCustomAttributes = false;
 
-	/** If true, all previous node attributes imported as Animation Curves will be deleted when doing a re-import. */
+	/** If enabled, all previous node attributes imported as Animation Curves will be deleted when doing a reimport. */
 	UPROPERTY(EditAnywhere, BlueprintReadWrite, Category = "Animations", meta = (SubCategory = "Curves", EditCondition = "bImportAnimations && bImportCustomAttribute", DisplayName = "Delete existing Animation Curves"))
 	bool bDeleteExistingCustomAttributeCurves = false;
 
-	/** If true, all previous morph target curves will be deleted when doing a re-import */
+	/** If enabled, all previous morph target curves will be deleted when doing a reimport. */
 	UPROPERTY(EditAnywhere, BlueprintReadWrite, Category = "Animations", meta = (SubCategory = "Curves", EditCondition = "bImportAnimations && bImportCustomAttribute", DisplayName = "Delete existing Morph Target Curves"))
 	bool bDeleteExistingMorphTargetCurves = false;
 
-	/** Name of source animation that was imported, used to reimport correct animation from the translated source */
+	/** Name of the source animation that was imported. This is used to reimport correct animation from the translated source. */
 	UPROPERTY()
 	FString SourceAnimationName;
 
diff --git a/Engine/Plugins/Interchange/Runtime/Source/Pipelines/Public/InterchangeGenericAssetsPipeline.h b/Engine/Plugins/Interchange/Runtime/Source/Pipelines/Public/InterchangeGenericAssetsPipeline.h
--- a/Engine/Plugins/Interchange/Runtime/Source/Pipelines/Public/InterchangeGenericAssetsPipeline.h
+++ b/Engine/Plugins/Interchange/Runtime/Source/Pipelines/Public/InterchangeGenericAssetsPipeline.h
@@ -31,8 +31,8 @@ struct FReferenceSkeleton;
 struct FMeshBoneInfo;
 
 /**
- * This pipeline is the generic pipeline option for all meshes type and should be call before specialized Mesh pipeline (like generic static mesh or skeletal mesh pipelines)
- * All shared import options between mesh type should be added here.
+ * This pipeline is the generic option for all types of meshes. It should be called before specialized mesh pipelines like the generic static mesh or skeletal mesh pipelines.
+ * All import options that are shared between mesh types should be added here.
  *
  */
 UCLASS(BlueprintType, editinlinenew)
@@ -49,15 +49,15 @@ public:
 	UPROPERTY(EditAnywhere, BlueprintReadWrite, Category = "Common", meta = (StandAlonePipelineProperty = "True", PipelineInternalEditionData = "True"))
 	FString PipelineDisplayName;
 
-	/* Allow user to choose the re-import strategy. */
+	/* Set the reimport strategy. */
 	UPROPERTY(EditAnywhere, BlueprintReadWrite, Category = "Common", meta = (AdjustPipelineAndRefreshDetailOnChange = "True"))
 	EReimportStrategyFlags ReimportStrategy = EReimportStrategyFlags::ApplyNoProperties;
 
-	/** If enable and AssetName is empty and there is only one asset and one source data, we will name the asset like the source data name. */
+	/** If enabled, and the Asset Name setting is empty, and there is only one asset and one source, the imported asset is given the same name as the source data. */
 	UPROPERTY(EditAnywhere, BlueprintReadWrite, Category = "Common")
 	bool bUseSourceNameForAsset = true;
 
-	/** If not empty, and there is only one asset and one source data, we will name the asset with this string. */
+	/** If set, and there is only one asset and one source, the imported asset is given this name. */
 	UPROPERTY(EditAnywhere, BlueprintReadWrite, Category = "Common", meta = (StandAlonePipelineProperty = "True", AlwaysResetToDefault = "True"))
 	FString AssetName;
 
diff --git a/Engine/Plugins/Interchange/Runtime/Source/Pipelines/Public/InterchangeGenericAssetsPipelineSharedSettings.h b/Engine/Plugins/Interchange/Runtime/Source/Pipelines/Public/InterchangeGenericAssetsPipelineSharedSettings.h
--- a/Engine/Plugins/Interchange/Runtime/Source/Pipelines/Public/InterchangeGenericAssetsPipelineSharedSettings.h
+++ b/Engine/Plugins/Interchange/Runtime/Source/Pipelines/Public/InterchangeGenericAssetsPipelineSharedSettings.h
@@ -18,15 +18,15 @@
 class UInterchangeSkeletonFactoryNode;
 class USkeleton;
 
-/** Force mesh type, if user want to import all meshes as one type*/
+/** Enumerates the options for importing all meshes as one type. */
 UENUM(BlueprintType)
 enum class EInterchangeForceMeshType : uint8
 {
-	/** Will import from the source type, no conversion */
+	/** Import both static meshes and skeletal meshes from the source without converting them. */
 	IFMT_None UMETA(DisplayName = "None"),
-	/** Will import any mesh to static mesh. */
+	/** Import all meshes from the source as static meshes. */
 	IFMT_StaticMesh UMETA(DisplayName = "Static Mesh"),
-	/** Will import any mesh to skeletal mesh. */
+	/** Import all meshes from the source as skeletal meshes. */
 	IFMT_SkeletalMesh UMETA(DisplayName = "Skeletal Mesh"),
 
 	IFMT_MAX
@@ -37,7 +37,7 @@ enum class EInterchangeVertexColorImportOption : uint8
 {
 	/** Import the mesh using the vertex colors from the translated source. */
 	IVCIO_Replace UMETA(DisplayName = "Replace"),
-	/** Ignore vertex colors from the translated source. In case of a re-import keep the existing mesh vertex colors. */
+	/** Ignore vertex colors from the translated source. In case of a reimport, keep the existing mesh vertex colors. */
 	IVCIO_Ignore UMETA(DisplayName = "Ignore"),
 	/** Override all vertex colors with the specified color. */
 	IVCIO_Override UMETA(DisplayName = "Override"),
@@ -53,22 +53,22 @@ class INTERCHANGEPIPELINES_API UInterchangeGenericCommonMeshesProperties : publi
 public:
 	//////	COMMON_MESHES_CATEGORY Properties //////
 
-	/** Allow to convert mesh to a particular type */
+	/** If set, imports all meshes in the source as either static meshes or skeletal meshes. */
  	UPROPERTY(EditAnywhere, BlueprintReadWrite, Category = "Common Meshes")
  	EInterchangeForceMeshType ForceAllMeshAsType = EInterchangeForceMeshType::IFMT_None;
 
 	/**
-	 * If enable, and some static mesh transforms are animated, the pipeline will convert the static mesh into rigid skeletal mesh.
-	 * Note: This setting will be ignored if ForceAllMeshAsType is not set to 'None'.
+	 * If enabled, and some static mesh transforms are animated, the pipeline will convert the static mesh into a rigid skeletal mesh.
+	 * This setting is only used if the Force All Meshes As Type setting is set to "None".
 	 */
 	UPROPERTY(EditAnywhere, BlueprintReadWrite, Category = "Common Meshes")
 	bool bAutoDetectMeshType = true;
 
-	/** If enable, meshes LODs will be imported. Note that it required the advanced bBakeMesh property to be enabled. */
+	/** If enabled, any existing LODs for meshes are imported. This setting is only used if the Bake Meshes setting is also enabled. */
 	UPROPERTY(EditAnywhere, BlueprintReadWrite, Category = "Common Meshes")
 	bool bImportLods = true;
 
-	/** If enable, meshes will be baked with the scene instance hierarchy transform. */
+	/** If enabled, meshes are baked with the scene instance hierarchy transform. */
 	UPROPERTY(EditAnywhere, BlueprintReadWrite, Category = "Common Meshes")
 	bool bBakeMeshes = true;
 	
@@ -76,43 +76,43 @@ public:
 	UPROPERTY(EditAnywhere, BlueprintReadWrite, Category = "Common Meshes")
 	bool bKeepSectionsSeparate = false;
 	
-	/** Specify how vertex colors should be imported */
+	/** Specify how vertex colors should be imported. */
 	UPROPERTY(EditAnywhere, BlueprintReadWrite, Category = "Common Meshes")
 	EInterchangeVertexColorImportOption VertexColorImportOption = EInterchangeVertexColorImportOption::IVCIO_Replace;
 
-	/** Specify override color in the case that VertexColorImportOption is set to Override */
+	/** Specify an override color for use when the Vertex Color Import Option setting is set to Override. */
 	UPROPERTY(EditAnywhere, BlueprintReadWrite, Category = "Common Meshes")
 	FColor VertexOverrideColor;
 
-	/** If true, normals in the imported mesh are ignored and recomputed. */
+	/** If enabled, normals in the imported mesh are ignored and recomputed. */
 	UPROPERTY(EditAnywhere, BlueprintReadWrite, Category = "Common Meshes", meta = (SubCategory = "Build"))
 	bool bRecomputeNormals = true;
 
-	/** If true, tangents in the imported mesh are ignored and recomputed. */
+	/** If enabled, tangents in the imported mesh are ignored and recomputed. */
 	UPROPERTY(EditAnywhere, BlueprintReadWrite, Category = "Common Meshes", meta = (SubCategory = "Build"))
 	bool bRecomputeTangents = true;
 
-	/** If true, recompute tangents will use mikkt space. */
+	/** If enabled, tangents are recomputed using MikkTSpace. */
 	UPROPERTY(EditAnywhere, BlueprintReadWrite, Category = "Common Meshes", meta = (SubCategory = "Build", editcondition = "bRecomputeTangents"))
 	bool bUseMikkTSpace = true;
 
-	/** If true, we will use the surface area and the corner angle of the triangle as a ratio when computing the normals. */
+	/** If enabled, normals are computed using the surface area and the corner angle of the triangle as a ratio. */
 	UPROPERTY(EditAnywhere, BlueprintReadWrite, Category = "Common Meshes", meta = (SubCategory = "Build", editcondition = "bRecomputeNormals"))
 	bool bComputeWeightedNormals = false;
 
-	/** If true, Tangents will be stored at 16 bit vs 8 bit precision. */
+	/** If true, tangents are stored at 16-bit vs 8-bit precision. */
 	UPROPERTY(EditAnywhere, BlueprintReadWrite, Category = "Common Meshes", meta = (SubCategory = "Build"))
 	bool bUseHighPrecisionTangentBasis = false;
 
-	/** If true, UVs will be stored at full floating point precision. */
+	/** If true, UVs are stored at full floating-point precision. */
 	UPROPERTY(EditAnywhere, BlueprintReadWrite, Category = "Common Meshes", meta = (SubCategory = "Build"))
 	bool bUseFullPrecisionUVs = false;
 
-	/** If true, UVs will use backwards-compatible F16 conversion with truncation for legacy meshes. */
+	/** If enabled, UVs are converted to 16-bit by a legacy truncation process instead of the default rounding process. This may avoid differences when reimporting older content. */
 	UPROPERTY(EditAnywhere, BlueprintReadWrite, Category = "Common Meshes", meta = (SubCategory = "Build"))
 	bool bUseBackwardsCompatibleF16TruncUVs = false;
 
-	/** If true, degenerate triangles will be removed. */
+	/** If true, degenerate triangles are removed. */
 	UPROPERTY(EditAnywhere, BlueprintReadWrite, Category = "Common Meshes", meta = (SubCategory = "Build"))
 	bool bRemoveDegenerates = false;
 #if WITH_EDITOR
@@ -140,23 +140,26 @@ class INTERCHANGEPIPELINES_API UInterchangeGenericCommonSkeletalMeshesAndAnimati
 {
 	GENERATED_BODY()
 public:
-	/** Enable this option to only import animation, a valid skeleton must be set to import only the animations. */
+	/** If enabled, only animations are imported from the source. You must also set a valid skeleton. */
 	UPROPERTY(EditAnywhere, BlueprintReadWrite, Category = "Common Skeletal Meshes and Animations")
 	bool bImportOnlyAnimations = false;
 
-	/** Skeleton to use for imported asset. When importing a skeletal mesh, leaving this as "None" will create a new skeleton. When importing an animation this MUST be specified to import the asset. */
+	/** Skeleton to use for imported asset. When importing a skeletal mesh, leaving this as "None" will create a new skeleton. When importing an animation, this must be specified. */
 	UPROPERTY(EditAnywhere, BlueprintReadWrite, Category = "Common Skeletal Meshes and Animations")
 	TWeakObjectPtr<USkeleton> Skeleton;
 
-	/** If checked, meshes nested in bone hierarchies will be imported instead of being converted to bones. */
+	/** 
+	 * If enabled, meshes nested in bone hierarchies will be imported as meshes instead of being converted to bones. If the meshes are not skinned, they are
+	 * added to the skeletal mesh and removed from the list of static meshes.
+	 */
 	UPROPERTY(EditAnywhere, BlueprintReadWrite, Category = "Common Skeletal Meshes and Animations")
 	bool bImportMeshesInBoneHierarchy = true;
 
-	/** Enable this option to use frame 0 as reference pose */
+	/** If enabled, frame 0 is used as the reference pose for skeletal meshes. */
 	UPROPERTY(EditAnywhere, BlueprintReadWrite, Category = "Common Skeletal Meshes and Animations")
 	bool bUseT0AsRefPose = false;
 
-	/** Allow to convert static mesh using morph target to skeletal mesh" */
+	/** If enabled, all static meshes that have morph targets will be imported as skeletal meshes instead. */
 	UPROPERTY(EditAnywhere, BlueprintReadWrite, Category = "Static Meshes")
 	bool bConvertStaticsWithMorphTargetsToSkeletals = false;
 #if WITH_EDITOR
diff --git a/Engine/Plugins/Interchange/Runtime/Source/Pipelines/Public/InterchangeGenericMaterialPipeline.h b/Engine/Plugins/Interchange/Runtime/Source/Pipelines/Public/InterchangeGenericMaterialPipeline.h
--- a/Engine/Plugins/Interchange/Runtime/Source/Pipelines/Public/InterchangeGenericMaterialPipeline.h
+++ b/Engine/Plugins/Interchange/Runtime/Source/Pipelines/Public/InterchangeGenericMaterialPipeline.h
@@ -32,7 +32,9 @@ namespace UE::Interchange::Materials::HashUtils
 UENUM(BlueprintType)
 enum class EInterchangeMaterialImportOption : uint8
 {
+	/** Import all materials from the source as material assets. */
 	ImportAsMaterials,
+	/** Import all materials from the source as material instance assets. */
 	ImportAsMaterialInstances,
 };
 
@@ -71,18 +73,19 @@ public:
 	UPROPERTY(EditAnywhere, BlueprintReadWrite, Category = "Materials")
 	EInterchangeMaterialSearchLocation SearchLocation = EInterchangeMaterialSearchLocation::Local;
 
-	/** If not empty, and there is only one asset and one source data, we will name the asset with this string. */
+	/** If set, and there is only one asset and one source, the imported asset will be given this name. */
 	UPROPERTY(EditAnywhere, BlueprintReadWrite, Category = "Materials", meta = (StandAlonePipelineProperty = "True", AlwaysResetToDefault = "True"))
 	FString AssetName;
 
+	/** Determines what kind of material assets should be created for the imported materials. */
 	UPROPERTY(EditAnywhere, BlueprintReadWrite, Category = "Materials", Meta=(EditCondition="bImportMaterials"))
 	EInterchangeMaterialImportOption MaterialImport = EInterchangeMaterialImportOption::ImportAsMaterials;
 	
-	/** If set, reference materials along with respective material instances would be created*/
+	/** If set, reference materials along with respective material instances are created. */
 	UPROPERTY(EditAnywhere, BlueprintReadWrite, Category = "Materials", Meta = (EditCondition = "bImportMaterials"))
 	bool bIdentifyDuplicateMaterials = false;
 
-	/** If set, additional material instance would be created for reference/parent material*/
+	/** If set, additional material instances are created for reference/parent materials. */
 	UPROPERTY(EditAnywhere, BlueprintReadWrite, Category = "Materials", Meta=(EditCondition="bIdentifyDuplicateMaterials"))
 	bool bCreateMaterialInstanceForParent = false;
 
@@ -158,7 +161,7 @@ private:
 		return HasThinTranslucency(ShaderGraphNode);
 	}
 
-	/** True if the shader graph has a base color input (Metallic/Roughness model. */
+	/** True if the shader graph has a base color input (Metallic/Roughness model). */
 	bool IsMetalRoughModel(const UInterchangeShaderGraphNode* ShaderGraphNode) const;
 	UE_DEPRECATED(5.3, "Deprecated. Use IsMetalRoughModel and IsSpecGlossModel to identify the correct PBR model.")
 	bool IsPBRModel(const UInterchangeShaderGraphNode* ShaderGraphNode) const
@@ -172,7 +175,7 @@ private:
 	/** True if the shader graph has a diffuse color input. */
 	bool IsLambertModel(const UInterchangeShaderGraphNode* ShaderGraphNode) const;
 
-	/** True if the shader graph has the surface unlit's shader type name */
+	/** True if the shader graph has the surface unlit's shader type name. */
 	bool IsSurfaceUnlitModel(const UInterchangeShaderGraphNode* ShaderGraphNode) const;
 
 	/** True if the shader graph has an unlit color input. */
diff --git a/Engine/Plugins/Interchange/Runtime/Source/Pipelines/Public/InterchangeGenericMeshPipeline.h b/Engine/Plugins/Interchange/Runtime/Source/Pipelines/Public/InterchangeGenericMeshPipeline.h
--- a/Engine/Plugins/Interchange/Runtime/Source/Pipelines/Public/InterchangeGenericMeshPipeline.h
+++ b/Engine/Plugins/Interchange/Runtime/Source/Pipelines/Public/InterchangeGenericMeshPipeline.h
@@ -41,7 +41,7 @@ public:
 	
 	//////	STATIC_MESHES_CATEGORY Properties //////
 
-	/** If enabled, import the static mesh asset found in the sources. */
+	/** If enabled, imports all static mesh assets found in the sources. */
 	UPROPERTY(EditAnywhere, BlueprintReadWrite, Category = "Static Meshes")
 	bool bImportStaticMeshes = true;
 
@@ -49,12 +49,12 @@ public:
 	UPROPERTY(EditAnywhere, BlueprintReadWrite, Category = "Static Meshes")
 	bool bCombineStaticMeshes = false;
 
-	/** The LODGroup to associate with this mesh when it is imported */
-	UPROPERTY(EditAnywhere, BlueprintReadWrite, Category = "Static Meshes")
+	/** The LOD group that will be assigned to this mesh. */
+	UPROPERTY(EditAnywhere, BlueprintReadWrite, Category = "Static Meshes", meta=(DisplayName="LOD Group"))
 	FName LodGroup = NAME_None;
 
-	/** If checked, custom collision will be imported and if there is no custom collision, a generic collision will be automatically generated.
-	 * If unchecked there will be no collision created or imported.
+	/** If enabled, custom collision will be imported. If enabled and there is no custom collision, a generic collision will be automatically generated.
+	 * If disabled, no collision will be created or imported.
 	 */
 	UPROPERTY(EditAnywhere, BlueprintReadWrite, Category = "Static Meshes", meta = (SubCategory = "Collision"))
 	bool bImportCollision = true;
@@ -77,111 +77,118 @@ public:
 
 	//////	Static Meshes Build settings Properties //////
 
-	/** If enabled this option will allow you to use Nanite rendering at runtime. Can only be used with simple opaque materials. */
+	/** If enabled, imported meshes will be rendered by Nanite at runtime. Make sure your meshes and materials meet the requirements for Nanite. */
 	UPROPERTY(EditAnywhere, BlueprintReadWrite, Category = "Static Meshes", meta = (SubCategory = "Build"))
 	bool bBuildNanite = false;
 
-	/** If enabled this option will make sure the staticmesh build a reverse index buffer. */
+	/** If enabled, builds a reversed index buffer for each static mesh. */
 	UPROPERTY(EditAnywhere, BlueprintReadWrite, Category = "Static Meshes", meta = (SubCategory = "Build"))
 	bool bBuildReversedIndexBuffer = true;
 	
-	/** If enabled this option will generate lightmap for this staticmesh. */
+	/** If enabled, generates lightmap UVs for each static mesh. */
 	UPROPERTY(EditAnywhere, BlueprintReadWrite, Category = "Static Meshes", meta = (SubCategory = "Build"))
 	bool bGenerateLightmapUVs = true;
 	
 	/** 
-	 * Whether to generate the distance field treating every triangle hit as a front face.  
-	 * When enabled prevents the distance field from being discarded due to the mesh being open, but also lowers Distance Field AO quality.
+	 * Determines whether to generate the distance field treating every triangle hit as a front face.  
+	 * When enabled, prevents the distance field from being discarded due to the mesh being open, but also lowers distance field ambient occlusion quality.
 	 */
 	UPROPERTY(EditAnywhere, BlueprintReadWrite, Category = "Static Meshes", meta=(SubCategory = "Build", DisplayName="Two-Sided Distance Field Generation"))
 	bool bGenerateDistanceFieldAsIfTwoSided = false;
 	
+	/* If enabled, imported static meshes are set up for use with physical material masks. */
 	UPROPERTY(EditAnywhere, BlueprintReadWrite, Category = "Static Meshes", meta = (SubCategory = "Build", DisplayName="Enable Physical Material Mask"))
 	bool bSupportFaceRemap = false;
 	
+	/* When generating lightmaps, determines the amount of padding used to pack UVs. Set this value to the lowest-resolution lightmap you expect to use with the imported meshes. */
 	UPROPERTY(EditAnywhere, BlueprintReadWrite, Category = "Static Meshes", meta = (SubCategory = "Build"))
 	int32 MinLightmapResolution = 64;
 	
+	/* Specifies the index of the UV channel that will be used as the source when generating lightmaps. */
 	UPROPERTY(EditAnywhere, BlueprintReadWrite, Category = "Static Meshes", meta = (SubCategory = "Build", DisplayName="Source Lightmap Index"))
 	int32 SrcLightmapIndex = 0;
 	
+	/* Specifies the index of the UV channel that will store generated lightmap UVs. */
 	UPROPERTY(EditAnywhere, BlueprintReadWrite, Category = "Static Meshes", meta = (SubCategory = "Build", DisplayName="Destination Lightmap Index"))
 	int32 DstLightmapIndex = 1;
 	
-	/** The local scale applied when building the mesh */
+	/** The local scale applied when building the mesh. */
 	UPROPERTY(EditAnywhere, BlueprintReadWrite, Category = "Static Meshes", meta = (SubCategory = "Build", DisplayName="Build Scale"))
 	FVector BuildScale3D = FVector(1.0);
 	
 	/** 
 	 * Scale to apply to the mesh when allocating the distance field volume texture.
-	 * The default scale is 1, which is assuming that the mesh will be placed unscaled in the world.
+	 * The default scale is 1, which assumes that the mesh will be placed unscaled in the world.
 	 */
 	UPROPERTY(EditAnywhere, BlueprintReadWrite, Category = "Static Meshes", meta = (SubCategory = "Build"))
 	float DistanceFieldResolutionScale = 1.0f;
 	
+	/**
+	 * If set, replaces the distance field for all imported meshes with the distance field of the specified Static Mesh.
+	 */
 	UPROPERTY(EditAnywhere, BlueprintReadWrite, Category = "Static Meshes", meta = (SubCategory = "Build"))
 	TWeakObjectPtr<class UStaticMesh> DistanceFieldReplacementMesh = nullptr;
 	
 	/** 
-	 * Max Lumen mesh cards to generate for this mesh.
-	 * More cards means that surface will have better coverage, but will result in increased runtime overhead.
-	 * Set to 0 in order to disable mesh card generation for this mesh.
-	 * Default is 12.
+	 * The maximum number of Lumen mesh cards to generate for this mesh.
+	 * More cards means that the surface will have better coverage, but will result in increased runtime overhead.
+	 * Set this to 0 to disable mesh card generation for this mesh.
+	 * The default is 12.
 	 */
 	UPROPERTY(EditAnywhere, BlueprintReadWrite, Category = "Static Meshes", meta = (SubCategory = "Build"))
 	int32 MaxLumenMeshCards = 12;
 
 	//////	SKELETAL_MESHES_CATEGORY Properties //////
 
-	/** If enable, import the animation asset find in the sources. */
+	/** If enabled, imports all skeletal mesh assets found in the sources. */
 	UPROPERTY(EditAnywhere, BlueprintReadWrite, Category = "Skeletal Meshes")
 	bool bImportSkeletalMeshes = true;
 
-	/** Re-import partially or totally a skeletal mesh. You can choose betwwen geometry, skinning or everything.*/
+	/** Determines what types of information are imported for skeletal meshes. */
 	UPROPERTY(EditAnywhere, BlueprintReadWrite, Category = "Skeletal Meshes", meta = (DisplayName = "Import Content Type"))
 	EInterchangeSkeletalMeshContentType SkeletalMeshImportContentType;
 	
-	/** The value of the content type during the last import. This cannot be edited and is set only on successful import or re-import*/
+	/** The value of the content type during the last import. This cannot be edited and is set only on successful import or reimport. */
 	UPROPERTY()
 	EInterchangeSkeletalMeshContentType LastSkeletalMeshImportContentType;
 
-	/** If enable all translated skinned mesh node will be imported has a one skeletal mesh, note that it can still create several skeletal mesh for each different skeleton root joint. */
+	/** If enabled, all skinned mesh nodes that belong to the same skeleton root joint are combined into a single skeletal mesh. */
 	UPROPERTY(EditAnywhere, BlueprintReadWrite, Category = "Skeletal Meshes")
 	bool bCombineSkeletalMeshes = true;
 
-	/** If enable any morph target shape will be imported. */
+	/** If enabled, imports all morph target shapes found in the source. */
 	UPROPERTY(EditAnywhere, BlueprintReadWrite, Category = "Skeletal Meshes")
 	bool bImportMorphTargets = true;
 
-	/** True to import per-vertex attributes from the FBX file */
+	/** If enabled, imports per-vertex attributes from the FBX file. */
 	UPROPERTY(EditAnywhere, BlueprintReadWrite, Category = "Skeletal Meshes", meta = (ToolTip = "If enabled, creates named vertex attributes for secondary vertex color data."))
 	bool bImportVertexAttributes = false;
 
-	/** Enable this option to update Skeleton (of the mesh)'s reference pose. Mesh's reference pose is always updated.  */
+	/** Enable this option to update the reference pose of the Skeleton (of the mesh). The reference pose of the mesh is always updated.  */
 	UPROPERTY(EditAnywhere, BlueprintReadWrite, Category = "Skeletal Meshes")
 	bool bUpdateSkeletonReferencePose = false;
 
-	/** If checked, create new PhysicsAsset if it doesn't have it */
+	/** If enabled, create new PhysicsAsset if one doesn't exist. */
 	UPROPERTY(EditAnywhere, BlueprintReadWrite, Category = "Skeletal Meshes")
 	bool bCreatePhysicsAsset = true;
 
-	/** If this is set, use this specified PhysicsAsset. If its not set and bCreatePhysicsAsset is false, the importer will not generate or set any physic asset. */
+	/** If set, use the specified PhysicsAsset. If not set and the Create Physics Asset setting is not enabled, the importer will not generate or set any physics asset. */
 	UPROPERTY(EditAnywhere, BlueprintReadWrite, Category = "Skeletal Meshes", meta = (editcondition = "!bCreatePhysicsAsset"))
 	TWeakObjectPtr<UPhysicsAsset> PhysicsAsset;
 
-	/** If Checked, use 16 bits for skin weights instead of 8 bits. */
+	/** If enabled, imported skin weights use 16 bits instead of 8 bits. */
 	UPROPERTY(EditAnywhere, BlueprintReadWrite, Category = "Skeletal Meshes", meta = (SubCategory = "Build"))
 	bool bUseHighPrecisionSkinWeights = false;
 
-	/** Threshold use to decide if two vertex position are equal. */
+	/** Threshold value that is used to decide whether two vertex positions are equal. */
 	UPROPERTY(EditAnywhere, BlueprintReadWrite, Category = "Skeletal Meshes", meta = (SubCategory = "Build"))
 	float ThresholdPosition = 0.00002f;
 	
-	/** Threshold use to decide if two normal, tangents or bi-normals are equal. */
+	/** Threshold value that is used to decide whether two normals, tangents, or bi-normals are equal. */
 	UPROPERTY(EditAnywhere, BlueprintReadWrite, Category = "Skeletal Meshes", meta = (SubCategory = "Build"))
 	float ThresholdTangentNormal = 0.00002f;
 	
-	/** Threshold use to decide if two UVs are equal. */
+	/** Threshold value that is used to decide whether two UVs are equal. */
 	UPROPERTY(EditAnywhere, BlueprintReadWrite, Category = "Skeletal Meshes", meta = (SubCategory = "Build"))
 	float ThresholdUV = 0.0009765625f;
 	
diff --git a/Engine/Plugins/Interchange/Runtime/Source/Pipelines/Public/InterchangeGenericScenesPipeline.h b/Engine/Plugins/Interchange/Runtime/Source/Pipelines/Public/InterchangeGenericScenesPipeline.h
--- a/Engine/Plugins/Interchange/Runtime/Source/Pipelines/Public/InterchangeGenericScenesPipeline.h
+++ b/Engine/Plugins/Interchange/Runtime/Source/Pipelines/Public/InterchangeGenericScenesPipeline.h
@@ -21,23 +21,23 @@ class INTERCHANGEPIPELINES_API UInterchangeGenericLevelPipeline : public UInterc
 {
 	GENERATED_BODY()
 public:
-	/* Allow user to choose the re-import strategy when re-importing into level. */
+	/* Set the reimport strategy when reimporting into the level. */
 	UPROPERTY(EditAnywhere, BlueprintReadWrite, Category = "Level Reimport Policy", meta = (SubCategory = "Actors properties", AdjustPipelineAndRefreshDetailOnChange = "True"))
 	EReimportStrategyFlags ReimportPropertyStrategy = EReimportStrategyFlags::ApplyNoProperties;
 
-	/* Enables or not the deletion of actors which were not part of the translation when re-importing into level. */
+	/* If enabled, deletes actors that were not part of the translation when reimporting into a level. */
 	UPROPERTY(EditAnywhere, BlueprintReadWrite, Category = "Level Reimport Policy", meta = (SubCategory = "Reimport Actors"))
 	bool bDeleteMissingActors = false;
 
-	/* Enables or not spawning actors which were deleted in the editor prior to a reimport. */
+	/* If enabled, respawns actors that were deleted in the editor prior to a reimport. */
 	UPROPERTY(EditAnywhere, BlueprintReadWrite, Category = "Level Reimport Policy", meta = (SubCategory = "Reimport Actors"))
 	bool bForceReimportDeletedActors = false;
 
-	/* Enables or not re-creating assets which were deleted in the editor prior to a reimport into level. */
+	/* If enabled, recreates assets that were deleted in the editor prior to reimporting into a level. */
 	UPROPERTY(EditAnywhere, BlueprintReadWrite, Category = "Level Reimport Policy", meta = (SubCategory = "Reimport Assets"))
 	bool bForceReimportDeletedAssets = false;
 
-	/* Enables or not the deletion of assets which were not part of the translation when re-importing into level. */
+	/* If enabled, deletes assets that were not part of the translation when reimporting into a level. */
 	UPROPERTY(EditAnywhere, BlueprintReadWrite, Category = "Level Reimport Policy", meta = (SubCategory = "Reimport Assets"))
 	bool bDeleteMissingAssets = false;
 
diff --git a/Engine/Plugins/Interchange/Runtime/Source/Pipelines/Public/InterchangeGenericTexturePipeline.h b/Engine/Plugins/Interchange/Runtime/Source/Pipelines/Public/InterchangeGenericTexturePipeline.h
--- a/Engine/Plugins/Interchange/Runtime/Source/Pipelines/Public/InterchangeGenericTexturePipeline.h
+++ b/Engine/Plugins/Interchange/Runtime/Source/Pipelines/Public/InterchangeGenericTexturePipeline.h
@@ -23,47 +23,45 @@ public:
 	UPROPERTY(EditAnywhere, BlueprintReadWrite, Category = "Textures", meta = (StandAlonePipelineProperty = "True", PipelineInternalEditionData = "True"))
 	FString PipelineDisplayName;
 
-	/** If enabled, imports the texture assets found in the sources. */
+	/** If enabled, imports all texture assets found in the source. */
 	UPROPERTY(EditAnywhere, BlueprintReadWrite, Category = "Textures")
 	bool bImportTextures = true;
 
-	/** If not empty, and there is only one asset and one source data, we will name the asset with this string. */
+	/** If set, and there is only one asset and one source, the imported asset will be given this name. */
 	UPROPERTY(EditAnywhere, BlueprintReadWrite, Category = "Textures", meta=(StandAlonePipelineProperty = "True", AlwaysResetToDefault = "True"))
 	FString AssetName;
 
 #if WITH_EDITORONLY_DATA
 	/** 
-	 * If enable, after a new import a test will be run to see if the texture is a normal map.
-	 * If the texture is a normal map the SRG, CompressionSettings and LODGroup settings will be adjusted.
+	 * If enabled, tests each newly imported texture to see if it is a normal map.
+	 * If it is, the SRGB, Compression Settings, and LOD Group properties of that texture will be adjusted.
 	 */
 	UPROPERTY(EditAnywhere, BlueprintReadWrite, Category = "Textures", Meta=(EditCondition="bImportTextures"))
 	bool bDetectNormalMapTexture = true;
 
-	/** If enabled, the texture's green channel will be inverted for normal maps. */
+	/** If enabled, the texture's green channel will be inverted for normal maps. This setting is only used if the Detect Normal Map Texture setting is also enabled. */
 	UPROPERTY(EditAnywhere, BlueprintReadWrite, Category = "Textures", Meta=(EditCondition="bImportTextures"))
 	bool bFlipNormalMapGreenChannel = false;
 
-	/** If enabled detect if a texture use a UDIM pattern and if so import it as UIDMs. */
+	/** If enabled, imports textures as UDIMs if they are named using a UDIM naming pattern. */
 	UPROPERTY(EditAnywhere, BlueprintReadWrite, Category = "Textures", Meta=(EditCondition="bImportTextures"))
 	bool bImportUDIMs = true;
 
-	/** Specify the files type that should be imported as long/lat cubemap */
+	/** Specify the file types that should be imported as long/lat cubemaps. */
 	UPROPERTY(EditAnywhere, BlueprintReadWrite, Category = "Textures", Meta=(EditCondition="bImportTextures"))
 	TSet<FString> FileExtensionsToImportAsLongLatCubemap = {"hdr"};
 
 	/** 
-	 * If true, tell the translator to provide a compressed source data payload when available.
-	 * This will generally result in smaller assets, but some operations like the texture build might be slower because the source data will need to be uncompressed.
-	 * If false, it will let the factory or another step in the pipeline decide what to do.
-	 * 
-	 * Compressed source source data is generally store the data as it is in the source file.
+	 * If enabled, the translator compresses the source data payload whenever possible. This generally results in smaller assets.
+	 * However, some operations like the texture build might be slower, because the source data first needs to be decompressed.
+	 * If disabled, the translator leaves the decision to the factory or the pipelines.
 	 */
 	UPROPERTY(EditAnywhere, BlueprintReadWrite, AdvancedDisplay, Category = "Textures", Meta=(EditCondition="bImportTextures"))
 	bool bPreferCompressedSourceData = false;
 
 #endif
 
-	/** Should the textures that have a non-power of two resolution be imported */
+	/** If enabled, textures that have a non-power-of-two resolution are imported. */
 	UPROPERTY(EditAnywhere, BlueprintReadWrite, AdvancedDisplay, Category = "Textures", Meta=(EditCondition="bImportTextures"))
 	bool bAllowNonPowerOfTwo = false;
 
diff --git a/Engine/Plugins/Interchange/Runtime/Source/Pipelines/Public/InterchangeMaterialXPipeline.h b/Engine/Plugins/Interchange/Runtime/Source/Pipelines/Public/InterchangeMaterialXPipeline.h
--- a/Engine/Plugins/Interchange/Runtime/Source/Pipelines/Public/InterchangeMaterialXPipeline.h
+++ b/Engine/Plugins/Interchange/Runtime/Source/Pipelines/Public/InterchangeMaterialXPipeline.h
@@ -35,19 +35,19 @@ enum class EInterchangeMaterialXShaders : uint8
 	/** Open PBR Surface shader	used for translucency*/
 	OpenPBRSurfaceTransmission,
 
-	/** Default settings for Autodesk's Standard Surface shader	*/
+	/** Default settings for Autodesk's Standard Surface shader.	*/
 	StandardSurface,
 
-	/** Standard Surface shader used for translucency */
+	/** Standard Surface shader used for translucency. */
 	StandardSurfaceTransmission,
 
-	/** Shader used for unlit surface*/
+	/** Shader used for unlit surfaces. */
 	SurfaceUnlit,
 
-	/** Default settings for USD's Surface shader */
+	/** Default settings for USD's Surface shader. */
 	UsdPreviewSurface,
 
-	/** Construct a surface shader from scattering and emission distribution functions.*/
+	/** A surface shader constructed from scattering and emission distribution functions. */
 	Surface,
 
 	MaxShaderCount UMETA(hidden)
@@ -83,7 +83,7 @@ enum class EInterchangeMaterialXBSDF : uint8
 	/** A microfacet BSDF for the back-scattering properties of cloth-like materials. */
 	Sheen,
 
-	/** Adds an iridescent thin film layer over a microfacet base BSDF. */
+	/** An iridescent thin film layer over a microfacet base BSDF. */
 	ThinFilm,
 
 	MaxBSDFCount UMETA(hidden)
@@ -96,10 +96,10 @@ enum class EInterchangeMaterialXEDF : uint8
 	/** An EDF node for uniform emission. */
 	Uniform,
 
-	/** Constructs an EDF emitting light inside a cone around the normal direction. */
+	/** An EDF emitting light inside a cone around the normal direction. */
 	Conical,
 
-	/** Constructs an EDF emitting light according to a measured IES light profile. */
+	/** An EDF emitting light according to a measured IES light profile. */
 	Measured,
 
 	MaxEDFCount UMETA(hidden)
@@ -109,10 +109,10 @@ UENUM(BlueprintType)
 /** Data type representing a Volume Distribution Function. */
 enum class EInterchangeMaterialXVDF : uint8
 {
-	/** Constructs a VDF for pure light absorption. */
+	/** A VDF for pure light absorption. */
 	Absorption,
 
-	/** Constructs a VDF scattering light for a participating medium, based on the Henyey-Greenstein phase function. */
+	/** A VDF scattering light for a participating medium, based on the Henyey-Greenstein phase function. */
 	Anisotropic,
 
 	MaxVDFCount UMETA(hidden)
@@ -124,7 +124,7 @@ uint32 INTERCHANGEPIPELINES_API GetTypeHash(EMaterialXSettings Key);
 
 bool INTERCHANGEPIPELINES_API operator==(EMaterialXSettings Lhs, EMaterialXSettings Rhs);
 
-UCLASS(config = Interchange, meta = (DisplayName = "Interchange MaterialX"))
+UCLASS(config = Interchange, meta = (DisplayName = "Interchange MaterialX", ToolTip = "Interchange settings for MaterialX conversions."))
 class INTERCHANGEPIPELINES_API UMaterialXPipelineSettings : public UDeveloperSettings
 {
 	GENERATED_BODY()
