MaterialX:
Fixing MaterialX crash caused by a wrong override of the material functions in the project settings that ends up trying loading one of the material function outside the Game Thread.
The idea behind this rework, is to have a better control of the Material Function paths, instead of using directly the different paths, using the types given by the different enums that were originally defined in the MaterialX pipeline is less error-prone.

- Moving InterchangeMaterialXDefinitions.h intoits own module/
- Moving the different Shader/BSDF enums from MaterialX Pipeline to InterchangeMaterialXDefinitions.h
- Refactoring the MaterialXManager matching material functions path data structure, using either a path (used by MF that are not overridable) or the enums types of the Shaders/BSDF
- Fixing MaterialXPipeline taking the wrong material function based on Substrate boolean
- Changing the different Shader surface to use the enum values instead of the material function path.
- Fixing OpenPBR normal textures to be created in tangent space (avoiding the prompt message in the editor by atuo-detecting normal maps)

diff --git a/Engine/Plugins/Interchange/Editor/Source/Pipelines/Private/InterchangeMaterialXPipelineCustomizations.cpp b/Engine/Plugins/Interchange/Editor/Source/Pipelines/Private/InterchangeMaterialXPipelineCustomizations.cpp
--- a/Engine/Plugins/Interchange/Editor/Source/Pipelines/Private/InterchangeMaterialXPipelineCustomizations.cpp
+++ b/Engine/Plugins/Interchange/Editor/Source/Pipelines/Private/InterchangeMaterialXPipelineCustomizations.cpp
@@ -126,22 +126,22 @@ void FInterchangeMaterialXPipelineSettingsCustomization::CustomizeDetails(IDetai
 	Customize(GET_MEMBER_NAME_CHECKED(UMaterialXPipelineSettings, PredefinedSurfaceShaders),
 			  TEXT("MaterialXPredefined | Surface Shaders"),
 			  NSLOCTEXT("InterchangeMaterialXPipelineSettingsCustomization", "MaterialXPredefined | SurfaceShaders", "Surface Shaders"),
-			  UMaterialXPipelineSettings::IndexSurfaceShaders);
+			  UE::Interchange::MaterialX::IndexSurfaceShaders);
 
 	Customize(GET_MEMBER_NAME_CHECKED(UMaterialXPipelineSettings, PredefinedBSDF),
 			  TEXT("MaterialXPredefined | BSDF"),
 			  NSLOCTEXT("InterchangeMaterialXPipelineSettingsCustomization", "MaterialXPredefined | BSDF", "Bidirectional Scattering Distribution Functions"),
-			  UMaterialXPipelineSettings::IndexBSDF);
+			  UE::Interchange::MaterialX::IndexBSDF);
 
 	Customize(GET_MEMBER_NAME_CHECKED(UMaterialXPipelineSettings, PredefinedEDF),
 			  TEXT("MaterialXPredefined | EDF"),
 			  NSLOCTEXT("InterchangeMaterialXPipelineSettingsCustomization", "MaterialXPredefined | EDF", "Emission Distribution Functions"),
-			  UMaterialXPipelineSettings::IndexEDF);
+			  UE::Interchange::MaterialX::IndexEDF);
 
 	Customize(GET_MEMBER_NAME_CHECKED(UMaterialXPipelineSettings, PredefinedVDF),
 			  TEXT("MaterialXPredefined | VDF"),
 			  NSLOCTEXT("InterchangeMaterialXPipelineSettingsCustomization", "MaterialXPredefined | VDF", "Volume Distribution Functions"),
-			  UMaterialXPipelineSettings::IndexVDF);
+			  UE::Interchange::MaterialX::IndexVDF);
 }
 
 bool FInterchangeMaterialXPipelineSettingsCustomization::OnShouldFilterAssetEnum(const FAssetData& InAssetData, uint8 EnumType, uint8 EnumValue)
diff --git a/Engine/Plugins/Interchange/Runtime/Config/BaseInterchange.ini b/Engine/Plugins/Interchange/Runtime/Config/BaseInterchange.ini
--- a/Engine/Plugins/Interchange/Runtime/Config/BaseInterchange.ini
+++ b/Engine/Plugins/Interchange/Runtime/Config/BaseInterchange.ini
@@ -38,6 +38,10 @@ bCanSaveAllSections=True
 +ClassRedirects=(OldName="/Script/InterchangeImport.MaterialExpressionUnpremult", NewName="/Script/InterchangeImport.MaterialExpressionMaterialXUnpremult")
 +EnumRedirects=(OldName="ETextureSampleBlurFilter", NewName="EMaterialXTextureSampleBlurFilter")
 +EnumRedirects=(OldName="ELuminanceMode", NewName="EMaterialXLuminanceMode")
++EnumRedirects=(OldName="/Script/InterchangePipelines.EInterchangeMaterialXShaders>", NewName="/Script/InterchangeCommon.EInterchangeMaterialXShaders>")
++EnumRedirects=(OldName="/Script/InterchangePipelines.EInterchangeMaterialXBSDF>", NewName="/Script/InterchangeCommon.EInterchangeMaterialXBSDF>")
++EnumRedirects=(OldName="/Script/InterchangePipelines.EInterchangeMaterialXEDF>", NewName="/Script/InterchangeCommon.EInterchangeMaterialXEDF>")
++EnumRedirects=(OldName="/Script/InterchangePipelines.EInterchangeMaterialXVDF>", NewName="/Script/InterchangeCommon.EInterchangeMaterialXVDF>")
 
 ; Rename InterchangeAnimationTrackSetFactory to InterchangeLevelSequenceFactory
 +ClassRedirects=(OldName="/Script/InterchangeImport.InterchangeAnimationTrackSetFactory", NewName="/Script/InterchangeImport.InterchangeLevelSequenceFactory")
diff --git a/Engine/Plugins/Interchange/Runtime/Interchange.uplugin b/Engine/Plugins/Interchange/Runtime/Interchange.uplugin
--- a/Engine/Plugins/Interchange/Runtime/Interchange.uplugin
+++ b/Engine/Plugins/Interchange/Runtime/Interchange.uplugin
@@ -61,6 +61,11 @@
 			"Type": "Runtime",
 			"LoadingPhase": "Default"
 		},
+		{
+			"Name": "InterchangeCommon",
+			"Type": "Runtime",
+			"LoadingPhase": "Default"
+		},
 		{
 			"Name": "InterchangeCommonParser",
 			"Type": "Runtime",
diff --git a/Engine/Plugins/Interchange/Runtime/Source/Common/InterchangeCommon.Build.cs b/Engine/Plugins/Interchange/Runtime/Source/Common/InterchangeCommon.Build.cs
new file mode 100644
--- /dev/null
+++ b/Engine/Plugins/Interchange/Runtime/Source/Common/InterchangeCommon.Build.cs
@@ -0,0 +1,20 @@
+// Copyright Epic Games, Inc. All Rights Reserved.
+
+namespace UnrealBuildTool.Rules
+{
+	public class InterchangeCommon : ModuleRules
+	{
+		public InterchangeCommon(ReadOnlyTargetRules Target) : base(Target)
+		{
+			PublicDependencyModuleNames.AddRange(
+				new string[]
+				{
+					"Core",
+					"CoreUObject"
+				}
+			);
+
+			PublicIncludePathModuleNames.Add("MaterialX");
+		}
+	}
+}
diff --git a/Engine/Plugins/Interchange/Runtime/Source/Common/Private/InterchangeCommonModule.cpp b/Engine/Plugins/Interchange/Runtime/Source/Common/Private/InterchangeCommonModule.cpp
new file mode 100644
--- /dev/null
+++ b/Engine/Plugins/Interchange/Runtime/Source/Common/Private/InterchangeCommonModule.cpp
@@ -0,0 +1,13 @@
+// Copyright Epic Games, Inc. All Rights Reserved.
+
+#include "CoreMinimal.h"
+#include "Modules/ModuleInterface.h"
+#include "Modules/ModuleManager.h"
+
+
+class FInterchangeCommonModule : public IModuleInterface
+{
+};
+
+IMPLEMENT_MODULE(FInterchangeCommonModule, InterchangeCommon)
+
diff --git a/Engine/Plugins/Interchange/Runtime/Source/Common/Public/MaterialX/InterchangeMaterialXDefinitions.h b/Engine/Plugins/Interchange/Runtime/Source/Common/Public/MaterialX/InterchangeMaterialXDefinitions.h
new file mode 100644
--- /dev/null
+++ b/Engine/Plugins/Interchange/Runtime/Source/Common/Public/MaterialX/InterchangeMaterialXDefinitions.h
@@ -0,0 +1,724 @@
+// Copyright Epic Games, Inc. All Rights Reserved. 
+#pragma once
+
+#include "Math/Color.h"
+#include "Math/MathFwd.h"
+#include "Math/Vector.h"
+
+// Interchange currently only brings in MaterialX when in the editor, so we
+// define its namespace macros ourselves otherwise.
+#if WITH_EDITOR
+
+THIRD_PARTY_INCLUDES_START
+#include "MaterialXCore/Library.h"
+THIRD_PARTY_INCLUDES_END
+
+#else
+
+#define MATERIALX_NAMESPACE_BEGIN namespace MaterialX {
+#define MATERIALX_NAMESPACE_END }
+
+#endif
+
+
+#include "InterchangeMaterialXDefinitions.generated.h"
+
+UENUM(BlueprintType)
+enum class EInterchangeMaterialXShaders : uint8
+{
+	/** Default settings for Open PBR Surface shader. */
+	OpenPBRSurface,
+
+	/** Open PBR Surface shader	used for translucency. */
+	OpenPBRSurfaceTransmission,
+
+	/** Default settings for Autodesk's Standard Surface shader. */
+	StandardSurface,
+
+	/** Standard Surface shader used for translucency. */
+	StandardSurfaceTransmission,
+
+	/** Shader used for unlit surfaces. */
+	SurfaceUnlit,
+
+	/** Default settings for USD's Surface shader. */
+	UsdPreviewSurface,
+
+	/** A surface shader constructed from scattering and emission distribution functions. */
+	Surface,
+
+	MaxShaderCount UMETA(hidden)
+};
+
+UENUM(BlueprintType)
+/** Data type representing a Bidirectional Scattering Distribution Function. */
+enum class EInterchangeMaterialXBSDF : uint8
+{
+	/** A BSDF node for diffuse reflections. */
+	OrenNayarDiffuse,
+
+	/** A BSDF node for Burley diffuse reflections. */
+	BurleyDiffuse,
+
+	/** A BSDF node for pure diffuse transmission. */
+	Translucent,
+
+	/** A reflection/transmission BSDF node based on a microfacet model and a Fresnel curve for dielectrics. */
+	Dielectric,
+
+	/** A reflection BSDF node based on a microfacet model and a Fresnel curve for conductors/metals. */
+	Conductor,
+
+	/** A reflection/transmission BSDF node based on a microfacet model and a generalized Schlick Fresnel curve. */
+	GeneralizedSchlick,
+
+	/** A subsurface scattering BSDF for true subsurface scattering. */
+	Subsurface,
+
+	/** A microfacet BSDF for the back-scattering properties of cloth-like materials. */
+	Sheen,
+
+	/** Adds an iridescent thin film layer over a microfacet base BSDF. */
+	ThinFilm,
+
+	MaxBSDFCount UMETA(hidden)
+};
+
+UENUM(BlueprintType)
+/** Data type representing an Emission Distribution Function. */
+enum class EInterchangeMaterialXEDF : uint8
+{
+	/** An EDF node for uniform emission. */
+	Uniform,
+
+	/** Constructs an EDF emitting light inside a cone around the normal direction. */
+	Conical,
+
+	/** Constructs an EDF emitting light according to a measured IES light profile. */
+	Measured,
+
+	MaxEDFCount UMETA(hidden)
+};
+
+UENUM(BlueprintType)
+/** Data type representing a Volume Distribution Function. */
+enum class EInterchangeMaterialXVDF : uint8
+{
+	/** Constructs a VDF for pure light absorption. */
+	Absorption,
+
+	/** Constructs a VDF scattering light for a participating medium, based on the Henyey-Greenstein phase function. */
+	Anisotropic,
+
+	MaxVDFCount UMETA(hidden)
+};
+
+namespace UE
+{
+	namespace Interchange
+	{
+		namespace MaterialX
+		{
+			static constexpr uint8 IndexSurfaceShaders = 0;
+			static constexpr uint8 IndexBSDF = 1;
+			static constexpr uint8 IndexEDF = 2;
+			static constexpr uint8 IndexVDF = 3;
+
+			namespace Attributes
+			{
+				constexpr const TCHAR* EnumType = TEXT("MaterialXEnumType");
+				constexpr const TCHAR* EnumValue = TEXT("MaterialXEnumValue");
+			}
+		}
+	}
+}
+
+MATERIALX_NAMESPACE_BEGIN
+
+	namespace OpenPBRSurface
+	{
+		namespace Input
+		{
+			static constexpr const char* BaseWeight = "base_weight";
+			static constexpr const char* BaseColor = "base_color";
+			static constexpr const char* BaseRoughness = "base_roughness";
+			static constexpr const char* BaseMetalness = "base_metalness";
+			static constexpr const char* SpecularWeight = "specular_weight";
+			static constexpr const char* SpecularColor = "specular_color";
+			static constexpr const char* SpecularRoughness = "specular_roughness";
+			static constexpr const char* SpecularIOR = "specular_ior";
+			static constexpr const char* SpecularIORLevel = "specular_ior_level";
+			static constexpr const char* SpecularAnisotropy = "specular_anisotropy";
+			static constexpr const char* SpecularRotation = "specular_rotation";
+			static constexpr const char* TransmissionWeight = "transmission_weight";
+			static constexpr const char* TransmissionColor = "transmission_color";
+			static constexpr const char* TransmissionDepth = "transmission_depth";
+			static constexpr const char* TransmissionScatter = "transmission_scatter";
+			static constexpr const char* TransmissionScatterAnisotropy= "transmission_scatter_anisotropy";
+			static constexpr const char* TransmissionDispersion = "transmission_dispersion";
+			static constexpr const char* SubsurfaceWeight = "subsurface_weight";
+			static constexpr const char* SubsurfaceColor = "subsurface_color";
+			static constexpr const char* SubsurfaceRadius = "subsurface_radius";
+			static constexpr const char* SubsurfaceRadiusScale = "subsurface_radius_scale";
+			static constexpr const char* SubsurfaceAnisotropy = "subsurface_anisotropy";
+			static constexpr const char* FuzzWeight = "fuzz_weight";
+			static constexpr const char* FuzzColor = "fuzz_color";
+			static constexpr const char* FuzzRoughness = "fuzz_roughness";
+			static constexpr const char* CoatWeight = "coat_weight";
+			static constexpr const char* CoatColor = "coat_color";
+			static constexpr const char* CoatRoughness = "coat_roughness";
+			static constexpr const char* CoatAnisotropy = "coat_anisotropy";
+			static constexpr const char* CoatRotation = "coat_rotation";
+			static constexpr const char* CoatIOR = "coat_ior";
+			static constexpr const char* CoatIORLevel = "coat_ior_level";
+			static constexpr const char* ThinFilmThickness = "thin_film_thickness";
+			static constexpr const char* ThinFilmIOR = "thin_film_ior";
+			static constexpr const char* EmissionLuminance = "emission_luminance";
+			static constexpr const char* EmissionColor = "emission_color";
+			static constexpr const char* GeometryOpacity = "geometry_opacity";
+			static constexpr const char* GeometryThinWalled = "geometry_thin_walled";
+			static constexpr const char* GeometryNormal = "geometry_normal";
+			static constexpr const char* GeometryCoatNormal = "geometry_coat_normal";
+			static constexpr const char* GeometryTangent = "geometry_tangent";
+		}
+
+		namespace DefaultValue
+		{
+			static constexpr float BaseWeight = 1.f;
+			static constexpr FLinearColor BaseColor{ 0.8, 0.8, 0.8 };
+			static constexpr float BaseRoughness = 0.f;
+			static constexpr float BaseMetalness = 0.f;
+			static constexpr float SpecularWeight = 1.f;
+			static constexpr FLinearColor SpecularColor{ 1, 1, 1 };
+			static constexpr float SpecularRoughness = 0.3f;
+			static constexpr float SpecularIOR = 1.5f;
+			static constexpr float SpecularIORLevel = 0.5f;
+			static constexpr float SpecularAnisotropy = 0.f;
+			static constexpr float SpecularRotation = 0.f;
+			static constexpr float TransmissionWeight = 0.f;
+			static constexpr FLinearColor TransmissionColor{ 1, 1, 1 };
+			static constexpr float TransmissionDepth = 0.f;
+			static constexpr FLinearColor TransmissionScatter{ 0, 0, 0 };
+			static constexpr float TransmissionScatterAnisotropy = 0.f;
+			static constexpr float TransmissionDispersion = 0.f;
+			static constexpr float SubsurfaceWeight = 0.f;
+			static constexpr FLinearColor SubsurfaceColor{ 0.8, 0.8, 0.8 };
+			static constexpr float SubsurfaceRadius = 1.f;
+			static constexpr FLinearColor SubsurfaceRadiusScale{1, 0.5, 0.25};
+			static constexpr float SubsurfaceAnisotropy = 0.f;
+			static constexpr float FuzzWeight = 0.f;
+			static constexpr FLinearColor FuzzColor{ 1, 1, 1 };
+			static constexpr float FuzzRoughness = 0.5f;
+			static constexpr float CoatWeight = 0.f;
+			static constexpr FLinearColor CoatColor{ 1, 1, 1 };
+			static constexpr float CoatRoughness = 0.f;
+			static constexpr float CoatAnisotropy = 0.f;
+			static constexpr float CoatRotation = 0.f;
+			static constexpr float CoatIOR = 1.6f;
+			static constexpr float CoatIORLevel = 0.5f;
+			static constexpr float ThinFilmThickness = 0.f;
+			static constexpr float ThinFilmIOR = 1.5f;
+			static constexpr float EmissionLuminance = 0.f;
+			static constexpr FLinearColor EmissionColor{ 1, 1, 1 };
+			static constexpr FLinearColor GeometryOpacity{ 1, 1, 1 };
+			static constexpr bool GeometryThinWalled = false;
+			static const FVector GeometryNormal{ 0, 0, 1 };
+			static const FVector GeometryCoatNormal{ 0, 0, 1 };
+			static const FVector GeometryTangent{ 0, 1, 0 };
+		}
+	}
+
+	namespace StandardSurface
+	{
+		namespace Input
+		{
+			static constexpr const char* Base = "base";
+			static constexpr const char* BaseColor = "base_color";
+			static constexpr const char* DiffuseRoughness = "diffuse_roughness";
+			static constexpr const char* Metalness = "metalness";
+			static constexpr const char* Specular = "specular";
+			static constexpr const char* SpecularColor = "specular_color";
+			static constexpr const char* SpecularRoughness = "specular_roughness";
+			static constexpr const char* SpecularIOR = "specular_IOR";
+			static constexpr const char* SpecularAnisotropy = "specular_anisotropy";
+			static constexpr const char* SpecularRotation = "specular_rotation";
+			static constexpr const char* Transmission = "transmission";
+			static constexpr const char* TransmissionColor = "transmission_color";
+			static constexpr const char* TransmissionDepth = "transmission_depth";
+			static constexpr const char* TransmissionScatter = "transmission_scatter";
+			static constexpr const char* TransmissionScatterAnisotropy = "transmission_scatter_anisotropy";
+			static constexpr const char* TransmissionDispersion = "transmission_dispersion";
+			static constexpr const char* TransmissionExtraRoughness = "transmission_extra_roughness";
+			static constexpr const char* Subsurface = "subsurface";
+			static constexpr const char* SubsurfaceColor = "subsurface_color";
+			static constexpr const char* SubsurfaceRadius = "subsurface_radius";
+			static constexpr const char* SubsurfaceScale = "subsurface_scale";
+			static constexpr const char* SubsurfaceAnisotropy = "subsurface_anisotropy";
+			static constexpr const char* Sheen = "sheen";
+			static constexpr const char* SheenColor = "sheen_color";
+			static constexpr const char* SheenRoughness = "sheen_roughness";
+			static constexpr const char* Coat = "coat";
+			static constexpr const char* CoatColor = "coat_color";
+			static constexpr const char* CoatRoughness = "coat_roughness";
+			static constexpr const char* CoatAnisotropy = "coat_anisotropy";
+			static constexpr const char* CoatRotation = "coat_rotation";
+			static constexpr const char* CoatIOR = "coat_IOR";
+			static constexpr const char* CoatNormal = "coat_normal";
+			static constexpr const char* CoatAffectColor = "coat_affect_color";
+			static constexpr const char* CoatAffectRoughness = "coat_affect_roughness";
+			static constexpr const char* ThinFilmThickness = "thin_film_thickness";
+			static constexpr const char* ThinFilmIOR = "thin_film_IOR";
+			static constexpr const char* Emission = "emission";
+			static constexpr const char* EmissionColor = "emission_color";
+			static constexpr const char* Opacity = "opacity";
+			static constexpr const char* ThinWalled = "thin_walled";
+			static constexpr const char* Normal = "normal";
+			static constexpr const char* Tangent = "tangent";
+		}
+
+		namespace DefaultValue
+		{
+			namespace Float
+			{
+				static float Base = 1.f;
+				static float DiffuseRoughness = 0.f;
+				static float Metalness = 0.f;
+				static float Specular = 1.f;
+				static float SpecularRoughness = 0.2f;
+				static float SpecularIOR = 1.5f;
+				static float SpecularAnisotropy = 0.f;
+				static float SpecularRotation = 0.f;
+				static float Transmission = 0.f;
+				static float TransmissionDepth = 0.f;
+				static float TransmissionScatterAnisotropy = 0.f;
+				static float TransmissionDispersion = 0.f;
+				static float TransmissionExtraRoughness = 0.f;
+				static float Subsurface = 0.f;
+				static float SubsurfaceScale = 1.f;
+				static float SubsurfaceAnisotropy = 0.f;
+				static float Sheen = 0.f;
+				static float SheenRoughness = 0.3f;
+				static float Coat = 0.f;
+				static float CoatRoughness = 0.1f;
+				static float CoatAnisotropy = 0.f;
+				static float CoatRotation = 0.f;
+				static float CoatIOR = 1.5f;
+				static float CoatAffectColor = 0.f;
+				static float CoatAffectRoughness = 0.f;
+				static float ThinFilmThickness = 0.f;
+				static float ThinFilmIOR = 1.5f;
+				static float Emission = 0.f;
+			}
+
+			namespace Color3
+			{
+				constexpr FLinearColor BaseColor{ 0.8, 0.8, 0.8 };
+				constexpr FLinearColor SpecularColor{ 1, 1, 1 };
+				constexpr FLinearColor TransmissionColor{ 1, 1, 1 };
+				constexpr FLinearColor TransmissionScatter{ 0, 0, 0 };
+				constexpr FLinearColor SubsurfaceColor{ 1, 1, 1 };
+				constexpr FLinearColor SubsurfaceRadius{ 1, 1, 1 };
+				constexpr FLinearColor SheenColor{ 1, 1, 1 };
+				constexpr FLinearColor CoatColor{ 1, 1, 1 };
+				constexpr FLinearColor EmissionColor{ 1, 1, 1 };
+				constexpr FLinearColor Opacity{ 1, 1, 1 };
+			}
+		}
+	}
+
+	namespace SurfaceUnlit
+	{
+		namespace Input
+		{
+			static constexpr const char* Emission = "emission";
+			static constexpr const char* EmissionColor = "emission_color";
+			static constexpr const char* Transmission = "transmission";
+			static constexpr const char* TransmissionColor = "transmission_color";
+			static constexpr const char* Opacity = "opacity";
+		}
+
+		namespace DefaultValue
+		{
+			namespace Float
+			{
+				constexpr float Emission = 1.f;
+				constexpr float Transmission = 0.f;
+				constexpr float Opacity = 1.f;
+			}
+
+			namespace Color3
+			{
+				constexpr FLinearColor EmissionColor{ 1, 1, 1 };
+				constexpr FLinearColor TransmissionColor{ 1, 1, 1 };
+			}
+		}
+	}
+
+	namespace Surface
+	{
+		namespace Input
+		{
+			static constexpr const char* Bsdf = "bsdf";
+			static constexpr const char* Edf = "edf";
+			static constexpr const char* Opacity = "opacity";
+		}
+
+		namespace DefaultValue
+		{
+			namespace Float
+			{
+				constexpr float Opacity = 1.f;
+			}
+		}
+	}
+
+	namespace UsdPreviewSurface
+	{
+		namespace Input
+		{
+			static constexpr const char* DiffuseColor = "diffuseColor";
+			static constexpr const char* EmissiveColor = "emissiveColor";
+			static constexpr const char* SpecularColor = "specularColor";
+			static constexpr const char* Metallic = "metallic";
+			static constexpr const char* Roughness = "roughness";
+			static constexpr const char* Clearcoat = "clearcoat";
+			static constexpr const char* ClearcoatRoughness = "clearcoatRoughness";
+			static constexpr const char* Opacity = "opacity";
+			static constexpr const char* OpacityThreshold = "opacityThreshold";
+			static constexpr const char* IOR = "ior";
+			static constexpr const char* Normal = "normal";
+			static constexpr const char* Displacement = "displacement";
+			static constexpr const char* Occlusion = "occlusion";
+		}
+
+		namespace DefaultValue
+		{
+			namespace Float
+			{
+				constexpr float Metallic = 0.f;
+				constexpr float Roughness = 0.5f;
+				constexpr float Clearcoat = 0.f;
+				constexpr float ClearcoatRoughness = 0.01f;
+				constexpr float Opacity = 1.f;
+				constexpr float OpacityThreshold = 0.f;
+				constexpr float IOR = 1.5f;
+				constexpr float Displacement = 0.f;
+				constexpr float Occlusion = 1.f;
+			}
+
+			namespace Color3
+			{
+				constexpr FLinearColor DiffuseColor{ 0.18f, 0.18f, 0.18f };
+				constexpr FLinearColor EmissiveColor{ 0, 0, 0 };
+				constexpr FLinearColor SpecularColor{ 0, 0, 0 };
+			}
+			
+			namespace Vector3
+			{
+				static const FVector3f Normal{ 0.f, 0.f, 1.f };
+			}
+		}
+	}
+
+	namespace Lights
+	{
+		//There's no input per se in a Light, but we can find some common inputs among those lights
+		namespace Input
+		{
+			static constexpr const char* Color = "color";
+			static constexpr const char* Intensity = "intensity";
+		}
+
+		namespace PointLight
+		{
+			namespace Input
+			{
+				using namespace Lights::Input;
+				static constexpr const char* Position = "position";
+				static constexpr const char* DecayRate = "decay_rate";
+			}
+		}
+
+		namespace DirectionalLight
+		{
+			namespace Input
+			{
+				using namespace Lights::Input;
+				static constexpr const char* Direction = "direction";
+			}
+		}
+
+		namespace SpotLight
+		{
+			namespace Input
+			{
+				using namespace PointLight::Input;
+				using namespace DirectionalLight::Input;
+				static constexpr const char* InnerAngle = "inner_angle";
+				static constexpr const char* OuterAngle = "outer_angle";
+			}
+		}
+	}
+
+	namespace Attributes
+	{
+		static constexpr const char* IsVisited = "UE:IsVisited";
+		static constexpr const char* NewName = "UE:NewName";
+		static constexpr const char* ParentName = "UE:ParentName";
+	}
+
+	namespace Category
+	{
+		// Math nodes
+		static constexpr const char* Absval = "absval";
+		static constexpr const char* Acos = "acos";
+		static constexpr const char* Add = "add";
+		static constexpr const char* ArrayAppend = "arrayappend";
+		static constexpr const char* Asin = "asin";
+		static constexpr const char* Atan2 = "atan2";
+		static constexpr const char* Ceil = "ceil";
+		static constexpr const char* Clamp = "clamp";
+		static constexpr const char* Cos = "cos";
+		static constexpr const char* CrossProduct = "crossproduct";
+		static constexpr const char* Determinant = "determinant";
+		static constexpr const char* Divide = "divide";
+		static constexpr const char* DotProduct = "dotproduct";
+		static constexpr const char* Exp = "exp";
+		static constexpr const char* Floor = "floor";
+		static constexpr const char* Invert = "invert";
+		static constexpr const char* InvertMatrix = "invertmatrix";
+		static constexpr const char* Ln = "ln";
+		static constexpr const char* Magnitude = "magnitude";
+		static constexpr const char* Max = "max";
+		static constexpr const char* Min = "min";
+		static constexpr const char* Modulo = "modulo";
+		static constexpr const char* Multiply = "multiply";
+		static constexpr const char* Normalize = "normalize";
+		static constexpr const char* NormalMap = "normalmap";
+		static constexpr const char* Place2D = "place2d";
+		static constexpr const char* Power = "power";
+		static constexpr const char* Rotate2D = "rotate2d";
+		static constexpr const char* Rotate3D = "rotate3d";
+		static constexpr const char* Sign = "sign";
+		static constexpr const char* Sin = "sin";
+		static constexpr const char* Sqrt = "sqrt";
+		static constexpr const char* Sub = "subtract";
+		static constexpr const char* Tan = "tan";
+		static constexpr const char* TransformMatrix = "transformmatrix";
+		static constexpr const char* TransformNormal = "transformnormal";
+		static constexpr const char* TransformPoint = "transformpoint";
+		static constexpr const char* TransformVector = "transformvector";
+		static constexpr const char* Transpose = "transpose";		
+		// Compositing nodes
+		static constexpr const char* Burn = "burn";
+		static constexpr const char* Difference = "difference";
+		static constexpr const char* Disjointover = "disjointover";
+		static constexpr const char* Dodge = "dodge";
+		static constexpr const char* In = "in";
+		static constexpr const char* Inside = "inside";
+		static constexpr const char* Mask = "mask";
+		static constexpr const char* Matte = "matte";
+		static constexpr const char* Minus = "minus";
+		static constexpr const char* Mix = "mix";
+		static constexpr const char* Out = "out";
+		static constexpr const char* Outside = "outside";
+		static constexpr const char* Over = "over";
+		static constexpr const char* Overlay = "overlay";
+		static constexpr const char* Plus = "plus";
+		static constexpr const char* Premult = "premult";
+		static constexpr const char* Screen = "screen";
+		static constexpr const char* Unpremult = "unpremult";
+		// Conditional nodes
+		static constexpr const char* IfGreater = "ifgreater";
+		static constexpr const char* IfGreaterEq = "ifgreatereq";
+		static constexpr const char* IfEqual = "ifequal";
+		// Channel nodes
+		static constexpr const char* Convert = "convert";
+		static constexpr const char* Combine2 = "combine2";
+		static constexpr const char* Combine3 = "combine3";
+		static constexpr const char* Combine4 = "combine4";
+		static constexpr const char* Extract = "extract";
+		static constexpr const char* Separate2 = "separate2";
+		static constexpr const char* Separate3 = "separate3";
+		static constexpr const char* Separate4 = "separate4";
+		static constexpr const char* Swizzle = "swizzle";
+		// Procedural nodes 
+		static constexpr const char* Constant = "constant";
+		// Procedural2D nodes 
+		static constexpr const char* Ramp4 = "ramp4";
+		static constexpr const char* RampLR = "ramplr";
+		static constexpr const char* RampTB = "ramptb";
+		static constexpr const char* SplitLR = "splitlr";
+		static constexpr const char* SplitTB = "splittb";
+		// Procedural3D nodes 
+		static constexpr const char* CellNoise3D = "cellnoise3d";
+		static constexpr const char* Fractal3D = "fractal3d";
+		static constexpr const char* Noise3D = "noise3d";
+		static constexpr const char* WorleyNoise3D = "worleynoise3d";
+		// Organization nodes 
+		static constexpr const char* Dot = "dot";
+		// Texture nodes 
+		static constexpr const char* Image = "image";
+		static constexpr const char* TiledImage = "tiledimage";
+		// Geometric nodes
+		static constexpr const char* Bitangent = "bitangent";
+		static constexpr const char* GeomColor = "geomcolor";
+		static constexpr const char* Normal = "normal";
+		static constexpr const char* Position = "position";
+		static constexpr const char* Tangent = "tangent";
+		static constexpr const char* TexCoord = "texcoord";
+		// Light nodes
+		static constexpr const char* DirectionalLight = "directional_light";
+		static constexpr const char* PointLight = "point_light";
+		static constexpr const char* SpotLight = "spot_light";
+		// Adjustment nodes
+		static constexpr const char* HsvToRgb = "hsvtorgb";
+		static constexpr const char* Luminance = "luminance";
+		static constexpr const char* Remap = "remap";
+		static constexpr const char* RgbToHsv = "rgbtohsv";
+		static constexpr const char* Saturate = "saturate";
+		static constexpr const char* Smoothstep = "smoothstep";
+		// Application
+		static constexpr const char* Time= "time";
+		// PBR
+		// BSDF
+		static constexpr const char* BurleyDiffuseBSDF= "burley_diffuse_bsdf";
+		static constexpr const char* ConductorBSDF = "conductor_bsdf";
+		static constexpr const char* DielectricBSDF = "dielectric_bsdf";
+		static constexpr const char* GeneralizedSchlickBSDF= "generalized_schlick_bsdf";
+		static constexpr const char* OrenNayarDiffuseBSDF = "oren_nayar_diffuse_bsdf";
+		static constexpr const char* SheenBSDF = "sheen_bsdf";
+		static constexpr const char* SubsurfaceBSDF= "subsurface_bsdf";
+		static constexpr const char* ThinFilmBSDF= "thin_film_bsdf";
+		static constexpr const char* TranslucentBSDF = "translucent_bsdf";
+		// EDF
+		static constexpr const char* ConicalEDF = "conical_edf";
+		static constexpr const char* MeasuredEDF = "measured_edf";
+		static constexpr const char* UniformEDF = "uniform_edf";
+		// VDF
+		static constexpr const char* AbsorptionVDF = "absorption_vdf";
+		static constexpr const char* AnisotropicVDF = "anisotropic_vdf";
+		// PBR Utility Nodes
+		static constexpr const char* ArtisticIOR = "artistic_ior";
+		static constexpr const char* Blackbody = "blackbody";
+		static constexpr const char* Layer = "layer";
+		static constexpr const char* RoughnessAnisotropy = "roughness_anisotropy";
+		static constexpr const char* RoughnessDual = "roughness_dual";
+		//Surface Shaders
+		static constexpr const char* GltfPbr = "gltf_pbr";
+		static constexpr const char* DisneyBSDF2012 = "disney_brdf_2012";
+		static constexpr const char* DisneyBSDF2015 = "disney_bsdf_2015";
+		static constexpr const char* OpenPBRSurface = "open_pbr_surface";
+		static constexpr const char* StandardSurface = "standard_surface";
+		static constexpr const char* Surface = "surface";
+		static constexpr const char* UsdPreviewSurface = "UsdPreviewSurface";
+		// Shader
+		static constexpr const char* SurfaceUnlit = "surface_unlit";
+		// Convolution
+		static constexpr const char* Blur = "blur";
+		static constexpr const char* HeightToNormal = "heighttonormal";
+	}
+
+	namespace NodeDefinition
+	{
+		static constexpr const char* OpenPBRSurface = "ND_open_pbr_surface_surfaceshader";
+		static constexpr const char* StandardSurface = "ND_standard_surface_surfaceshader";
+		static constexpr const char* SurfaceUnlit = "ND_surface_unlit";
+		static constexpr const char* UsdPreviewSurface = "ND_UsdPreviewSurface_surfaceshader";
+		static constexpr const char* PointLight = "ND_point_light";
+		static constexpr const char* DirectionalLight = "ND_directional_light";
+		static constexpr const char* SpotLight = "ND_spot_light";
+		static constexpr const char* Surface = "ND_surface";
+	}
+
+	namespace Library
+	{
+		static constexpr const char* Libraries = "libraries";
+	}
+
+	namespace Type
+	{
+		static constexpr const char* Boolean = "boolean";
+		static constexpr const char* Integer = "integer";
+		static constexpr const char* Float = "float";
+		static constexpr const char* Color3 = "color3";
+		static constexpr const char* Color4 = "color4";
+		static constexpr const char* Vector2 = "vector2";
+		static constexpr const char* Vector3 = "vector3";
+		static constexpr const char* Vector4 = "vector4";
+		static constexpr const char* Matrix33 = "matrix33";
+		static constexpr const char* Matrix44 = "matrix44";
+		static constexpr const char* String = "string";
+		static constexpr const char* Filename = "filename";
+		static constexpr const char* GeomName = "geomname";
+		static constexpr const char* SurfaceShader = "surfaceshader";
+		static constexpr const char* DisplacementShader = "displacementshader";
+		static constexpr const char* VolumeShader = "volumeshader";
+		static constexpr const char* LightShader = "lightshader";
+		static constexpr const char* Material = "material";
+		static constexpr const char* None = "none";
+		static constexpr const char* IntegerArray = "integerarray";
+		static constexpr const char* FloatArray = "floatarray";
+		static constexpr const char* Color3Array = "color3array";
+		static constexpr const char* Color4Array = "color4array";
+		static constexpr const char* Vector2Array = "vector2array";
+		static constexpr const char* Vector3Array = "vector3array";
+		static constexpr const char* Vector4Array = "vector4array";
+		static constexpr const char* StringArray = "stringarray";
+		static constexpr const char* GeomNameArray = "geomnamearray";
+	}
+
+	namespace NodeGroup
+	{
+		namespace Texture2D
+		{
+			namespace Inputs
+			{
+				static constexpr const char* File = "file";
+				static constexpr const char* Default = "default";
+				static constexpr const char* TexCoord = "texcoord";
+				static constexpr const char* FilterType = "filtertype";
+				static constexpr const char* FrameRange = "framerange";
+				static constexpr const char* FrameOffset = "frameoffset";
+				static constexpr const char* FrameEndAction = "frameendaction";
+			}
+		}
+
+		namespace Math
+		{
+			template<typename T>
+			static const T NeutralZero = T{ 0 };
+
+			template<typename T>
+			static const T NeutralOne = T{ 1 };
+		}
+
+		static constexpr const char* PBR = "pbr";
+	}
+
+	namespace Image
+	{
+		namespace Inputs
+		{
+			using namespace NodeGroup::Texture2D::Inputs;
+
+			static constexpr const char* Layer = "layer";
+			static constexpr const char* UAddressMode = "uaddressmode";
+			static constexpr const char* VAddressMode = "vaddressmode";
+		}
+	}
+
+	namespace TiledImage
+	{
+		namespace Inputs
+		{
+			using namespace NodeGroup::Texture2D::Inputs;
+
+			static constexpr const char* UVTiling = "uvtiling";
+			static constexpr const char* UVOffset = "uvoffset";
+			static constexpr const char* RealWorldImageSize = "realworldimagesize";
+			static constexpr const char* RealWorldTileSize = "realworldtilesize";
+		}
+	}
+
+MATERIALX_NAMESPACE_END
diff --git a/Engine/Plugins/Interchange/Runtime/Source/Import/InterchangeImport.Build.cs b/Engine/Plugins/Interchange/Runtime/Source/Import/InterchangeImport.Build.cs
--- a/Engine/Plugins/Interchange/Runtime/Source/Import/InterchangeImport.Build.cs
+++ b/Engine/Plugins/Interchange/Runtime/Source/Import/InterchangeImport.Build.cs
@@ -14,6 +14,7 @@ namespace UnrealBuildTool.Rules
 					"CoreUObject",
 					"Engine",
 					"InterchangeCore",
+					"InterchangeCommon",
 					"InterchangeDispatcher",
 					"InterchangeEngine",
 					"InterchangeFactoryNodes",
diff --git a/Engine/Plugins/Interchange/Runtime/Source/Import/Private/MaterialX/InterchangeMaterialXDefinitions.h b/Engine/Plugins/Interchange/Runtime/Source/Import/Private/MaterialX/InterchangeMaterialXDefinitions.h
deleted file mode 100644
diff --git a/Engine/Plugins/Interchange/Runtime/Source/Import/Private/MaterialX/InterchangeMaterialXTranslator.cpp b/Engine/Plugins/Interchange/Runtime/Source/Import/Private/MaterialX/InterchangeMaterialXTranslator.cpp
--- a/Engine/Plugins/Interchange/Runtime/Source/Import/Private/MaterialX/InterchangeMaterialXTranslator.cpp
+++ b/Engine/Plugins/Interchange/Runtime/Source/Import/Private/MaterialX/InterchangeMaterialXTranslator.cpp
@@ -8,7 +8,7 @@
 
 #if WITH_EDITOR
 #include "MaterialXFormat/Util.h"
-#include "InterchangeMaterialXDefinitions.h"
+#include "MaterialX/InterchangeMaterialXDefinitions.h"
 #include "MaterialX/MaterialXUtils/MaterialXBase.h"
 #endif
 
diff --git a/Engine/Plugins/Interchange/Runtime/Source/Import/Private/MaterialX/MaterialXUtils/MaterialXManager.cpp b/Engine/Plugins/Interchange/Runtime/Source/Import/Private/MaterialX/MaterialXUtils/MaterialXManager.cpp
--- a/Engine/Plugins/Interchange/Runtime/Source/Import/Private/MaterialX/MaterialXUtils/MaterialXManager.cpp
+++ b/Engine/Plugins/Interchange/Runtime/Source/Import/Private/MaterialX/MaterialXUtils/MaterialXManager.cpp
@@ -222,26 +222,26 @@ FMaterialXManager::FMaterialXManager()
 {
 		MatchingMaterialFunctions = {
 			// BSDF Nodes
-			{mx::Category::BurleyDiffuseBSDF ,		MX_MATERIALFUNCTION("MX_BurleyDiffuseBSDF")},
-			{mx::Category::ConductorBSDF,			MX_MATERIALFUNCTION("MX_ConductorBSDF")},
-			{mx::Category::DielectricBSDF,			MX_MATERIALFUNCTION("MX_DielectricBSDF")},
-			{mx::Category::GeneralizedSchlickBSDF,	MX_MATERIALFUNCTION("MX_GeneralizedSchlickBSDF")},
-			{mx::Category::OrenNayarDiffuseBSDF ,	MX_MATERIALFUNCTION("MX_OrenNayarBSDF")},
-			{mx::Category::SheenBSDF,				MX_MATERIALFUNCTION("MX_SheenBSDF")},
-			{mx::Category::SubsurfaceBSDF ,			MX_MATERIALFUNCTION("MX_SubsurfaceBSDF")},
-			{mx::Category::ThinFilmBSDF,			MX_MATERIALFUNCTION("MX_ThinFilmBSDF")},
-			{mx::Category::TranslucentBSDF,			MX_MATERIALFUNCTION("MX_TranslucentBSDF")},
+			{mx::Category::BurleyDiffuseBSDF,		FMaterialXMaterialFunction{TInPlaceType<EInterchangeMaterialXBSDF>{}, EInterchangeMaterialXBSDF::BurleyDiffuse}},
+			{mx::Category::ConductorBSDF,			FMaterialXMaterialFunction{TInPlaceType<EInterchangeMaterialXBSDF>{}, EInterchangeMaterialXBSDF::Conductor}},
+			{mx::Category::DielectricBSDF,			FMaterialXMaterialFunction{TInPlaceType<EInterchangeMaterialXBSDF>{}, EInterchangeMaterialXBSDF::Dielectric}},
+			{mx::Category::GeneralizedSchlickBSDF,	FMaterialXMaterialFunction{TInPlaceType<EInterchangeMaterialXBSDF>{}, EInterchangeMaterialXBSDF::GeneralizedSchlick}},
+			{mx::Category::OrenNayarDiffuseBSDF ,	FMaterialXMaterialFunction{TInPlaceType<EInterchangeMaterialXBSDF>{}, EInterchangeMaterialXBSDF::OrenNayarDiffuse}},
+			{mx::Category::SheenBSDF,				FMaterialXMaterialFunction{TInPlaceType<EInterchangeMaterialXBSDF>{}, EInterchangeMaterialXBSDF::Sheen}},
+			{mx::Category::SubsurfaceBSDF ,			FMaterialXMaterialFunction{TInPlaceType<EInterchangeMaterialXBSDF>{}, EInterchangeMaterialXBSDF::Subsurface}},
+			{mx::Category::ThinFilmBSDF,			FMaterialXMaterialFunction{TInPlaceType<EInterchangeMaterialXBSDF>{}, EInterchangeMaterialXBSDF::ThinFilm}},
+			{mx::Category::TranslucentBSDF,			FMaterialXMaterialFunction{TInPlaceType<EInterchangeMaterialXBSDF>{}, EInterchangeMaterialXBSDF::Translucent}},
 			// EDF Nodes
-			{mx::Category::ConicalEDF,				MX_MATERIALFUNCTION("MX_ConicalEDF")},
-			{mx::Category::MeasuredEDF,				MX_MATERIALFUNCTION("MX_MeasuredEDF")},
-			{mx::Category::UniformEDF,				MX_MATERIALFUNCTION("MX_UniformEDF")},
+			{mx::Category::ConicalEDF,				FMaterialXMaterialFunction{TInPlaceType<EInterchangeMaterialXEDF>{}, EInterchangeMaterialXEDF::Conical}},
+			{mx::Category::MeasuredEDF,				FMaterialXMaterialFunction{TInPlaceType<EInterchangeMaterialXEDF>{}, EInterchangeMaterialXEDF::Measured}},
+			{mx::Category::UniformEDF,				FMaterialXMaterialFunction{TInPlaceType<EInterchangeMaterialXEDF>{}, EInterchangeMaterialXEDF::Uniform}},
 			// VDF Nodes
-			{mx::Category::AbsorptionVDF,			MX_MATERIALFUNCTION("MX_AbsorptionVDF")},
-			{mx::Category::AnisotropicVDF,			MX_MATERIALFUNCTION("MX_AnisotropicVDF")},
+			{mx::Category::AbsorptionVDF,			FMaterialXMaterialFunction{TInPlaceType<EInterchangeMaterialXVDF>{}, EInterchangeMaterialXVDF::Absorption}},
+			{mx::Category::AnisotropicVDF,			FMaterialXMaterialFunction{TInPlaceType<EInterchangeMaterialXVDF>{}, EInterchangeMaterialXVDF::Anisotropic}},
 			// Utility nodes
-			{mx::Category::ArtisticIOR,				TEXT("/Interchange/Functions/MX_Artistic_IOR.MX_Artistic_IOR")},
-			{mx::Category::RoughnessAnisotropy,		TEXT("/Interchange/Functions/MX_Roughness_Anisotropy.MX_Roughness_Anisotropy")},
-			{mx::Category::RoughnessDual,			TEXT("/Interchange/Functions/MX_Roughness_Dual.MX_Roughness_Dual")},
+			{mx::Category::ArtisticIOR,				FMaterialXMaterialFunction{TInPlaceType<FString>{}, TEXT("/Interchange/Functions/MX_Artistic_IOR.MX_Artistic_IOR")}},
+			{mx::Category::RoughnessAnisotropy,		FMaterialXMaterialFunction{TInPlaceType<FString>{}, TEXT("/Interchange/Functions/MX_Roughness_Anisotropy.MX_Roughness_Anisotropy")}},
+			{mx::Category::RoughnessDual,			FMaterialXMaterialFunction{TInPlaceType<FString>{}, TEXT("/Interchange/Functions/MX_Roughness_Dual.MX_Roughness_Dual")}}
 		};
 
 		if(bIsSubstrateEnabled)
@@ -293,9 +293,35 @@ const FString* FMaterialXManager::FindMatchingMaterialExpression(const FString&
 	return MatchingMaterialExpressions.Find({ CategoryKey, NodeGroup });
 }
 
-const FString* FMaterialXManager::FindMatchingMaterialFunction(const FString& CategoryKey) const
+bool FMaterialXManager::FindMatchingMaterialFunction(const FString& CategoryKey, const FString*& MaterialFunctionPath, uint8& EnumType, uint8& EnumValue) const
 {
-	return MatchingMaterialFunctions.Find(CategoryKey);
+	if(const FMaterialXMaterialFunction* MaterialFunction = MatchingMaterialFunctions.Find(CategoryKey))
+	{
+		MaterialFunctionPath = MaterialFunction->TryGet<FString>();
+		if(const EInterchangeMaterialXShaders* Shader = MaterialFunction->TryGet<EInterchangeMaterialXShaders>())
+		{
+			EnumType = UE::Interchange::MaterialX::IndexSurfaceShaders;
+			EnumValue = static_cast<uint8>(*Shader);
+		}
+		else if(const EInterchangeMaterialXBSDF* Bsdf = MaterialFunction->TryGet<EInterchangeMaterialXBSDF>())
+		{
+			EnumType = UE::Interchange::MaterialX::IndexBSDF;
+			EnumValue = static_cast<uint8>(*Bsdf);
+		}
+		else if(const EInterchangeMaterialXEDF* Edf = MaterialFunction->TryGet<EInterchangeMaterialXEDF>())
+		{
+			EnumType = UE::Interchange::MaterialX::IndexEDF;
+			EnumValue = static_cast<uint8>(*Edf);
+		}
+		else if(const EInterchangeMaterialXVDF* Vdf = MaterialFunction->TryGet<EInterchangeMaterialXVDF>())
+		{
+			EnumType = UE::Interchange::MaterialX::IndexVDF;
+			EnumValue = static_cast<uint8>(*Vdf);
+		}
+		return true;
+	}
+
+	return false;
 }
 
 TSharedPtr<FMaterialXBase> FMaterialXManager::GetShaderTranslator(const FString& CategoryShader, UInterchangeBaseNodeContainer & NodeContainer)
diff --git a/Engine/Plugins/Interchange/Runtime/Source/Import/Private/MaterialX/MaterialXUtils/MaterialXOpenPBRSurfaceShader.cpp b/Engine/Plugins/Interchange/Runtime/Source/Import/Private/MaterialX/MaterialXUtils/MaterialXOpenPBRSurfaceShader.cpp
--- a/Engine/Plugins/Interchange/Runtime/Source/Import/Private/MaterialX/MaterialXUtils/MaterialXOpenPBRSurfaceShader.cpp
+++ b/Engine/Plugins/Interchange/Runtime/Source/Import/Private/MaterialX/MaterialXUtils/MaterialXOpenPBRSurfaceShader.cpp
@@ -26,10 +26,9 @@ void FMaterialXOpenPBRSurfaceShader::Translate(MaterialX::NodePtr OpenPBRSurface
 	using namespace mx::OpenPBRSurface;
 
 	this->SurfaceShaderNode = OpenPBRSurfaceNode;
-	UInterchangeFunctionCallShaderNode* OpenPBRSurfaceShaderNode = CreateFunctionCallShaderNode(SurfaceShaderNode->getName().c_str(), 
-																								bIsSubstrateEnabled ?
-																								TEXT("/Engine/Functions/Substrate/MF_Substrate_OpenPBR_Opaque.MF_Substrate_OpenPBR_Opaque") :
-																								TEXT("/Interchange/Functions/MX_OpenPBR_Opaque.MX_OpenPBR_Opaque"));
+	UInterchangeFunctionCallShaderNode* OpenPBRSurfaceShaderNode = CreateFunctionCallShaderNode(SurfaceShaderNode->getName().c_str(), UE::Interchange::MaterialX::IndexSurfaceShaders, uint8(EInterchangeMaterialXShaders::OpenPBRSurface));
+
+	const bool bTangentSpace = true;
 
 	// Inputs
 	//Base
@@ -73,10 +72,10 @@ void FMaterialXOpenPBRSurfaceShader::Translate(MaterialX::NodePtr OpenPBRSurface
 	ConnectNodeOutputToInput(Input::TransmissionWeight, OpenPBRSurfaceShaderNode, OpenPBRSurface::Parameters::TransmissionWeight.ToString(), DefaultValue::TransmissionWeight);
 
 	//Geometry
-	ConnectNodeOutputToInput(Input::GeometryCoatNormal, OpenPBRSurfaceShaderNode, OpenPBRSurface::Parameters::GeometryCoatNormal.ToString(), DefaultValue::GeometryCoatNormal);
-	ConnectNodeOutputToInput(Input::GeometryNormal, OpenPBRSurfaceShaderNode, OpenPBRSurface::Parameters::GeometryNormal.ToString(), DefaultValue::GeometryNormal);
+	ConnectNodeOutputToInput(Input::GeometryCoatNormal, OpenPBRSurfaceShaderNode, OpenPBRSurface::Parameters::GeometryCoatNormal.ToString(), DefaultValue::GeometryCoatNormal, bTangentSpace);
+	ConnectNodeOutputToInput(Input::GeometryNormal, OpenPBRSurfaceShaderNode, OpenPBRSurface::Parameters::GeometryNormal.ToString(), DefaultValue::GeometryNormal, bTangentSpace);
 	ConnectNodeOutputToInput(Input::GeometryOpacity, OpenPBRSurfaceShaderNode, OpenPBRSurface::Parameters::GeometryOpacity.ToString(), DefaultValue::GeometryOpacity);
-	ConnectNodeOutputToInput(Input::GeometryTangent, OpenPBRSurfaceShaderNode, OpenPBRSurface::Parameters::GeometryTangent.ToString(), DefaultValue::GeometryTangent);
+	ConnectNodeOutputToInput(Input::GeometryTangent, OpenPBRSurfaceShaderNode, OpenPBRSurface::Parameters::GeometryTangent.ToString(), DefaultValue::GeometryTangent, bTangentSpace);
 	ConnectNodeOutputToInput(Input::GeometryThinWalled, OpenPBRSurfaceShaderNode, OpenPBRSurface::Parameters::GeometryThinWalled.ToString(), DefaultValue::GeometryThinWalled);
 	
 	//Fuzz
@@ -87,9 +86,8 @@ void FMaterialXOpenPBRSurfaceShader::Translate(MaterialX::NodePtr OpenPBRSurface
 	// We can't have Subsurface inputs if we have Transmission inputs
 	if(UInterchangeShaderPortsAPI::HasInput(OpenPBRSurfaceShaderNode, OpenPBRSurface::Parameters::TransmissionWeight))
 	{
-		OpenPBRSurfaceShaderNode->SetCustomMaterialFunction(bIsSubstrateEnabled ?
-															TEXT("/Engine/Functions/Substrate/MF_Substrate_OpenPBR_Translucent.MF_Substrate_OpenPBR_Translucent") :
-															TEXT("/Interchange/Functions/MX_OpenPBR_Translucent.MX_OpenPBR_Translucent"));
+		OpenPBRSurfaceShaderNode->AddInt32Attribute(UE::Interchange::MaterialX::Attributes::EnumType, UE::Interchange::MaterialX::IndexSurfaceShaders);
+		OpenPBRSurfaceShaderNode->AddInt32Attribute(UE::Interchange::MaterialX::Attributes::EnumValue, int32(EInterchangeMaterialXShaders::OpenPBRSurfaceTransmission));
 		ShaderGraphNode->SetCustomBlendMode(EBlendMode::BLEND_TranslucentColoredTransmittance);
 	}
 	else
diff --git a/Engine/Plugins/Interchange/Runtime/Source/Import/Private/MaterialX/MaterialXUtils/MaterialXStandardSurfaceShader.cpp b/Engine/Plugins/Interchange/Runtime/Source/Import/Private/MaterialX/MaterialXUtils/MaterialXStandardSurfaceShader.cpp
--- a/Engine/Plugins/Interchange/Runtime/Source/Import/Private/MaterialX/MaterialXUtils/MaterialXStandardSurfaceShader.cpp
+++ b/Engine/Plugins/Interchange/Runtime/Source/Import/Private/MaterialX/MaterialXUtils/MaterialXStandardSurfaceShader.cpp
@@ -26,24 +26,30 @@ void FMaterialXStandardSurfaceShader::Translate(mx::NodePtr StandardSurfaceNode)
 	
 	using namespace UE::Interchange::Materials;
 
+	UInterchangeFunctionCallShaderNode* StandardSurfaceShaderNode = CreateFunctionCallShaderNode(SurfaceShaderNode->getName().c_str(), UE::Interchange::MaterialX::IndexSurfaceShaders, uint8(EInterchangeMaterialXShaders::StandardSurface));
+
 	if(bIsSubstrateEnabled)
 	{
-		ConnectToSubstrateStandardSurface();
+		ConnectToSubstrateStandardSurface(StandardSurfaceShaderNode);
 	}
 	else
 	{
-		ConnectToStandardSurface();
+		ConnectToStandardSurface(StandardSurfaceShaderNode);
+	}
+
+	if(UInterchangeShaderPortsAPI::HasInput(StandardSurfaceShaderNode, StandardSurface::Parameters::Transmission))
+	{
+		StandardSurfaceShaderNode->AddInt32Attribute(UE::Interchange::MaterialX::Attributes::EnumType, UE::Interchange::MaterialX::IndexSurfaceShaders);
+		StandardSurfaceShaderNode->AddInt32Attribute(UE::Interchange::MaterialX::Attributes::EnumValue, int32(EInterchangeMaterialXShaders::StandardSurfaceTransmission));
 	}
 }
 
-void FMaterialXStandardSurfaceShader::ConnectToStandardSurface()
+void FMaterialXStandardSurfaceShader::ConnectToStandardSurface(UInterchangeFunctionCallShaderNode* StandardSurfaceShaderNode)
 {
 	using namespace UE::Interchange::Materials;
 	using namespace mx::StandardSurface;
 	constexpr bool bInputInTangentSpace = true;
 
-	UInterchangeFunctionCallShaderNode* StandardSurfaceShaderNode = CreateFunctionCallShaderNode(SurfaceShaderNode->getName().c_str(), TEXT("/Interchange/Functions/MX_StandardSurface.MX_StandardSurface"));
-
 	// Inputs
 	//Base
 	ConnectNodeOutputToInput(Input::Base, StandardSurfaceShaderNode, StandardSurface::Parameters::Base.ToString(), DefaultValue::Float::Base);
@@ -154,7 +160,6 @@ void FMaterialXStandardSurfaceShader::ConnectToStandardSurface()
 	// We can't have all shading models at once, so we have to make a choice here
 	if(UInterchangeShaderPortsAPI::HasInput(StandardSurfaceShaderNode, StandardSurface::Parameters::Transmission))
 	{
-		StandardSurfaceShaderNode->SetCustomMaterialFunction(TEXT("/Interchange/Functions/MX_TransmissionSurface.MX_TransmissionSurface"));
 		UInterchangeShaderPortsAPI::ConnectOuputToInputByName(ShaderGraphNode, PBRMR::Parameters::Opacity.ToString(), StandardSurfaceShaderNode->GetUniqueID(), PBRMR::Parameters::Opacity.ToString());
 		UInterchangeShaderPortsAPI::ConnectOuputToInputByName(ShaderGraphNode, ThinTranslucent::Parameters::TransmissionColor.ToString(), StandardSurfaceShaderNode->GetUniqueID(), ThinTranslucent::Parameters::TransmissionColor.ToString());
 		UInterchangeShaderPortsAPI::ConnectOuputToInputByName(ShaderGraphNode, Common::Parameters::Refraction.ToString(), StandardSurfaceShaderNode->GetUniqueID(), Common::Parameters::Refraction.ToString());
@@ -177,14 +182,12 @@ void FMaterialXStandardSurfaceShader::ConnectToStandardSurface()
 	}
 }
 
-void FMaterialXStandardSurfaceShader::ConnectToSubstrateStandardSurface()
+void FMaterialXStandardSurfaceShader::ConnectToSubstrateStandardSurface(UInterchangeFunctionCallShaderNode* StandardSurfaceShaderNode)
 {
 	using namespace UE::Interchange::Materials;
 	using namespace mx::StandardSurface;
 	constexpr bool bInputInTangentSpace = true;
 
-	UInterchangeFunctionCallShaderNode* StandardSurfaceShaderNode = CreateFunctionCallShaderNode(SurfaceShaderNode->getName().c_str(), TEXT("/Engine/Functions/Substrate/Substrate-StandardSurface-Opaque.Substrate-StandardSurface-Opaque"));
-
 	// Inputs
 	//Base
 	ConnectNodeOutputToInput(Input::Base, StandardSurfaceShaderNode, StandardSurface::Parameters::Base.ToString(), DefaultValue::Float::Base);
@@ -285,7 +288,6 @@ void FMaterialXStandardSurfaceShader::ConnectToSubstrateStandardSurface()
 	// Outputs
 	if(UInterchangeShaderPortsAPI::HasInput(StandardSurfaceShaderNode, StandardSurface::Parameters::Transmission))
 	{
-		StandardSurfaceShaderNode->SetCustomMaterialFunction(TEXT("/Engine/Functions/Substrate/Substrate-StandardSurface-Translucent.Substrate-StandardSurface-Translucent"));
 		ShaderGraphNode->SetCustomBlendMode(EBlendMode::BLEND_TranslucentColoredTransmittance);
 		UInterchangeShaderPortsAPI::ConnectOuputToInputByName(ShaderGraphNode, PBRMR::Parameters::Opacity.ToString(), StandardSurfaceShaderNode->GetUniqueID(), PBRMR::Parameters::Opacity.ToString());
 		UInterchangeShaderPortsAPI::ConnectOuputToInputByName(ShaderGraphNode, SubstrateMaterial::Parameters::FrontMaterial.ToString(), StandardSurfaceShaderNode->GetUniqueID(), StandardSurface::SubstrateMaterial::Outputs::Translucent.ToString());
diff --git a/Engine/Plugins/Interchange/Runtime/Source/Import/Private/MaterialX/MaterialXUtils/MaterialXStandardSurfaceShader.h b/Engine/Plugins/Interchange/Runtime/Source/Import/Private/MaterialX/MaterialXUtils/MaterialXStandardSurfaceShader.h
--- a/Engine/Plugins/Interchange/Runtime/Source/Import/Private/MaterialX/MaterialXUtils/MaterialXStandardSurfaceShader.h
+++ b/Engine/Plugins/Interchange/Runtime/Source/Import/Private/MaterialX/MaterialXUtils/MaterialXStandardSurfaceShader.h
@@ -22,8 +22,8 @@ public:
 
 private:
 
-	void ConnectToStandardSurface();
+	void ConnectToStandardSurface(UInterchangeFunctionCallShaderNode* StandardSurfaceShaderNode);
 
-	void ConnectToSubstrateStandardSurface();
+	void ConnectToSubstrateStandardSurface(UInterchangeFunctionCallShaderNode* StandardSurfaceShaderNode);
 };
 #endif
\ No newline at end of file
diff --git a/Engine/Plugins/Interchange/Runtime/Source/Import/Private/MaterialX/MaterialXUtils/MaterialXSurfaceShader.cpp b/Engine/Plugins/Interchange/Runtime/Source/Import/Private/MaterialX/MaterialXUtils/MaterialXSurfaceShader.cpp
--- a/Engine/Plugins/Interchange/Runtime/Source/Import/Private/MaterialX/MaterialXUtils/MaterialXSurfaceShader.cpp
+++ b/Engine/Plugins/Interchange/Runtime/Source/Import/Private/MaterialX/MaterialXUtils/MaterialXSurfaceShader.cpp
@@ -24,7 +24,7 @@ void FMaterialXSurfaceShader::Translate(MaterialX::NodePtr SurfaceNode)
 {
 	this->SurfaceShaderNode = SurfaceNode;
 
-	UInterchangeFunctionCallShaderNode* FunctionSurfaceShaderNode = CreateFunctionCallShaderNode(SurfaceNode->getName().c_str(), bIsSubstrateEnabled ? TEXT("/Interchange/Functions/MX_Surface.MX_Surface") : TEXT("/Interchange/Substrate/MX_Surface.MX_Surface"));
+	UInterchangeFunctionCallShaderNode* FunctionSurfaceShaderNode = CreateFunctionCallShaderNode(SurfaceNode->getName().c_str(), UE::Interchange::MaterialX::IndexSurfaceShaders, uint8(EInterchangeMaterialXShaders::Surface));
 
 	using namespace UE::Interchange::Materials;
 
@@ -44,10 +44,10 @@ void FMaterialXSurfaceShader::Translate(MaterialX::NodePtr SurfaceNode)
 	}
 	else
 	{
-		UInterchangeShaderPortsAPI::ConnectOuputToInputByName(ShaderGraphNode, SubstrateMaterial::Parameters::OpacityMask.ToString(), FunctionSurfaceShaderNode->GetUniqueID(), Surface::Substrate::Outputs::Surface.ToString());
+		UInterchangeShaderPortsAPI::ConnectOuputToInputByName(ShaderGraphNode, SubstrateMaterial::Parameters::FrontMaterial.ToString(), FunctionSurfaceShaderNode->GetUniqueID(), Surface::Substrate::Outputs::Surface.ToString());
 		if(UInterchangeShaderPortsAPI::HasInput(FunctionSurfaceShaderNode, Surface::Parameters::Opacity))
 		{
-			UInterchangeShaderPortsAPI::ConnectOuputToInputByName(ShaderGraphNode, SubstrateMaterial::Parameters::FrontMaterial.ToString(), FunctionSurfaceShaderNode->GetUniqueID(), Surface::Substrate::Outputs::Opacity.ToString());
+			UInterchangeShaderPortsAPI::ConnectOuputToInputByName(ShaderGraphNode, SubstrateMaterial::Parameters::OpacityMask.ToString(), FunctionSurfaceShaderNode->GetUniqueID(), Surface::Substrate::Outputs::Opacity.ToString());
 		}
 	}
 }
diff --git a/Engine/Plugins/Interchange/Runtime/Source/Import/Private/MaterialX/MaterialXUtils/MaterialXSurfaceShaderAbstract.cpp b/Engine/Plugins/Interchange/Runtime/Source/Import/Private/MaterialX/MaterialXUtils/MaterialXSurfaceShaderAbstract.cpp
--- a/Engine/Plugins/Interchange/Runtime/Source/Import/Private/MaterialX/MaterialXUtils/MaterialXSurfaceShaderAbstract.cpp
+++ b/Engine/Plugins/Interchange/Runtime/Source/Import/Private/MaterialX/MaterialXUtils/MaterialXSurfaceShaderAbstract.cpp
@@ -232,13 +232,8 @@ bool FMaterialXSurfaceShaderAbstract::ConnectMatchingNodeOutputToInput(const FCo
 
 	bool bIsConnected = false;
 
-	auto ConnectOutputToInput = [&](const FString* ShaderType, auto* (FMaterialXSurfaceShaderAbstract::* CreateFunctionCallOrShaderNode)(const FString&, const FString&, const FString&), bool bFindMatchingInput = true)
+	auto ConnectOutputToInputInternal = [&](UInterchangeShaderNode* OperatorNode, bool bFindMatchingInput = true)
 	{
-		UInterchangeShaderNode* OperatorNode = nullptr;
-
-		//We don't take the node output here because it would cause the creation of a new node (output is meaningful with ComponentMaskNode/separate where we have to create a new expression
-		OperatorNode = (this->*CreateFunctionCallOrShaderNode)(Connect.UpstreamNode->getName().c_str(), *ShaderType, DefaultOutput);
-
 		for(mx::InputPtr Input : Connect.UpstreamNode->getInputs())
 		{
 			if(!bFindMatchingInput)
@@ -258,6 +253,28 @@ bool FMaterialXSurfaceShaderAbstract::ConnectMatchingNodeOutputToInput(const FCo
 			UInterchangeShaderPortsAPI::ConnectOuputToInputByIndex(Connect.ParentShaderNode, Connect.InputChannelName, OperatorNode->GetUniqueID(), IndexOutput);
 	};
 
+	auto ConnectOutputToInput = [&](const FString* ShaderType, auto* (FMaterialXSurfaceShaderAbstract::* CreateFunctionCallOrShaderNode)(const FString&, const FString&, const FString&), bool bFindMatchingInput = true)
+	{
+		UInterchangeShaderNode* OperatorNode = nullptr;
+
+		//We don't take the node output here because it would cause the creation of a new node (output is meaningful with ComponentMaskNode/separate where we have to create a new expression
+		OperatorNode = (this->*CreateFunctionCallOrShaderNode)(Connect.UpstreamNode->getName().c_str(), *ShaderType, DefaultOutput);
+
+		ConnectOutputToInputInternal(OperatorNode, bFindMatchingInput);
+	};
+
+	auto ConnectFunctionShaderNodeOutputToInput = [&](uint8 EnumType, uint8 EnumValue, auto* (FMaterialXSurfaceShaderAbstract::* CreateFunctionCallOrShaderNode)(const FString&, uint8, uint8, const FString&), bool bFindMatchingInput = true)
+	{
+		UInterchangeShaderNode* OperatorNode = nullptr;
+
+		//We don't take the node output here because it would cause the creation of a new node (output is meaningful with ComponentMaskNode/separate where we have to create a new expression
+		OperatorNode = (this->*CreateFunctionCallOrShaderNode)(Connect.UpstreamNode->getName().c_str(), EnumType, EnumValue, DefaultOutput);
+
+		ConnectOutputToInputInternal(OperatorNode, bFindMatchingInput);
+	};
+
+	const FString* MaterialFunctionPath = nullptr;
+
 	// First search a matching Material Expression
 	// search for a Material Expression based on the node group (essentially used for Substrate Mix)
 	if(const FString* ShaderType = Manager.FindMatchingMaterialExpression(Connect.UpstreamNode->getCategory().c_str(), Connect.UpstreamNode->getNodeDef(mx::EMPTY_STRING, true)->getNodeGroup().c_str()))
@@ -272,9 +289,11 @@ bool FMaterialXSurfaceShaderAbstract::ConnectMatchingNodeOutputToInput(const FCo
 	{
 		bIsConnected = Delegate->ExecuteIfBound(Connect);
 	}
-	else if(const FString* FunctionPath = Manager.FindMatchingMaterialFunction(Connect.UpstreamNode->getCategory().c_str()))
+	else if(uint8 EnumType, EnumValue; Manager.FindMatchingMaterialFunction(Connect.UpstreamNode->getCategory().c_str(), MaterialFunctionPath, EnumType, EnumValue))
 	{
-		ConnectOutputToInput(FunctionPath, &FMaterialXSurfaceShaderAbstract::CreateFunctionCallShaderNode, false);
+		MaterialFunctionPath ?
+			ConnectOutputToInput(MaterialFunctionPath, &FMaterialXSurfaceShaderAbstract::CreateFunctionCallShaderNode, false) :
+			ConnectFunctionShaderNodeOutputToInput(EnumType, EnumValue, &FMaterialXSurfaceShaderAbstract::CreateFunctionCallShaderNode, false);
 	}
 
 	return bIsConnected;
@@ -1206,6 +1225,28 @@ UInterchangeFunctionCallShaderNode* FMaterialXSurfaceShaderAbstract::CreateFunct
 	return Node;
 }
 
+UInterchangeFunctionCallShaderNode* FMaterialXSurfaceShaderAbstract::CreateFunctionCallShaderNode(const FString& NodeName, uint8 EnumType, uint8 EnumValue, const FString& OutputName)
+{
+	UInterchangeFunctionCallShaderNode* Node;
+
+	const FString NodeUID = UInterchangeShaderNode::MakeNodeUid(NodeName + TEXT('_') + OutputName, FStringView{});
+
+	if(Node = const_cast<UInterchangeFunctionCallShaderNode*>(Cast<UInterchangeFunctionCallShaderNode>(NodeContainer.GetNode(NodeUID))); !Node)
+	{
+		Node = NewObject<UInterchangeFunctionCallShaderNode>(&NodeContainer);
+		Node->InitializeNode(NodeUID, NodeName, EInterchangeNodeContainerType::TranslatedAsset);
+		//this is just a dummy path name so the Generic Material Pipeline consider it as a FunctionCallShader but where in fact the path is given by an enum
+		Node->SetCustomMaterialFunction(TEXT("/Game/Default.Default"));
+		Node->AddInt32Attribute(UE::Interchange::MaterialX::Attributes::EnumType, EnumType);
+		Node->AddInt32Attribute(UE::Interchange::MaterialX::Attributes::EnumValue, EnumValue);
+		NodeContainer.AddNode(Node);
+
+		ShaderNodes.Add({ NodeName, OutputName }, Node);
+	}
+
+	return Node;
+}
+
 const FString& FMaterialXSurfaceShaderAbstract::GetMatchedInputName(MaterialX::NodePtr Node, MaterialX::InputPtr Input) const
 {
 	FMaterialXManager& Manager = FMaterialXManager::GetInstance();
diff --git a/Engine/Plugins/Interchange/Runtime/Source/Import/Private/MaterialX/MaterialXUtils/MaterialXSurfaceUnlitShader.cpp b/Engine/Plugins/Interchange/Runtime/Source/Import/Private/MaterialX/MaterialXUtils/MaterialXSurfaceUnlitShader.cpp
--- a/Engine/Plugins/Interchange/Runtime/Source/Import/Private/MaterialX/MaterialXUtils/MaterialXSurfaceUnlitShader.cpp
+++ b/Engine/Plugins/Interchange/Runtime/Source/Import/Private/MaterialX/MaterialXUtils/MaterialXSurfaceUnlitShader.cpp
@@ -24,7 +24,7 @@ void FMaterialXSurfaceUnlitShader::Translate(mx::NodePtr SurfaceUnlitNode)
 {
 	this->SurfaceShaderNode = SurfaceUnlitNode;
 
-	UInterchangeFunctionCallShaderNode* SurfaceUnlitShaderNode = CreateFunctionCallShaderNode(SurfaceUnlitNode->getName().c_str(), TEXT("/Interchange/Functions/MX_SurfaceUnlit.MX_SurfaceUnlit"));
+	UInterchangeFunctionCallShaderNode* SurfaceUnlitShaderNode = CreateFunctionCallShaderNode(SurfaceUnlitNode->getName().c_str(), UE::Interchange::MaterialX::IndexSurfaceShaders, uint8(EInterchangeMaterialXShaders::SurfaceUnlit));
 
 	using namespace UE::Interchange::Materials;
 
diff --git a/Engine/Plugins/Interchange/Runtime/Source/Import/Private/MaterialX/MaterialXUtils/MaterialXUsdPreviewSurfaceShader.cpp b/Engine/Plugins/Interchange/Runtime/Source/Import/Private/MaterialX/MaterialXUtils/MaterialXUsdPreviewSurfaceShader.cpp
--- a/Engine/Plugins/Interchange/Runtime/Source/Import/Private/MaterialX/MaterialXUtils/MaterialXUsdPreviewSurfaceShader.cpp
+++ b/Engine/Plugins/Interchange/Runtime/Source/Import/Private/MaterialX/MaterialXUtils/MaterialXUsdPreviewSurfaceShader.cpp
@@ -25,8 +25,7 @@ void FMaterialXUsdPreviewSurfaceShader::Translate(MaterialX::NodePtr UsdPreviewS
 
 	using namespace UE::Interchange::Materials;
 
-	const FString FolderFunction = bIsSubstrateEnabled ? TEXT("Substrate/") : TEXT("Functions/");
-	UInterchangeFunctionCallShaderNode* UsdPreviewSurfaceShaderNode = CreateFunctionCallShaderNode(UsdPreviewSurfaceNode->getName().c_str(), TEXT("/Interchange/") + FolderFunction + TEXT("MX_UsdPreviewSurface.MX_UsdPreviewSurface"));
+	UInterchangeFunctionCallShaderNode* UsdPreviewSurfaceShaderNode = CreateFunctionCallShaderNode(UsdPreviewSurfaceNode->getName().c_str(), UE::Interchange::MaterialX::IndexSurfaceShaders, uint8(EInterchangeMaterialXShaders::UsdPreviewSurface));
 
 	// Inputs
 	//Diffuse Color
diff --git a/Engine/Plugins/Interchange/Runtime/Source/Import/Public/MaterialX/MaterialXUtils/MaterialXManager.h b/Engine/Plugins/Interchange/Runtime/Source/Import/Public/MaterialX/MaterialXUtils/MaterialXManager.h
--- a/Engine/Plugins/Interchange/Runtime/Source/Import/Public/MaterialX/MaterialXUtils/MaterialXManager.h
+++ b/Engine/Plugins/Interchange/Runtime/Source/Import/Public/MaterialX/MaterialXUtils/MaterialXManager.h
@@ -7,6 +7,8 @@
 #if WITH_EDITOR
 
 #include "MaterialXCore/Node.h"
+#include "Misc/TVariant.h"
+#include "MaterialX/InterchangeMaterialXDefinitions.h"
 
 class FMaterialXBase;
 class UInterchangeBaseNodeContainer;
@@ -31,7 +33,7 @@ public:
 	const FString* FindMatchingMaterialExpression(const FString& CategoryKey, const FString& NodeGroup = {}) const;
 
 	/** Find a matching Material Function given a MaterialX category*/
-	const FString* FindMatchingMaterialFunction(const FString& CategoryKey) const;
+	bool FindMatchingMaterialFunction(const FString& CategoryKey, const FString*& MaterialFunctionPath, uint8& EnumType, uint8& EnumValue) const;
 
 	TSharedPtr<FMaterialXBase> GetShaderTranslator(const FString& CategoryShader, UInterchangeBaseNodeContainer& NodeContainer);
 
@@ -47,6 +49,12 @@ public:
 
 private:
 
+	/**
+	* FString: Material Function path
+	* Enums: data-driven BSDF nodes
+	*/
+	using FMaterialXMaterialFunction = TVariant<FString, EInterchangeMaterialXShaders, EInterchangeMaterialXBSDF, EInterchangeMaterialXEDF, EInterchangeMaterialXVDF>;
+
 	struct FKeyCategoryNodegroup
 	{
 		template<typename CategoryString>
@@ -83,7 +91,7 @@ private:
 	TMap<FKeyCategoryNodegroup, FString> MatchingMaterialExpressions;
 
 	/** Given a MaterialX node category, return the UE material function, used for BSDF nodes*/
-	TMap<FString, FString> MatchingMaterialFunctions;
+	TMap<FString, FMaterialXMaterialFunction> MatchingMaterialFunctions;
 
 	/** The different inputs of material expression that we may encounter, the MaterialX Document is modified consequently regarding those*/
 	TSet<FString> MaterialExpressionInputs;
diff --git a/Engine/Plugins/Interchange/Runtime/Source/Import/Public/MaterialX/MaterialXUtils/MaterialXSurfaceShaderAbstract.h b/Engine/Plugins/Interchange/Runtime/Source/Import/Public/MaterialX/MaterialXUtils/MaterialXSurfaceShaderAbstract.h
--- a/Engine/Plugins/Interchange/Runtime/Source/Import/Public/MaterialX/MaterialXUtils/MaterialXSurfaceShaderAbstract.h
+++ b/Engine/Plugins/Interchange/Runtime/Source/Import/Public/MaterialX/MaterialXUtils/MaterialXSurfaceShaderAbstract.h
@@ -308,6 +308,7 @@ protected:
 	 * @return The shader node that was created.
 	 */
 	UInterchangeFunctionCallShaderNode* CreateFunctionCallShaderNode(const FString& NodeName, const FString& FunctionPath, const FString& OutputName = TEXT("out"));
+	UInterchangeFunctionCallShaderNode* CreateFunctionCallShaderNode(const FString& NodeName, uint8 EnumType, uint8 EnumValue, const FString& OutputName = TEXT("out"));
 
 	/**
 	 * Helper function to create an InterchangeTextureNode.
diff --git a/Engine/Plugins/Interchange/Runtime/Source/Pipelines/InterchangePipelines.Build.cs b/Engine/Plugins/Interchange/Runtime/Source/Pipelines/InterchangePipelines.Build.cs
--- a/Engine/Plugins/Interchange/Runtime/Source/Pipelines/InterchangePipelines.Build.cs
+++ b/Engine/Plugins/Interchange/Runtime/Source/Pipelines/InterchangePipelines.Build.cs
@@ -13,6 +13,7 @@ namespace UnrealBuildTool.Rules
 					"CoreUObject",
 					"Engine",
 					"InterchangeCore",
+					"InterchangeCommon",
 					"InterchangeEngine",
 					"InterchangeFactoryNodes",
 					"InterchangeNodes"
diff --git a/Engine/Plugins/Interchange/Runtime/Source/Pipelines/Private/InterchangeMaterialXPipeline.cpp b/Engine/Plugins/Interchange/Runtime/Source/Pipelines/Private/InterchangeMaterialXPipeline.cpp
--- a/Engine/Plugins/Interchange/Runtime/Source/Pipelines/Private/InterchangeMaterialXPipeline.cpp
+++ b/Engine/Plugins/Interchange/Runtime/Source/Pipelines/Private/InterchangeMaterialXPipeline.cpp
@@ -51,9 +51,9 @@ namespace
 	MATERIALX_FUNCTIONS_SUBSTRATE_PATH(AnisotropicVDF);
 }
 
-TMap<FString, EMaterialXSettings> UInterchangeMaterialXPipeline::PathToEnumMapping;
+TMap<FString, EInterchangeMaterialXSettings> UInterchangeMaterialXPipeline::PathToEnumMapping;
 #if WITH_EDITOR
-TMap<EMaterialXSettings, TPair<TSet<FName>, TSet<FName>>> UMaterialXPipelineSettings::SettingsInputsOutputs;
+TMap<EInterchangeMaterialXSettings, TPair<TSet<FName>, TSet<FName>>> UMaterialXPipelineSettings::SettingsInputsOutputs;
 #endif // WITH_EDITOR
 
 UMaterialXPipelineSettings::UMaterialXPipelineSettings()
@@ -662,22 +662,22 @@ UInterchangeMaterialXPipeline::UInterchangeMaterialXPipeline()
 
 	for (const TPair<EInterchangeMaterialXShaders, FSoftObjectPath>& Entry : MaterialXSettings->PredefinedSurfaceShaders)
 	{
-		PathToEnumMapping.FindOrAdd(Entry.Value.GetAssetPathString(), EMaterialXSettings{ TInPlaceType<EInterchangeMaterialXShaders>{}, Entry.Key });
+		PathToEnumMapping.FindOrAdd(Entry.Value.GetAssetPathString(), EInterchangeMaterialXSettings{ TInPlaceType<EInterchangeMaterialXShaders>{}, Entry.Key });
 	}
 
 	for(const TPair<EInterchangeMaterialXBSDF, FSoftObjectPath>& Entry : MaterialXSettings->PredefinedBSDF)
 	{
-		PathToEnumMapping.FindOrAdd(Entry.Value.GetAssetPathString(), EMaterialXSettings{ TInPlaceType<EInterchangeMaterialXBSDF>{}, Entry.Key });
+		PathToEnumMapping.FindOrAdd(Entry.Value.GetAssetPathString(), EInterchangeMaterialXSettings{ TInPlaceType<EInterchangeMaterialXBSDF>{}, Entry.Key });
 	}
 
 	for(const TPair<EInterchangeMaterialXEDF, FSoftObjectPath>& Entry : MaterialXSettings->PredefinedEDF)
 	{
-		PathToEnumMapping.FindOrAdd(Entry.Value.GetAssetPathString(), EMaterialXSettings{ TInPlaceType<EInterchangeMaterialXEDF>{}, Entry.Key });
+		PathToEnumMapping.FindOrAdd(Entry.Value.GetAssetPathString(), EInterchangeMaterialXSettings{ TInPlaceType<EInterchangeMaterialXEDF>{}, Entry.Key });
 	}
 
 	for(const TPair<EInterchangeMaterialXVDF, FSoftObjectPath>& Entry : MaterialXSettings->PredefinedVDF)
 	{
-		PathToEnumMapping.FindOrAdd(Entry.Value.GetAssetPathString(), EMaterialXSettings{ TInPlaceType<EInterchangeMaterialXVDF>{}, Entry.Key });
+		PathToEnumMapping.FindOrAdd(Entry.Value.GetAssetPathString(), EInterchangeMaterialXSettings{ TInPlaceType<EInterchangeMaterialXVDF>{}, Entry.Key });
 	}
 }
 
@@ -696,14 +696,24 @@ void UInterchangeMaterialXPipeline::ExecutePipeline(UInterchangeBaseNodeContaine
 	Super::ExecutePipeline(NodeContainer, InSourceDatas, ContentBasePath);
 
 #if WITH_EDITOR
-	auto UpdateMaterialXNodes = [this](const FString& NodeUid, UInterchangeMaterialFunctionCallExpressionFactoryNode* FactorNode)
+	auto UpdateMaterialXNodes = [this, NodeContainer](const FString& NodeUid, UInterchangeMaterialFunctionCallExpressionFactoryNode* FactorNode)
 	{
 		const FString MaterialFunctionMemberName = GET_MEMBER_NAME_CHECKED(UMaterialExpressionMaterialFunctionCall, MaterialFunction).ToString();
 
-		FString MaterialFunctionPath;
-		if (FactorNode->GetStringAttribute(MaterialFunctionMemberName, MaterialFunctionPath))
+		FString FunctionShaderNodeUID = FactorNode->GetUniqueID();
+		FunctionShaderNodeUID.RemoveFromStart(TEXT("Factory_"));
+
+		const UInterchangeFunctionCallShaderNode* FunctionCallShaderNode = Cast<UInterchangeFunctionCallShaderNode>(NodeContainer->GetNode(FunctionShaderNodeUID));
+		
+		if(int32 EnumType; FunctionCallShaderNode->GetInt32Attribute(UE::Interchange::MaterialX::Attributes::EnumType, EnumType))
+		{
+			int32 EnumValue;
+			FunctionCallShaderNode->GetInt32Attribute(UE::Interchange::MaterialX::Attributes::EnumValue, EnumValue);
+			FactorNode->AddStringAttribute(MaterialFunctionMemberName, MaterialXSettings->GetAssetPathString(MaterialXSettings->ToEnumKey(EnumType, EnumValue)));
+		}
+		if(FString MaterialFunctionPath; FactorNode->GetStringAttribute(MaterialFunctionMemberName, MaterialFunctionPath))
 		{
-			if (const EMaterialXSettings* EnumPtr = PathToEnumMapping.Find(MaterialFunctionPath))
+			if (const EInterchangeMaterialXSettings* EnumPtr = PathToEnumMapping.Find(MaterialFunctionPath))
 			{
 				FactorNode->AddStringAttribute(MaterialFunctionMemberName, MaterialXSettings->GetAssetPathString(*EnumPtr));
 			}
@@ -748,19 +758,19 @@ bool UMaterialXPipelineSettings::AreRequiredPackagesLoaded()
 									  std::is_same_v<EnumT, EInterchangeMaterialXVDF>,
 									  "Enum type not supported");
 
-						uint8 EnumType = IndexSurfaceShaders;
+						uint8 EnumType = UE::Interchange::MaterialX::IndexSurfaceShaders;
 
 						if constexpr(std::is_same_v<EnumT, EInterchangeMaterialXBSDF>)
 						{
-							EnumType = IndexBSDF;
+							EnumType = UE::Interchange::MaterialX::IndexBSDF;
 						}
 						else if constexpr(std::is_same_v<EnumT, EInterchangeMaterialXEDF>)
 						{
-							EnumType = IndexEDF;
+							EnumType = UE::Interchange::MaterialX::IndexEDF;
 						}
 						else if constexpr(std::is_same_v<EnumT, EInterchangeMaterialXVDF>)
 						{
-							EnumType = IndexVDF;
+							EnumType = UE::Interchange::MaterialX::IndexVDF;
 						}
 
 						if(FMaterialXSettings::ValueType* Settings = SettingsInputsOutputs.Find(ToEnumKey(EnumType, uint8(Pair.Key))))
@@ -789,7 +799,7 @@ void UMaterialXPipelineSettings::InitPredefinedAssets()
 {
 	if(bIsSubstrateEnabled)
 	{
-		TArray<TTuple<EMaterialXSettings, FString, FString>> MappingToSubstrate
+		TArray<TTuple<EInterchangeMaterialXSettings, FString, FString>> MappingToSubstrate
 		{
 			{UMaterialXPipelineSettings::ToEnumKey(EInterchangeMaterialXShaders::OpenPBRSurface), OpenPBRSurfaceFunctionsPath, OpenPBRSurfaceSubstratePath},
 			{UMaterialXPipelineSettings::ToEnumKey(EInterchangeMaterialXShaders::OpenPBRSurfaceTransmission), OpenPBRTransmissionSurfaceFunctionsPath, OpenPBRTransmissionSurfaceSubstratePath},
@@ -818,9 +828,9 @@ void UMaterialXPipelineSettings::InitPredefinedAssets()
 			{UMaterialXPipelineSettings::ToEnumKey(EInterchangeMaterialXVDF::Anisotropic), AnisotropicVDFFunctionsPath, AnisotropicVDFSubstratePath},
 		};
 
-		for(auto Mapping : MappingToSubstrate)
+		for(const TTuple<EInterchangeMaterialXSettings, FString, FString> & Mapping : MappingToSubstrate)
 		{
-			const EMaterialXSettings& ShadersSettings = Mapping.Get<0>();
+			const EInterchangeMaterialXSettings& ShadersSettings = Mapping.Get<0>();
 			if(FString Path = GetAssetPathString(ShadersSettings); Path == Mapping.Get<1>())
 			{
 				if(ShadersSettings.IsType<EInterchangeMaterialXShaders>())
@@ -845,7 +855,7 @@ void UMaterialXPipelineSettings::InitPredefinedAssets()
 }
 #endif // WITH_EDITOR
 
-FString UMaterialXPipelineSettings::GetAssetPathString(EMaterialXSettings EnumValue) const
+FString UMaterialXPipelineSettings::GetAssetPathString(EInterchangeMaterialXSettings EnumValue) const
 {
 	auto FindAssetPathString = [](const auto& PredefinedEnumPath, auto Enum) -> FString
 	{
@@ -910,47 +920,47 @@ bool UMaterialXPipelineSettings::ShouldFilterAssets(UMaterialFunction* Asset, co
 	return !(InputMatches > 0 && OutputMatches == Outputs.Num());
 }
 
-EMaterialXSettings UMaterialXPipelineSettings::ToEnumKey(uint8 EnumType, uint8 EnumValue)
+EInterchangeMaterialXSettings UMaterialXPipelineSettings::ToEnumKey(uint8 EnumType, uint8 EnumValue)
 {
 	switch(EnumType)
 	{
-	case IndexBSDF:
-		return EMaterialXSettings{ TInPlaceType<EInterchangeMaterialXBSDF>{}, EInterchangeMaterialXBSDF{EnumValue} };
+	case UE::Interchange::MaterialX::IndexBSDF:
+		return EInterchangeMaterialXSettings{ TInPlaceType<EInterchangeMaterialXBSDF>{}, EInterchangeMaterialXBSDF{EnumValue} };
 
-	case IndexEDF:
-		return EMaterialXSettings{ TInPlaceType<EInterchangeMaterialXEDF>{}, EInterchangeMaterialXEDF{EnumValue} };
+	case UE::Interchange::MaterialX::IndexEDF:
+		return EInterchangeMaterialXSettings{ TInPlaceType<EInterchangeMaterialXEDF>{}, EInterchangeMaterialXEDF{EnumValue} };
 
-	case IndexVDF:
-		return EMaterialXSettings{ TInPlaceType<EInterchangeMaterialXVDF>{}, EInterchangeMaterialXVDF{EnumValue} };
+	case UE::Interchange::MaterialX::IndexVDF:
+		return EInterchangeMaterialXSettings{ TInPlaceType<EInterchangeMaterialXVDF>{}, EInterchangeMaterialXVDF{EnumValue} };
 
 	default:
-		return EMaterialXSettings{ TInPlaceType<EInterchangeMaterialXShaders>{}, EInterchangeMaterialXShaders{EnumValue} };
+		return EInterchangeMaterialXSettings{ TInPlaceType<EInterchangeMaterialXShaders>{}, EInterchangeMaterialXShaders{EnumValue} };
 	}
 }
 #endif // WITH_EDITOR
 
 namespace
 {
-	static uint8 GetMaterialXSettingsIndexValue(const EMaterialXSettings Enum, SIZE_T& Index)
+	static uint8 GetMaterialXSettingsIndexValue(const EInterchangeMaterialXSettings Enum, SIZE_T& Index)
 	{
 		Index = Enum.GetIndex();
 		const uint8* RawValuePointer = 
-			Index == 1 ? reinterpret_cast<const uint8*>(Enum.TryGet<EInterchangeMaterialXBSDF>()) :
-			Index == 2 ? reinterpret_cast<const uint8*>(Enum.TryGet<EInterchangeMaterialXEDF>()) :
-			Index == 3 ? reinterpret_cast<const uint8*>(Enum.TryGet<EInterchangeMaterialXVDF>()) :
+			Index == UE::Interchange::MaterialX::IndexBSDF ? reinterpret_cast<const uint8*>(Enum.TryGet<EInterchangeMaterialXBSDF>()) :
+			Index == UE::Interchange::MaterialX::IndexEDF ? reinterpret_cast<const uint8*>(Enum.TryGet<EInterchangeMaterialXEDF>()) :
+			Index == UE::Interchange::MaterialX::IndexVDF ? reinterpret_cast<const uint8*>(Enum.TryGet<EInterchangeMaterialXVDF>()) :
 			reinterpret_cast<const uint8*>(Enum.TryGet<EInterchangeMaterialXShaders>());
 		return *RawValuePointer;
 	}
 }
 
-uint32 GetTypeHash(EMaterialXSettings Key)
+uint32 GetTypeHash(EInterchangeMaterialXSettings Key)
 {
 	SIZE_T Index;
 	const uint8 UnderlyingValue = GetMaterialXSettingsIndexValue(Key, Index);
 	return HashCombine(Index, UnderlyingValue);
 }
 
-bool operator==(EMaterialXSettings Lhs, EMaterialXSettings Rhs)
+bool operator==(EInterchangeMaterialXSettings Lhs, EInterchangeMaterialXSettings Rhs)
 {
 	SIZE_T LhsIndex;
 	const uint8 LhsUnderlyingValue = GetMaterialXSettingsIndexValue(Lhs, LhsIndex);
diff --git a/Engine/Plugins/Interchange/Runtime/Source/Pipelines/Public/InterchangeMaterialXPipeline.h b/Engine/Plugins/Interchange/Runtime/Source/Pipelines/Public/InterchangeMaterialXPipeline.h
--- a/Engine/Plugins/Interchange/Runtime/Source/Pipelines/Public/InterchangeMaterialXPipeline.h
+++ b/Engine/Plugins/Interchange/Runtime/Source/Pipelines/Public/InterchangeMaterialXPipeline.h
@@ -5,126 +5,23 @@
 #include "CoreMinimal.h"
 
 #include "InterchangePipelineBase.h"
-#include "InterchangeSourceData.h"
 #include "Nodes/InterchangeBaseNodeContainer.h"
 
 #include "Engine/DeveloperSettings.h"
-#include "Templates/Tuple.h"
-#include "UObject/Object.h"
-#include "UObject/ObjectMacros.h"
-#include "UObject/SoftObjectPath.h"
+#include "MaterialX/InterchangeMaterialXDefinitions.h"
 
 #include "InterchangeMaterialXPipeline.generated.h"
 
-class UInterchangeDatasmithPbrMaterialNode;
-class UInterchangeFactoryBaseNode;
-class UInterchangeMaterialFactoryNode;
-class UInterchangeMaterialInstanceFactoryNode;
-class UInterchangeShaderGraphNode;
-class UInterchangeMaterialInstanceNode;
-
 class UMaterialFunction;
 class UMaterialInterface;
 
-UENUM(BlueprintType)
-enum class EInterchangeMaterialXShaders : uint8
-{
-	/** Default settings for Open PBR Surface shader	*/
-	OpenPBRSurface,
-
-	/** Open PBR Surface shader	used for translucency*/
-	OpenPBRSurfaceTransmission,
-
-	/** Default settings for Autodesk's Standard Surface shader.	*/
-	StandardSurface,
-
-	/** Standard Surface shader used for translucency. */
-	StandardSurfaceTransmission,
-
-	/** Shader used for unlit surfaces. */
-	SurfaceUnlit,
-
-	/** Default settings for USD's Surface shader. */
-	UsdPreviewSurface,
-
-	/** A surface shader constructed from scattering and emission distribution functions. */
-	Surface,
-
-	MaxShaderCount UMETA(hidden)
-};
-
-#define INTERCHANGE_MATERIALX_EXPERIEMENTAL
-
-UENUM(BlueprintType)
-/** Data type representing a Bidirectional Scattering Distribution Function. */
-enum class EInterchangeMaterialXBSDF : uint8
-{
-	/** A BSDF node for diffuse reflections. */
-	OrenNayarDiffuse,
-
-	/** A BSDF node for Burley diffuse reflections. */
-	BurleyDiffuse,
-
-	/** A BSDF node for pure diffuse transmission. */
-	Translucent,
-
-	/** A reflection/transmission BSDF node based on a microfacet model and a Fresnel curve for dielectrics. */
-	Dielectric,
-
-	/** A reflection BSDF node based on a microfacet model and a Fresnel curve for conductors/metals. */
-	Conductor,
-
-	/** A reflection/transmission BSDF node based on a microfacet model and a generalized Schlick Fresnel curve. */
-	GeneralizedSchlick,
+using EInterchangeMaterialXSettings = TVariant<EInterchangeMaterialXShaders, EInterchangeMaterialXBSDF, EInterchangeMaterialXEDF, EInterchangeMaterialXVDF>;
 
-	/** A subsurface scattering BSDF for true subsurface scattering. */
-	Subsurface,
+uint32 INTERCHANGEPIPELINES_API GetTypeHash(EInterchangeMaterialXSettings Key);
 
-	/** A microfacet BSDF for the back-scattering properties of cloth-like materials. */
-	Sheen,
+bool INTERCHANGEPIPELINES_API operator==(EInterchangeMaterialXSettings Lhs, EInterchangeMaterialXSettings Rhs);
 
-	/** An iridescent thin film layer over a microfacet base BSDF. */
-	ThinFilm,
-
-	MaxBSDFCount UMETA(hidden)
-};
-
-UENUM(BlueprintType)
-/** Data type representing an Emission Distribution Function. */
-enum class EInterchangeMaterialXEDF : uint8
-{
-	/** An EDF node for uniform emission. */
-	Uniform,
-
-	/** An EDF emitting light inside a cone around the normal direction. */
-	Conical,
-
-	/** An EDF emitting light according to a measured IES light profile. */
-	Measured,
-
-	MaxEDFCount UMETA(hidden)
-};
-
-UENUM(BlueprintType)
-/** Data type representing a Volume Distribution Function. */
-enum class EInterchangeMaterialXVDF : uint8
-{
-	/** A VDF for pure light absorption. */
-	Absorption,
-
-	/** A VDF scattering light for a participating medium, based on the Henyey-Greenstein phase function. */
-	Anisotropic,
-
-	MaxVDFCount UMETA(hidden)
-};
-
-using EMaterialXSettings = TVariant<EInterchangeMaterialXShaders, EInterchangeMaterialXBSDF, EInterchangeMaterialXEDF, EInterchangeMaterialXVDF>;
-
-uint32 INTERCHANGEPIPELINES_API GetTypeHash(EMaterialXSettings Key);
-
-bool INTERCHANGEPIPELINES_API operator==(EMaterialXSettings Lhs, EMaterialXSettings Rhs);
-
-UCLASS(config = Interchange, meta = (DisplayName = "Interchange MaterialX", ToolTip = "Interchange settings for MaterialX conversions."))
+UCLASS(config = Interchange, meta = (DisplayName = "Interchange MaterialX"))
 class INTERCHANGEPIPELINES_API UMaterialXPipelineSettings : public UDeveloperSettings
 {
 	GENERATED_BODY()
@@ -134,7 +31,7 @@ public:
 
 	bool AreRequiredPackagesLoaded();
 
-	FString GetAssetPathString(EMaterialXSettings EnumType) const;
+	FString GetAssetPathString(EInterchangeMaterialXSettings EnumType) const;
 
 	template<typename EnumT>
 	FString GetAssetPathString(EnumT EnumValue) const
@@ -145,7 +42,7 @@ public:
 					  std::is_same_v<EnumT, EInterchangeMaterialXVDF>,
 					  "Enum type not supported");
 
-		return GetAssetPathString(EMaterialXSettings{ TInPlaceType<EnumT>{}, EnumValue });
+		return GetAssetPathString(EInterchangeMaterialXSettings{ TInPlaceType<EnumT>{}, EnumValue });
 	}
 
 	UPROPERTY(EditAnywhere, config, Category = "MaterialXPredefined | Surface Shaders", meta = (DisplayName = "MaterialX Predefined Surface Shaders"))
@@ -168,14 +65,14 @@ private:
 	friend class FInterchangeMaterialXPipelineSettingsCustomization;
 	friend class UInterchangeMaterialXPipeline;
 
-	using FMaterialXSettings = TMap<EMaterialXSettings, TPair<TSet<FName>, TSet<FName>>>;
+	using FMaterialXSettings = TMap<EInterchangeMaterialXSettings, TPair<TSet<FName>, TSet<FName>>>;
 
 	static bool ShouldFilterAssets(UMaterialFunction* Asset, const TSet<FName>& Inputs, const TSet<FName>& Outputs);
 
-	static EMaterialXSettings ToEnumKey(uint8 EnumType, uint8 EnumValue);
+	static EInterchangeMaterialXSettings ToEnumKey(uint8 EnumType, uint8 EnumValue);
 
 	template<typename EnumT>
-	static EMaterialXSettings ToEnumKey(EnumT EnumValue)
+	static EInterchangeMaterialXSettings ToEnumKey(EnumT EnumValue)
 	{
 		static_assert(std::is_same_v<EnumT, EInterchangeMaterialXShaders> ||
 					  std::is_same_v<EnumT, EInterchangeMaterialXBSDF> ||
@@ -183,17 +80,12 @@ private:
 					  std::is_same_v<EnumT, EInterchangeMaterialXVDF>,
 					  "Enum type not supported");
 
-		return EMaterialXSettings{ TInPlaceType<EnumT>{}, EnumValue };
+		return EInterchangeMaterialXSettings{ TInPlaceType<EnumT>{}, EnumValue };
 	}
 
 	/** The key is a combination of the index in the variant + the corresponding enum */
 	static FMaterialXSettings SettingsInputsOutputs;
 
-	static constexpr uint8 IndexSurfaceShaders = 0;
-	static constexpr uint8 IndexBSDF = 1;
-	static constexpr uint8 IndexEDF = 2;
-	static constexpr uint8 IndexVDF = 3;
-
 	bool bIsSubstrateEnabled{ false };
 #endif // WITH_EDITOR
 };
@@ -220,5 +112,5 @@ protected:
 
 private:
 
-	static TMap<FString, EMaterialXSettings> PathToEnumMapping;
+	static TMap<FString, EInterchangeMaterialXSettings> PathToEnumMapping;
 };
