[Interchange] Import FBX asset - Crashes when re-importing original assets

#rnx

diff --git a/Engine/Plugins/Interchange/Runtime/Source/Import/Private/Mesh/InterchangeSkeletalMeshFactory.cpp b/Engine/Plugins/Interchange/Runtime/Source/Import/Private/Mesh/InterchangeSkeletalMeshFactory.cpp
--- a/Engine/Plugins/Interchange/Runtime/Source/Import/Private/Mesh/InterchangeSkeletalMeshFactory.cpp
+++ b/Engine/Plugins/Interchange/Runtime/Source/Import/Private/Mesh/InterchangeSkeletalMeshFactory.cpp
@@ -1178,12 +1178,17 @@ UInterchangeFactoryBase::FImportAssetResult UInterchangeSkeletalMeshFactory::Beg
 
 			}
 
-			if (!ensure(SkeletonReference))
+			if (!Arguments.ReimportObject)
 			{
-				UInterchangeResultWarning_Generic* Message = AddMessage<UInterchangeResultWarning_Generic>();
-				Message->Text = WarningMessage_InvalidSkeleton;
-				break;
+				//In case its a SkeletalMesh AssetReimport without the Skeleton, then we won't have SkeletonReference which is expected
+				if (!ensure(SkeletonReference))
+				{
+					UInterchangeResultWarning_Generic* Message = AddMessage<UInterchangeResultWarning_Generic>();
+					Message->Text = WarningMessage_InvalidSkeleton;
+					break;
+				}
 			}
+			
 			ImportAssetObjectData.SkeletonReference = SkeletonReference;
 		}
 
@@ -2243,7 +2248,7 @@ bool UInterchangeSkeletalMeshFactory::FImportAssetObjectData::IsValid() const
 			return false;
 		}
 #if WITH_EDITOR
-		if (ImportAssetObjectLODData.ImportedMaterials.Num() < 1)
+		if (!bIsReImport && ImportAssetObjectLODData.ImportedMaterials.Num() < 1)
 		{
 			return false;
 		}
diff --git a/Engine/Plugins/Interchange/Runtime/Source/Import/Private/Mesh/InterchangeSkeletonHelper.cpp b/Engine/Plugins/Interchange/Runtime/Source/Import/Private/Mesh/InterchangeSkeletonHelper.cpp
--- a/Engine/Plugins/Interchange/Runtime/Source/Import/Private/Mesh/InterchangeSkeletonHelper.cpp
+++ b/Engine/Plugins/Interchange/Runtime/Source/Import/Private/Mesh/InterchangeSkeletonHelper.cpp
@@ -75,6 +75,11 @@ namespace UE::Interchange::Private
 
 	bool FSkeletonHelper::IsCompatibleSkeleton(const USkeleton* Skeleton, const FString RootJoinUid, const UInterchangeBaseNodeContainer* BaseNodeContainer, bool bConvertStaticToSkeletalActive)
 	{
+		if (!Skeleton)
+		{
+			return false;
+		}
+
 		// at least % of bone should match 
 		int32 NumOfBoneMatches = 0;
 		//Make sure the specified Skeleton fit this skeletal mesh
diff --git a/Engine/Plugins/Interchange/Runtime/Source/Pipelines/Private/InterchangeGenericAssetsPipeline.cpp b/Engine/Plugins/Interchange/Runtime/Source/Pipelines/Private/InterchangeGenericAssetsPipeline.cpp
--- a/Engine/Plugins/Interchange/Runtime/Source/Pipelines/Private/InterchangeGenericAssetsPipeline.cpp
+++ b/Engine/Plugins/Interchange/Runtime/Source/Pipelines/Private/InterchangeGenericAssetsPipeline.cpp
@@ -819,6 +819,12 @@ void UInterchangeGenericAssetsPipeline::CreateSkeletonConflict(USkeleton* Specif
 		{
 			return;
 		}
+
+		if (!SpecifiedSkeleton && !SkeletalMesh->GetSkeleton())
+		{
+			return;
+		}
+
 		const bool bConvertStaticToSkeletalActive = CommonSkeletalMeshesAndAnimationsProperties->bConvertStaticsWithMorphTargetsToSkeletals || CommonMeshesProperties->ForceAllMeshAsType == EInterchangeForceMeshType::IFMT_SkeletalMesh;
 		//If we have a compatible skeleton we do not need to create a conflict
 		if (UE::Interchange::Private::FSkeletonHelper::IsCompatibleSkeleton(SpecifiedSkeleton ? SpecifiedSkeleton : SkeletalMesh->GetSkeleton(), RootJointNodeId, TransientBaseNodeContainer, bConvertStaticToSkeletalActive))
@@ -1837,7 +1843,7 @@ TSharedPtr<SWidget> SInterchangeGenericAssetSkeletonConflictWidget::ConstructSke
 	FString SkeletonReferenceStatistic;
 	if (AssetReferencingSkeleton.Num() > 0)
 	{
-		SkeletonReferenceStatistic += TEXT("Skeleton is references by ") + FString::FromInt(AssetReferencingSkeleton.Num()) + TEXT(" assets.");
+		SkeletonReferenceStatistic += TEXT("Skeleton is referenced by ") + FString::FromInt(AssetReferencingSkeleton.Num()) + TEXT(" assets.");
 	}
 	
 	return SNew(SBox)
