Fixing imported normal map textures with World group instead of WorlNormalMap

diff --git a/Engine/Plugins/Interchange/Runtime/Source/Pipelines/Private/InterchangeGenericMaterialPipeline.cpp b/Engine/Plugins/Interchange/Runtime/Source/Pipelines/Private/InterchangeGenericMaterialPipeline.cpp
--- a/Engine/Plugins/Interchange/Runtime/Source/Pipelines/Private/InterchangeGenericMaterialPipeline.cpp
+++ b/Engine/Plugins/Interchange/Runtime/Source/Pipelines/Private/InterchangeGenericMaterialPipeline.cpp
@@ -1923,6 +1923,7 @@ void UInterchangeGenericMaterialPipeline::HandleTextureNode(
 			if(DesiredTextureUsage == EMaterialInputType::Vector)
 			{
 				TextureFactoryNode->SetCustomCompressionSettings(TextureCompressionSettings::TC_Normalmap);
+				TextureFactoryNode->SetCustomLODGroup(TextureGroup::TEXTUREGROUP_WorldNormalMap);
 			}
 			else if(DesiredTextureUsage == EMaterialInputType::Scalar)
 			{
