- Interchange add PipelineDisplayName variable in GenericLevelPipeline to fix the name that is replaced by an index after 2 reimport
- Set 'Default Scene Level Pipeline' name to DefaultSceneLevelPipeline.uasset

#virtualized

diff --git a/Engine/Plugins/Interchange/Runtime/Source/Pipelines/Public/InterchangeGenericScenesPipeline.h b/Engine/Plugins/Interchange/Runtime/Source/Pipelines/Public/InterchangeGenericScenesPipeline.h
--- a/Engine/Plugins/Interchange/Runtime/Source/Pipelines/Public/InterchangeGenericScenesPipeline.h
+++ b/Engine/Plugins/Interchange/Runtime/Source/Pipelines/Public/InterchangeGenericScenesPipeline.h
@@ -21,6 +21,10 @@ class INTERCHANGEPIPELINES_API UInterchangeGenericLevelPipeline : public UInterc
 {
 	GENERATED_BODY()
 public:
+	/** The name of the pipeline that will be display in the import dialog. */
+	UPROPERTY(EditAnywhere, BlueprintReadWrite, Category = "Level Reimport Policy", meta = (StandAlonePipelineProperty = "True", PipelineInternalEditionData = "True"))
+	FString PipelineDisplayName;
+
 	/* Set the reimport strategy when reimporting into the level. */
 	UPROPERTY(EditAnywhere, BlueprintReadWrite, Category = "Level Reimport Policy", meta = (SubCategory = "Actors properties", AdjustPipelineAndRefreshDetailOnChange = "True"))
 	EReimportStrategyFlags ReimportPropertyStrategy = EReimportStrategyFlags::ApplyNoProperties;
