MaterialX:
- Update open_pbr nodedef that missed 'transmission_dispersion_scale'  and 'transmission_dispersion_abbe_number' inputs
- Add/modify transmission_dispersion_scale and transmission_dispersion_abbe_number inputs in both C++ and assets.

#virtualized

diff --git a/Engine/Plugins/Interchange/Runtime/Source/Common/Public/MaterialX/InterchangeMaterialXDefinitions.h b/Engine/Plugins/Interchange/Runtime/Source/Common/Public/MaterialX/InterchangeMaterialXDefinitions.h
--- a/Engine/Plugins/Interchange/Runtime/Source/Common/Public/MaterialX/InterchangeMaterialXDefinitions.h
+++ b/Engine/Plugins/Interchange/Runtime/Source/Common/Public/MaterialX/InterchangeMaterialXDefinitions.h
@@ -155,7 +155,8 @@ MATERIALX_NAMESPACE_BEGIN
 			static constexpr const char* TransmissionDepth = "transmission_depth";
 			static constexpr const char* TransmissionScatter = "transmission_scatter";
 			static constexpr const char* TransmissionScatterAnisotropy= "transmission_scatter_anisotropy";
-			static constexpr const char* TransmissionDispersion = "transmission_dispersion";
+			static constexpr const char* TransmissionDispersionScale = "transmission_dispersion_scale";
+			static constexpr const char* TransmissionDispersionAbbeNumber= "transmission_dispersion_abbe_number";
 			static constexpr const char* SubsurfaceWeight = "subsurface_weight";
 			static constexpr const char* SubsurfaceColor = "subsurface_color";
 			static constexpr const char* SubsurfaceRadius = "subsurface_radius";
@@ -200,7 +201,8 @@ MATERIALX_NAMESPACE_BEGIN
 			static constexpr float TransmissionDepth = 0.f;
 			static constexpr FLinearColor TransmissionScatter{ 0, 0, 0 };
 			static constexpr float TransmissionScatterAnisotropy = 0.f;
-			static constexpr float TransmissionDispersion = 0.f;
+			static constexpr float TransmissionDispersionScale = 0.f;
+			static constexpr float TransmissionDispersionAbbeNumber = 0.f;
 			static constexpr float SubsurfaceWeight = 0.f;
 			static constexpr FLinearColor SubsurfaceColor{ 0.8, 0.8, 0.8 };
 			static constexpr float SubsurfaceRadius = 1.f;
diff --git a/Engine/Plugins/Interchange/Runtime/Source/Import/Private/MaterialX/MaterialXUtils/MaterialXOpenPBRSurfaceShader.cpp b/Engine/Plugins/Interchange/Runtime/Source/Import/Private/MaterialX/MaterialXUtils/MaterialXOpenPBRSurfaceShader.cpp
--- a/Engine/Plugins/Interchange/Runtime/Source/Import/Private/MaterialX/MaterialXUtils/MaterialXOpenPBRSurfaceShader.cpp
+++ b/Engine/Plugins/Interchange/Runtime/Source/Import/Private/MaterialX/MaterialXUtils/MaterialXOpenPBRSurfaceShader.cpp
@@ -66,7 +66,8 @@ void FMaterialXOpenPBRSurfaceShader::Translate(MaterialX::NodePtr OpenPBRSurface
 	//Transmission
 	ConnectNodeOutputToInput(Input::TransmissionColor, OpenPBRSurfaceShaderNode, OpenPBRSurface::Parameters::TransmissionColor.ToString(), DefaultValue::TransmissionColor);
 	ConnectNodeOutputToInput(Input::TransmissionDepth, OpenPBRSurfaceShaderNode, OpenPBRSurface::Parameters::TransmissionDepth.ToString(), DefaultValue::TransmissionDepth);
-	ConnectNodeOutputToInput(Input::TransmissionDispersion, OpenPBRSurfaceShaderNode, OpenPBRSurface::Parameters::TransmissionDispersion.ToString(), DefaultValue::TransmissionDispersion);
+	ConnectNodeOutputToInput(Input::TransmissionDispersionScale, OpenPBRSurfaceShaderNode, OpenPBRSurface::Parameters::TransmissionDispersionScale.ToString(), DefaultValue::TransmissionDispersionScale);
+	ConnectNodeOutputToInput(Input::TransmissionDispersionAbbeNumber, OpenPBRSurfaceShaderNode, OpenPBRSurface::Parameters::TransmissionDispersionAbbeNumber.ToString(), DefaultValue::TransmissionDispersionAbbeNumber);
 	ConnectNodeOutputToInput(Input::TransmissionScatter, OpenPBRSurfaceShaderNode, OpenPBRSurface::Parameters::TransmissionScatter.ToString(), DefaultValue::TransmissionScatter);
 	ConnectNodeOutputToInput(Input::TransmissionScatterAnisotropy, OpenPBRSurfaceShaderNode, OpenPBRSurface::Parameters::TransmissionScatterAnisotropy.ToString(), DefaultValue::TransmissionScatterAnisotropy);
 	ConnectNodeOutputToInput(Input::TransmissionWeight, OpenPBRSurfaceShaderNode, OpenPBRSurface::Parameters::TransmissionWeight.ToString(), DefaultValue::TransmissionWeight);
diff --git a/Engine/Plugins/Interchange/Runtime/Source/Nodes/Public/InterchangeMaterialDefinitions.h b/Engine/Plugins/Interchange/Runtime/Source/Nodes/Public/InterchangeMaterialDefinitions.h
--- a/Engine/Plugins/Interchange/Runtime/Source/Nodes/Public/InterchangeMaterialDefinitions.h
+++ b/Engine/Plugins/Interchange/Runtime/Source/Nodes/Public/InterchangeMaterialDefinitions.h
@@ -358,7 +358,8 @@ namespace Interchange
 				const FName TransmissionDepth = TEXT("transmission_depth");
 				const FName TransmissionScatter = TEXT("transmission_scatter");
 				const FName TransmissionScatterAnisotropy = TEXT("transmission_scatter_anisotropy");
-				const FName TransmissionDispersion = TEXT("transmission_dispersion");
+				const FName TransmissionDispersionScale = TEXT("transmission_dispersion_scale");
+				const FName TransmissionDispersionAbbeNumber = TEXT("transmission_dispersion_abbe_number");
 				const FName SubsurfaceWeight = TEXT("subsurface_weight");
 				const FName SubsurfaceColor = TEXT("subsurface_color");
 				const FName SubsurfaceRadius = TEXT("subsurface_radius");
diff --git a/Engine/Plugins/Interchange/Runtime/Source/Pipelines/Private/InterchangeMaterialXPipeline.cpp b/Engine/Plugins/Interchange/Runtime/Source/Pipelines/Private/InterchangeMaterialXPipeline.cpp
--- a/Engine/Plugins/Interchange/Runtime/Source/Pipelines/Private/InterchangeMaterialXPipeline.cpp
+++ b/Engine/Plugins/Interchange/Runtime/Source/Pipelines/Private/InterchangeMaterialXPipeline.cpp
@@ -149,7 +149,8 @@ UMaterialXPipelineSettings::UMaterialXPipelineSettings()
 						UE::Interchange::Materials::OpenPBRSurface::Parameters::TransmissionWeight,
 						UE::Interchange::Materials::OpenPBRSurface::Parameters::TransmissionColor,
 						UE::Interchange::Materials::OpenPBRSurface::Parameters::TransmissionDepth,
-						UE::Interchange::Materials::OpenPBRSurface::Parameters::TransmissionDispersion,
+						UE::Interchange::Materials::OpenPBRSurface::Parameters::TransmissionDispersionScale,
+						UE::Interchange::Materials::OpenPBRSurface::Parameters::TransmissionDispersionAbbeNumber,
 						UE::Interchange::Materials::OpenPBRSurface::Parameters::TransmissionScatter,
 						UE::Interchange::Materials::OpenPBRSurface::Parameters::TransmissionScatterAnisotropy,
 						UE::Interchange::Materials::OpenPBRSurface::Parameters::FuzzWeight,
