[Interchange] Dialog's Skeletal Conflict window shows more references than the legacy FBX conflict window does Fix

#rnx

diff --git a/Engine/Plugins/Interchange/Runtime/Source/Pipelines/Private/InterchangeGenericAssetsPipeline.cpp b/Engine/Plugins/Interchange/Runtime/Source/Pipelines/Private/InterchangeGenericAssetsPipeline.cpp
--- a/Engine/Plugins/Interchange/Runtime/Source/Pipelines/Private/InterchangeGenericAssetsPipeline.cpp
+++ b/Engine/Plugins/Interchange/Runtime/Source/Pipelines/Private/InterchangeGenericAssetsPipeline.cpp
@@ -993,7 +993,11 @@ void UInterchangeGenericAssetsPipeline::ShowConflictDialog(const FGuid& Conflict
 							MemoryDependencies.Reserve(ReferencerObjects.Num());
 							for (UObject* ReferencerObject : ReferencerObjects)
 							{
-								MemoryDependencies.Add(*ReferencerObject->GetFullName());
+								if (ReferencerObject->GetOuter()->GetClass() == UPackage::StaticClass()
+									&& ReferencerObject != Skeleton)
+								{
+									MemoryDependencies.Add(*ReferencerObject->GetFullName());
+								}
 							}
 						}
 
