Physics Asset isn't properly populated in Interchange dialog and no Physics Asset can be chosen fix.

#rnx

diff --git a/Engine/Plugins/Interchange/Runtime/Source/Pipelines/Private/InterchangeGenericMeshPipeline.cpp b/Engine/Plugins/Interchange/Runtime/Source/Pipelines/Private/InterchangeGenericMeshPipeline.cpp
--- a/Engine/Plugins/Interchange/Runtime/Source/Pipelines/Private/InterchangeGenericMeshPipeline.cpp
+++ b/Engine/Plugins/Interchange/Runtime/Source/Pipelines/Private/InterchangeGenericMeshPipeline.cpp
@@ -17,6 +17,7 @@
 #include "Misc/Paths.h"
 #include "Nodes/InterchangeBaseNode.h"
 #include "Nodes/InterchangeBaseNodeContainer.h"
+#include "PhysicsEngine/PhysicsAsset.h"
 
 #include UE_INLINE_GENERATED_CPP_BY_NAME(InterchangeGenericMeshPipeline)
 
@@ -78,6 +79,11 @@ void UInterchangeGenericMeshPipeline::AdjustSettingsForContext(EInterchangePipel
 		{
 			//Set the skeleton to the current asset skeleton
 			CommonSkeletalMeshesAndAnimationsProperties->Skeleton = SkeletalMesh->GetSkeleton();
+			PhysicsAsset = SkeletalMesh->GetPhysicsAsset();
+			if (PhysicsAsset.IsValid())
+			{
+				bCreatePhysicsAsset = false;
+			}
 			bImportStaticMeshes = false;
 			HideCategories.Add(StaticMeshesCategory);
 			if(SkeletalMeshImportContentType == EInterchangeSkeletalMeshContentType::Geometry)
