Fixed reimport of actors
    - Root cause: New actor naming done in the actor factory instead of the level pipeline
    - Solution: Translators must provide uniqueasset  names for actor factory nodes. Otherwise, reimport could fail
    - Caveat: Not all translator can do that unless the translated format does enforce uniqueness accross versions, one way or the other. GLTF does not :-(

 - Removed unnecessary argument to FAsset::GenerateNames
 - Modified GLTF translator to use new GLTF node's unique ids as asset names for actors.Misc: Added Tech Writer's corrections too
 - Modified Datasmith translator to use Datasmith uniques names as asset names for actors.
 - Modified FBX parser to use unique ids as asset names for actors
 - Modified generic level pipeline to transfer asset names from translated nodes to actor factory nodes
 - Modified FTaskCreateSceneObjects::DoTask to use factory node's asset name as object name
 - Changed default behavior forDatasmith and generic  level pipelines: now similar to Datasmith behavior
 - Side fixe: eliminate a useless ensure when opening the import dialog for a reimport while no pipeline has a conflict.

#rnx

#virtualized

diff --git a/Engine/Plugins/Interchange/Editor/Source/Pipelines/Private/InterchangeEditorPipelineDetails.cpp b/Engine/Plugins/Interchange/Editor/Source/Pipelines/Private/InterchangeEditorPipelineDetails.cpp
--- a/Engine/Plugins/Interchange/Editor/Source/Pipelines/Private/InterchangeEditorPipelineDetails.cpp
+++ b/Engine/Plugins/Interchange/Editor/Source/Pipelines/Private/InterchangeEditorPipelineDetails.cpp
@@ -112,8 +112,11 @@ TSharedRef<IDetailCustomization> FInterchangePipelineBaseDetailsCustomization::M
 static TArray<TArray<FInterchangeConflictInfo>> ConflicInfosStack;
 void FInterchangePipelineBaseDetailsCustomization::SetConflictsInfo(TArray<FInterchangeConflictInfo>& ConflictInfos)
 {
-	ensure(ConflicInfosStack.Num() == 0);
-	ConflicInfosStack.Push(ConflictInfos);
+	// Only add info when there are any
+	if (ConflicInfosStack.Num() > 0)
+	{
+		ConflicInfosStack.Push(ConflictInfos);
+	}
 }
 
 void FInterchangePipelineBaseDetailsCustomization::SetTextComboBoxWidget(IDetailPropertyRow& PropertyRow, const TSharedPtr<IPropertyHandle>& Handle, const TArray<FString>& PossibleValues)
diff --git a/Engine/Plugins/Interchange/Runtime/Source/Import/Private/Gltf/InterchangeGltfTranslator.cpp b/Engine/Plugins/Interchange/Runtime/Source/Import/Private/Gltf/InterchangeGltfTranslator.cpp
--- a/Engine/Plugins/Interchange/Runtime/Source/Import/Private/Gltf/InterchangeGltfTranslator.cpp
+++ b/Engine/Plugins/Interchange/Runtime/Source/Import/Private/Gltf/InterchangeGltfTranslator.cpp
@@ -223,6 +223,7 @@ void UInterchangeGLTFTranslator::HandleGltfNode( UInterchangeBaseNodeContainer&
 
 	UInterchangeSceneNode* InterchangeSceneNode = NewObject< UInterchangeSceneNode >( &NodeContainer );
 	InterchangeSceneNode->InitializeNode( NodeUid, GltfNode.Name, EInterchangeNodeContainerType::TranslatedScene );
+	InterchangeSceneNode->SetAssetName(GltfNode.UniqueId);
 	NodeContainer.AddNode( InterchangeSceneNode );
 
 	NodeUidMap.Add( &GltfNode, NodeUid );
@@ -256,7 +257,7 @@ void UInterchangeGLTFTranslator::HandleGltfNode( UInterchangeBaseNodeContainer&
 				}
 				else
 				{
-					UE_LOG(LogInterchangeImport, Warning, TEXT("GLTF Node [%d] Import Warning. Gltf Node's MorphTargetNames count is missmatched against MorphTargetWeights count."), *GltfNode.UniqueId);
+					UE_LOG(LogInterchangeImport, Warning, TEXT("glTF Node [%d] Import Warning: The glTF node's MorphTargetNames count does not match its MorphTargetWeights count."), *GltfNode.UniqueId);
 				}
 			}
 
@@ -305,7 +306,7 @@ void UInterchangeGLTFTranslator::HandleGltfNode( UInterchangeBaseNodeContainer&
 					}
 					else
 					{
-						UE_LOG(LogInterchangeImport, Warning, TEXT("GLTF Node [%d] Import Warning. Gltf Node's MorphTargetNames count is missmatched against MorphTargetWeights count."), *GltfNode.UniqueId);
+						UE_LOG(LogInterchangeImport, Warning, TEXT("glTF Node [%d] Import Warning: The glTF node's MorphTargetNames count does not match its MorphTargetWeights count."), *GltfNode.UniqueId);
 					}
 				}
 
@@ -450,7 +451,7 @@ bool UInterchangeGLTFTranslator::Translate( UInterchangeBaseNodeContainer& NodeC
 		UInterchangeResultError_Generic* ErrorResult = AddMessage< UInterchangeResultError_Generic >();
 		ErrorResult->SourceAssetName = FileName;
 		ErrorResult->Text = FText::Format(
-			LOCTEXT("UnsupportedRequiredExtensions", "Not All Required Extensions are supported. (Unsupported extensions: {0})"),
+			LOCTEXT("UnsupportedRequiredExtensions", "Not all required extensions are supported. (Unsupported extensions: {0})"),
 			FText::FromString(NotSupportedRequiredExtensionsStringified));
 
 		SendAnalytics(TranslationResult::NOTSUPPORTED_EXTENSION_FOUND, GltfAsset);
@@ -788,7 +789,7 @@ bool UInterchangeGLTFTranslator::Translate( UInterchangeBaseNodeContainer& NodeC
 
 	if (UnusedGltfMeshIndices.Num() != 0)
 	{
-		UE_LOG(LogInterchangeImport, Warning, TEXT("GLTF Mesh Import Warning. Gltf Mesh Usage expectation is not met."));
+		UE_LOG(LogInterchangeImport, Warning, TEXT("glTF Mesh Import Warning: glTF mesh usage expectations are not met."));
 	}
 
 	SendAnalytics(TranslationResult::SUCCESSFULL, GltfAsset);
@@ -1153,7 +1154,7 @@ TFuture< TOptional< UE::Interchange::FMeshPayloadData > > UInterchangeGLTFTransl
 				{
 					UInterchangeResultError_Generic* ErrorResult = AddMessage<UInterchangeResultError_Generic>();
 					ErrorResult->SourceAssetName = SourceData ? SourceData->GetFilename() : FString();
-					ErrorResult->Text = NSLOCTEXT("UInterchangeGLTFTranslator", "GetMeshPayloadData_ValidateMeshDescriptionFail", "Invalid mesh data (NAN) was found and fix to zero. Mesh render can be bad.");
+					ErrorResult->Text = NSLOCTEXT("UInterchangeGLTFTranslator", "GetMeshPayloadData_ValidateMeshDescriptionFail", "Invalid mesh data (NAN) was found and changed to zero. This may affect the mesh rendering.");
 				}
 
 				Result.Emplace(MeshPayLoadData);
@@ -1363,7 +1364,7 @@ UInterchangeGLTFTranslator::UInterchangeGLTFTranslator()
 
 	if (!UE::Interchange::GLTFMaterials::AreRequiredPackagesLoaded())
 	{
-		UE_LOG(LogInterchangeImport, Warning, TEXT("UInterchangeGLTFPipeline: Some required packages are missing. Material import might be wrong"));
+		UE_LOG(LogInterchangeImport, Warning, TEXT("UInterchangeGLTFPipeline: Some required packages are missing. Material import might be wrong."));
 	}
 }
 
diff --git a/Engine/Plugins/Interchange/Runtime/Source/Import/Private/Scene/InterchangeActorHelper.cpp b/Engine/Plugins/Interchange/Runtime/Source/Import/Private/Scene/InterchangeActorHelper.cpp
--- a/Engine/Plugins/Interchange/Runtime/Source/Import/Private/Scene/InterchangeActorHelper.cpp
+++ b/Engine/Plugins/Interchange/Runtime/Source/Import/Private/Scene/InterchangeActorHelper.cpp
@@ -39,7 +39,7 @@ AActor* UE::Interchange::ActorHelper::GetSpawnedParentActor(const UInterchangeBa
 
 AActor* UE::Interchange::ActorHelper::SpawnFactoryActor(const UInterchangeFactoryBase::FImportSceneObjectsParams& CreateSceneObjectsParams)
 {
-	const UInterchangeActorFactoryNode* FactoryNode = Cast<UInterchangeActorFactoryNode>(CreateSceneObjectsParams.FactoryNode);
+	UInterchangeActorFactoryNode* FactoryNode = Cast<UInterchangeActorFactoryNode>(CreateSceneObjectsParams.FactoryNode);
 	const UInterchangeBaseNodeContainer* NodeContainer = CreateSceneObjectsParams.NodeContainer;
 
 	if (!FactoryNode || !NodeContainer)
@@ -47,20 +47,8 @@ AActor* UE::Interchange::ActorHelper::SpawnFactoryActor(const UInterchangeFactor
 		return nullptr;
 	}
 
-	//To avoid too long actor's name, compute a 128-bit hash based on the string
-	// and use that as a unique name in the form of a GUID
-	FString UniqueID;
-	{
-		FTCHARToUTF8 Converted(*CreateSceneObjectsParams.ObjectName);
-		FMD5 MD5Gen;
-		MD5Gen.Update((const uint8*)Converted.Get(), Converted.Length());
-		uint32 Digest[4];
-		MD5Gen.Final((uint8*)Digest);
-		UniqueID = FGuid(Digest[0], Digest[1], Digest[2], Digest[3]).ToString(EGuidFormats::Base36Encoded);
-	}
-
 	FActorSpawnParameters SpawnParameters;
-	SpawnParameters.Name = FName(*UniqueID);
+	SpawnParameters.Name = FName(CreateSceneObjectsParams.ObjectName);
 	SpawnParameters.NameMode = FActorSpawnParameters::ESpawnActorNameMode::Requested;
 	SpawnParameters.OverrideLevel = CreateSceneObjectsParams.Level;
 	SpawnParameters.SpawnCollisionHandlingOverride = ESpawnActorCollisionHandlingMethod::AlwaysSpawn;
@@ -125,7 +113,7 @@ AActor* UE::Interchange::ActorHelper::SpawnFactoryActor(const UInterchangeFactor
 	if (SpawnedActor)
 	{
 #if WITH_EDITOR
-		SpawnedActor->SetActorLabel(CreateSceneObjectsParams.ObjectName);
+		SpawnedActor->SetActorLabel(FactoryNode->GetDisplayLabel());
 #endif
 		if (!SpawnedActor->GetRootComponent())
 		{
diff --git a/Engine/Plugins/Interchange/Runtime/Source/Parsers/GLTFCore/Private/GLTF/GLTFAsset.cpp b/Engine/Plugins/Interchange/Runtime/Source/Parsers/GLTFCore/Private/GLTF/GLTFAsset.cpp
--- a/Engine/Plugins/Interchange/Runtime/Source/Parsers/GLTFCore/Private/GLTF/GLTFAsset.cpp
+++ b/Engine/Plugins/Interchange/Runtime/Source/Parsers/GLTFCore/Private/GLTF/GLTFAsset.cpp
@@ -32,7 +32,7 @@ namespace GLTF
 			}
 		}
 
-		//Make the Name FName complient.
+		//Make the Name FName compliant.
 		//For Scenario where Asset.Name is '_x' (where x is a number) making FName to be none:
 		//UE-174253
 		template <typename T>
@@ -98,13 +98,10 @@ namespace GLTF
 			ExtraBinData.Empty(ExtraBinBufferKBytes * 1024);
 	}
 
-	void FAsset::GenerateNames(const FString& Prefix)
+	void FAsset::GenerateNames()
 	{
-		ensure(!Prefix.IsEmpty());
-
 		{
-			const FString ScenePrefix = Prefix + TEXT("_scene_");
-
+			const FString ScenePrefix = Name + TEXT("_scene_");
 			for (size_t SceneIndex = 0; SceneIndex < Scenes.Num(); SceneIndex++)
 			{
 				FScene& Scene = Scenes[SceneIndex];
@@ -121,8 +118,8 @@ namespace GLTF
 		}
 
 		{
-			const FString NodePrefix = Prefix + TEXT("_node_");
-			const FString JointPrefix = Prefix + TEXT("_joint_");
+			const FString NodePrefix = Name + TEXT("_node_");
+			const FString JointPrefix = Name + TEXT("_joint_");
 			for (int32 NodeIndex = 0; NodeIndex < Nodes.Num(); ++NodeIndex)
 			{
 				FNode& Node = Nodes[NodeIndex];
@@ -142,7 +139,7 @@ namespace GLTF
 		}
 
 		{
-			const FString TexPrefix = Prefix + TEXT("_texture_");
+			const FString TexPrefix = Name + TEXT("_texture_");
 			for (int32 TextureIndex = 0; TextureIndex < Textures.Num(); ++TextureIndex)
 			{
 				FTexture& Tex = Textures[TextureIndex];
@@ -179,7 +176,7 @@ namespace GLTF
 		}
 
 		{
-			const FString CameraPrefix = Prefix + TEXT("_camera_");
+			const FString CameraPrefix = Name + TEXT("_camera_");
 			for (int32 CameraIndex = 0; CameraIndex < Cameras.Num(); ++CameraIndex)
 			{
 				FCamera& Camera = Cameras[CameraIndex];
@@ -203,7 +200,7 @@ namespace GLTF
 		}
 
 		{
-			const FString LightPrefix = Prefix + TEXT("_light_");
+			const FString LightPrefix = Name + TEXT("_light_");
 			for (int32 LightIndex = 0; LightIndex < Lights.Num(); ++LightIndex)
 			{
 				FLight& Light = Lights[LightIndex];
@@ -228,14 +225,14 @@ namespace GLTF
 
 		// General pattern here is that we'll have a Prefix that will be a filename (like MyFile). 
 		// The name (display label) and the UniqueId of an object has been separated.
-		// UniqueId will consist of: [Prefil]+[_ElementType_]+[ElementIndex]
+		// UniqueId will consist of: [Prefix]+[_ElementType_]+[ElementIndex]
 		// Name will be the provided GLTF object.name, if present,
 		// if Name is not present the UniqueId will be used as the Name.
-		GLTF::GenerateNames(Prefix + TEXT("_material_"), Materials);
-		GLTF::GenerateNames(Prefix + TEXT("_skin_"), Skins);
-		GLTF::GenerateNames(Prefix + TEXT("_animation_"), Animations);
-		GLTF::GenerateNames(Prefix + TEXT("_image_"), Images);
-		GLTF::GenerateNames(Prefix + TEXT("_mesh_"), Meshes);
+		GLTF::GenerateNames(Name + TEXT("_material_"), Materials);
+		GLTF::GenerateNames(Name + TEXT("_skin_"), Skins);
+		GLTF::GenerateNames(Name + TEXT("_animation_"), Animations);
+		GLTF::GenerateNames(Name + TEXT("_image_"), Images);
+		GLTF::GenerateNames(Name + TEXT("_mesh_"), Meshes);
 
 		// Validate FName criteria for Materials Textures Meshes Animations (assets)
 		GLTF::ValidateFNameCriteria(TEXT("Material_"), Materials);
diff --git a/Engine/Plugins/Interchange/Runtime/Source/Parsers/GLTFCore/Private/GLTF/GLTFReader.cpp b/Engine/Plugins/Interchange/Runtime/Source/Parsers/GLTFCore/Private/GLTF/GLTFReader.cpp
--- a/Engine/Plugins/Interchange/Runtime/Source/Parsers/GLTFCore/Private/GLTF/GLTFReader.cpp
+++ b/Engine/Plugins/Interchange/Runtime/Source/Parsers/GLTFCore/Private/GLTF/GLTFReader.cpp
@@ -924,7 +924,7 @@ namespace GLTF
 				}
 			}
 		}
-		OutAsset.GenerateNames(OutAsset.Name);
+		OutAsset.GenerateNames();
 
 		if (OutAsset.ValidationCheck() != FAsset::Valid)
 		{
diff --git a/Engine/Plugins/Interchange/Runtime/Source/Parsers/GLTFCore/Public/GLTFAsset.h b/Engine/Plugins/Interchange/Runtime/Source/Parsers/GLTFCore/Public/GLTFAsset.h
--- a/Engine/Plugins/Interchange/Runtime/Source/Parsers/GLTFCore/Public/GLTFAsset.h
+++ b/Engine/Plugins/Interchange/Runtime/Source/Parsers/GLTFCore/Public/GLTFAsset.h
@@ -126,7 +126,7 @@ namespace GLTF
 		 *
 		 * @param Prefix - prefix to add to the entities name.
 		 */
-		void GenerateNames(const FString& Prefix);
+		void GenerateNames();
 
 		/**
 		 * Finds the indices for the nodes which are root nodes.
diff --git a/Engine/Plugins/Interchange/Runtime/Source/Pipelines/Private/InterchangeGenericScenesPipeline.cpp b/Engine/Plugins/Interchange/Runtime/Source/Pipelines/Private/InterchangeGenericScenesPipeline.cpp
--- a/Engine/Plugins/Interchange/Runtime/Source/Pipelines/Private/InterchangeGenericScenesPipeline.cpp
+++ b/Engine/Plugins/Interchange/Runtime/Source/Pipelines/Private/InterchangeGenericScenesPipeline.cpp
@@ -449,6 +449,9 @@ void UInterchangeGenericLevelPipeline::ExecuteSceneNodePreImport(const FTransfor
 	ActorFactoryNode->InitializeNode(FactoryNodeUid, SceneNode->GetDisplayLabel(), EInterchangeNodeContainerType::FactoryData);
 	const FString ActorFactoryNodeUid = BaseNodeContainer->AddNode(ActorFactoryNode);
 
+	// The translator is responsible to provide a unique name
+	ActorFactoryNode->SetAssetName(SceneNode->GetAssetName());
+
 	if (!SceneNode->GetParentUid().IsEmpty())
 	{
 		/* Find all scene node that are active joint. Non active joint should be convert to actor if they are in a static mesh hierarchy */
diff --git a/Engine/Source/Runtime/Interchange/Engine/Private/Tasks/InterchangeTaskCreateSceneObjects.cpp b/Engine/Source/Runtime/Interchange/Engine/Private/Tasks/InterchangeTaskCreateSceneObjects.cpp
--- a/Engine/Source/Runtime/Interchange/Engine/Private/Tasks/InterchangeTaskCreateSceneObjects.cpp
+++ b/Engine/Source/Runtime/Interchange/Engine/Private/Tasks/InterchangeTaskCreateSceneObjects.cpp
@@ -74,15 +74,15 @@ void UE::Interchange::FTaskCreateSceneObjects::DoTask(ENamedThreads::Type Curren
 		Factory->SetResultsContainer(AsyncHelper->AssetImportResult->GetResults());
 		AsyncHelper->AddCreatedFactory(FactoryNode->GetUniqueID(), Factory);
 
-		FString NodeDisplayName = FactoryNode->GetDisplayLabel();
-		SanitizeObjectName(NodeDisplayName);
+		FString SceneNodeName = FactoryNode->GetAssetName();
+		SanitizeObjectName(SceneNodeName);
 
 		UInterchangeFactoryBase::FImportSceneObjectsParams CreateSceneObjectsParams;
-		CreateSceneObjectsParams.ObjectName = NodeDisplayName;
+		CreateSceneObjectsParams.ObjectName = SceneNodeName;
 		CreateSceneObjectsParams.FactoryNode = FactoryNode;
 		CreateSceneObjectsParams.Level = ImportLevel;
-		CreateSceneObjectsParams.ReimportObject = FFactoryCommon::GetObjectToReimport(ReimportObject, *FactoryNode, WorldPath, WorldName, NodePrefix + NodeDisplayName);
-		CreateSceneObjectsParams.ReimportFactoryNode = FFactoryCommon::GetFactoryNode(ReimportObject, WorldPath, WorldName, NodePrefix + NodeDisplayName);
+		CreateSceneObjectsParams.ReimportObject = FFactoryCommon::GetObjectToReimport(ReimportObject, *FactoryNode, WorldPath, WorldName, NodePrefix + SceneNodeName);
+		CreateSceneObjectsParams.ReimportFactoryNode = FFactoryCommon::GetFactoryNode(ReimportObject, WorldPath, WorldName, NodePrefix + SceneNodeName);
 
 		if (AsyncHelper->BaseNodeContainers.IsValidIndex(SourceIndex))
 		{
