[Interchange - FBX] Importing Quixel FBX files as asset gives different results from legacy fix.
On behalf of Balazs Toereky
#rnx

diff --git a/Engine/Plugins/Interchange/Runtime/Source/Pipelines/Private/InterchangeGenericStaticMeshPipeline.cpp b/Engine/Plugins/Interchange/Runtime/Source/Pipelines/Private/InterchangeGenericStaticMeshPipeline.cpp
--- a/Engine/Plugins/Interchange/Runtime/Source/Pipelines/Private/InterchangeGenericStaticMeshPipeline.cpp
+++ b/Engine/Plugins/Interchange/Runtime/Source/Pipelines/Private/InterchangeGenericStaticMeshPipeline.cpp
@@ -30,22 +30,38 @@ enum class EMeshCollisionType
 	Convex
 };
 
-static TTuple<EMeshCollisionType, FString> GetCollisionMeshType(const UInterchangeBaseNodeContainer& NodeContainer, const FString& NodeUid, const TArray<FString>& AllNodeUids)
+namespace UE::Interchange::Private
 {
-	auto GetNodeName = [&NodeContainer](const FString& NodeUid)
+	FString GetNodeName(const UInterchangeBaseNodeContainer& NodeContainer, const FString& NodeUid)
 	{
 		const UInterchangeBaseNode* BaseNode = NodeContainer.GetNode(NodeUid);
 		if (BaseNode)
 		{
-			return BaseNode->GetDisplayLabel();
+			FString NodeName = BaseNode->GetDisplayLabel();
+			if (const UInterchangeSceneNode* SceneNode = Cast<UInterchangeSceneNode>(BaseNode))
+			{
+				if (SceneNode->IsSpecializedTypeContains(UE::Interchange::FSceneNodeStaticData::GetLodGroupSpecializeTypeString()))
+				{
+					TArray<FString> LodGroupChildrens = NodeContainer.GetNodeChildrenUids(SceneNode->GetUniqueID());
+					if (LodGroupChildrens.Num() > 0)
+					{
+						return GetNodeName(NodeContainer, LodGroupChildrens[0]);
+					}
+				}
+			}
+
+			return NodeName;
 		}
 		else
 		{
 			return FString();
 		}
-	};
+	}
+}
 
-	FString MeshName = GetNodeName(NodeUid);
+static TTuple<EMeshCollisionType, FString> GetCollisionMeshType(const UInterchangeBaseNodeContainer& NodeContainer, const FString& NodeUid, const TArray<FString>& AllNodeUids)
+{
+	FString MeshName = UE::Interchange::Private::GetNodeName(NodeContainer, NodeUid);
 	EMeshCollisionType CollisionType = EMeshCollisionType::None;
 
 	// Determine if the mesh name is a potential collision mesh
@@ -81,10 +97,10 @@ static TTuple<EMeshCollisionType, FString> GetCollisionMeshType(const UInterchan
 	int32 LastUnderscore = INDEX_NONE;
 	verify(MeshName.FindLastChar(TEXT('_'), LastUnderscore));
 
-	auto MatchPredicate = [&GetNodeName](FString Body)
+	auto MatchPredicate = [&NodeContainer](FString Body)
 	{
 		// Generate a predicate to be used by the below Finds.
-		return [Body, &GetNodeName](const FString& ToCompare) { return Body == GetNodeName(ToCompare); };
+		return [Body, &NodeContainer](const FString& ToCompare) { return Body == UE::Interchange::Private::GetNodeName(NodeContainer, ToCompare); };
 	};
 
 	// If we find a mesh named the same as the collision mesh (following the collision prefix), we have a match
