[glTF] Imported animation seems broken fixes. MultiRoot fixes, Skeletal and Skeleton not as CommonRoot use cases.

#rnx

diff --git a/Engine/Plugins/Interchange/Runtime/Source/Import/Private/Gltf/InterchangeGltfTranslator.cpp b/Engine/Plugins/Interchange/Runtime/Source/Import/Private/Gltf/InterchangeGltfTranslator.cpp
--- a/Engine/Plugins/Interchange/Runtime/Source/Import/Private/Gltf/InterchangeGltfTranslator.cpp
+++ b/Engine/Plugins/Interchange/Runtime/Source/Import/Private/Gltf/InterchangeGltfTranslator.cpp
@@ -213,7 +213,7 @@ namespace UE::Interchange::Gltf::Private
 
 void UInterchangeGLTFTranslator::HandleGltfNode( UInterchangeBaseNodeContainer& NodeContainer, const GLTF::FNode& GltfNode, const FString& ParentNodeUid, const int32 NodeIndex, 
 	bool &bHasVariants, TArray<int32>& SkinnedMeshNodes, TSet<int>& UnusedMeshIndices,
-	const TMap<int32, FTransform>& T0Transforms ) const
+	const TMap<int32, FTransform>& T0Transforms, const FString& SceneNodeUid) const
 {
 	using namespace UE::Interchange::Gltf::Private;
 
@@ -238,6 +238,11 @@ void UInterchangeGLTFTranslator::HandleGltfNode( UInterchangeBaseNodeContainer&
 		{
 			SkinnedMeshNodes.Add(NodeIndex);
 
+			if (GltfNode.ParentIndex != INDEX_NONE)
+			{
+				UE_LOG(LogInterchangeImport, Warning, TEXT("Node [%s] with a skinned mesh is not root. Parent transforms will not affect a skinned mesh."), *GltfNode.Name);
+			}
+
 			if (!bHasVariants && GltfAsset.Variants.Num() > 0)
 			{
 				bHasVariants |= CheckForVariants(GltfAsset.Meshes[GltfNode.MeshIndex], GltfAsset.Variants.Num(), GltfAsset.Materials.Num());
@@ -352,7 +357,11 @@ void UInterchangeGLTFTranslator::HandleGltfNode( UInterchangeBaseNodeContainer&
 
 	InterchangeSceneNode->SetCustomLocalTransform(&NodeContainer, Transform, bResetCache);
 
-	if ( !ParentNodeUid.IsEmpty() )
+	if (GltfNode.Type == GLTF::FNode::EType::MeshSkinned)
+	{
+		NodeContainer.SetNodeParentUid(NodeUid, SceneNodeUid);
+	}
+	else if ( !ParentNodeUid.IsEmpty() )
 	{
 		NodeContainer.SetNodeParentUid( NodeUid, ParentNodeUid );
 	}
@@ -361,7 +370,7 @@ void UInterchangeGLTFTranslator::HandleGltfNode( UInterchangeBaseNodeContainer&
 	{
 		if ( GltfAsset.Nodes.IsValidIndex( ChildIndex ) )
 		{
-			HandleGltfNode( NodeContainer, GltfAsset.Nodes[ ChildIndex ], NodeUid, ChildIndex, bHasVariants, SkinnedMeshNodes, UnusedMeshIndices, T0Transforms);
+			HandleGltfNode( NodeContainer, GltfAsset.Nodes[ ChildIndex ], NodeUid, ChildIndex, bHasVariants, SkinnedMeshNodes, UnusedMeshIndices, T0Transforms, SceneNodeUid);
 		}
 	}
 }
@@ -721,7 +730,7 @@ bool UInterchangeGLTFTranslator::Translate( UInterchangeBaseNodeContainer& NodeC
 			{
 				if ( GltfAsset.Nodes.IsValidIndex( NodeIndex ) )
 				{
-					HandleGltfNode( NodeContainer, GltfAsset.Nodes[ NodeIndex ], SceneNodeUid, NodeIndex, bHasVariants, SkinnedMeshNodes, UnusedGltfMeshIndices, T0Transforms);
+					HandleGltfNode( NodeContainer, GltfAsset.Nodes[ NodeIndex ], SceneNodeUid, NodeIndex, bHasVariants, SkinnedMeshNodes, UnusedGltfMeshIndices, T0Transforms, SceneNodeUid);
 				}
 			}
 
diff --git a/Engine/Plugins/Interchange/Runtime/Source/Import/Private/Mesh/InterchangeSkeletalMeshFactory.cpp b/Engine/Plugins/Interchange/Runtime/Source/Import/Private/Mesh/InterchangeSkeletalMeshFactory.cpp
--- a/Engine/Plugins/Interchange/Runtime/Source/Import/Private/Mesh/InterchangeSkeletalMeshFactory.cpp
+++ b/Engine/Plugins/Interchange/Runtime/Source/Import/Private/Mesh/InterchangeSkeletalMeshFactory.cpp
@@ -1458,7 +1458,8 @@ UInterchangeFactoryBase::FImportAssetResult UInterchangeSkeletalMeshFactory::Imp
 		ensure(RootJointNode->GetCustomGlobalTransform(Arguments.NodeContainer, GlobalOffsetTransform, RootJointNodeGlobalTransform));
 		FTransform RootJointNodeLocalTransform;
 		ensure(RootJointNode->GetCustomLocalTransform(RootJointNodeLocalTransform));
-		FTransform BakeToRootJointTransfromModifier = RootJointNodeGlobalTransform.Inverse() * RootJointNodeLocalTransform;
+		FTransform BakeToRootJointTransfromModifier = RootJointNodeGlobalTransform.Inverse() * RootJointNodeLocalTransform; //It is used for !bBakeMeshes, and so the GlobalTransform will be inversed out when multiplied into the CustomBindPoseGlobalTransform
+		FTransform BakeFromRootJointTransfromModifier = RootJointNodeLocalTransform.Inverse() * RootJointNodeGlobalTransform * GlobalOffsetTransform.Inverse(); //GlobalOffseTransform will be added by the BindPoseGlobalTransform when bBakeMeshes && !bRootAncestorOfMeshDependency is used
 
 		TArray<UE::Interchange::Private::FMeshNodeContext> MeshReferences;
 		//Scope to query the mesh node
@@ -1488,16 +1489,32 @@ UInterchangeFactoryBase::FImportAssetResult UInterchangeSkeletalMeshFactory::Imp
 					FString MeshDependencyUid;
 					MeshReference.SceneNode->GetCustomAssetInstanceUid(MeshDependencyUid);
 					MeshReference.MeshNode = Cast<UInterchangeMeshNode>(Arguments.NodeContainer->GetNode(MeshDependencyUid));
+					bool bRootAncestorOfMeshdependency = Arguments.NodeContainer->GetIsAncestor(MeshDependencyUid, RootJointNodeId);
 					//Cache the scene node global matrix, we will use this matrix to bake the vertices, add the node geometric mesh offset to this matrix to bake it properly
 					FTransform SceneNodeTransform;
 					if (!ImportAssetObjectLODData.bUseTimeZeroAsBindPose || !MeshReference.SceneNode->GetCustomTimeZeroGlobalTransform(Arguments.NodeContainer, GlobalOffsetTransform, SceneNodeTransform))
 					{
 						ensure(MeshReference.SceneNode->GetCustomBindPoseGlobalTransform(Arguments.NodeContainer, GlobalOffsetTransform, SceneNodeTransform));
-						if (!bBakeMeshes)
+						if (bRootAncestorOfMeshdependency)
 						{
-							SceneNodeTransform *= BakeToRootJointTransfromModifier;
+							if (!bBakeMeshes)
+							{
+								SceneNodeTransform *= BakeToRootJointTransfromModifier;
+							}
+						}
+						else
+						{
+							if (bBakeMeshes)
+							{
+								SceneNodeTransform = BakeFromRootJointTransfromModifier * SceneNodeTransform;
+							}
+							else
+							{
+								SceneNodeTransform *= GlobalOffsetTransform.Inverse();
+							}
 						}
 					}
+
 					FTransform SceneNodeGeometricTransform;
 					if(MeshReference.SceneNode->GetCustomGeometricTransform(SceneNodeGeometricTransform))
 					{
diff --git a/Engine/Plugins/Interchange/Runtime/Source/Import/Public/Gltf/InterchangeGltfTranslator.h b/Engine/Plugins/Interchange/Runtime/Source/Import/Public/Gltf/InterchangeGltfTranslator.h
--- a/Engine/Plugins/Interchange/Runtime/Source/Import/Public/Gltf/InterchangeGltfTranslator.h
+++ b/Engine/Plugins/Interchange/Runtime/Source/Import/Public/Gltf/InterchangeGltfTranslator.h
@@ -65,7 +65,7 @@ protected:
 
 	void HandleGltfSkeletons( UInterchangeBaseNodeContainer& NodeContainer, const FString& SceneNodeUid, const TArray<int32>& SkinnedMeshNodes, TSet<int>& UnusedMeshIndices ) const;
 	void HandleGltfNode( UInterchangeBaseNodeContainer& NodeContainer, const GLTF::FNode& GltfNode, const FString& ParentNodeUid, const int32 NodeIndex, 
-		bool &bHasVariants, TArray<int32>& SkinnedMeshNodes, TSet<int>& UnusedMeshIndices, const TMap<int32, FTransform>& T0Transforms) const;
+		bool &bHasVariants, TArray<int32>& SkinnedMeshNodes, TSet<int>& UnusedMeshIndices, const TMap<int32, FTransform>& T0Transforms, const FString& SceneNodeUid) const;
 	void HandleGltfVariants(UInterchangeBaseNodeContainer& NodeContainer, const FString& FileName) const;
 	UInterchangeMeshNode* HandleGltfMesh(UInterchangeBaseNodeContainer& NodeContainer, const GLTF::FMesh& GltfMesh, int MeshIndex,
 		TSet<int>& UnusedMeshIndices, const FString& SkeletalName = "" /*If set it creates the mesh even if it was already created (for Skeletals)*/, const FString& SkeletalId = "") const;
diff --git a/Engine/Plugins/Interchange/Runtime/Source/Parsers/GLTFCore/Private/GLTF/GLTFReader.cpp b/Engine/Plugins/Interchange/Runtime/Source/Parsers/GLTFCore/Private/GLTF/GLTFReader.cpp
--- a/Engine/Plugins/Interchange/Runtime/Source/Parsers/GLTFCore/Private/GLTF/GLTFReader.cpp
+++ b/Engine/Plugins/Interchange/Runtime/Source/Parsers/GLTFCore/Private/GLTF/GLTFReader.cpp
@@ -1109,6 +1109,17 @@ namespace GLTF
 		if (!SetupObjects(SceneCount, TEXT("scenes"), [this](const FJsonObject& Object) { SetupScene(Object); })) { return; }
 		if (!SetupObjects(CameraCount, TEXT("cameras"), [this](const FJsonObject& Object) { SetupCamera(Object); })) { return; }
 		if (!SetupObjects(SkinCount, TEXT("skins"), [this](const FJsonObject& Object) { SetupSkin(Object); })) { return; }
+
+		{//BuildRootJoints can affect the node hierarchy and data, as the Animation setup currently stores references to the Nodes, we have to do these setups before the SetupAnimations.
+			SetupNodesType();
+
+			GenerateInverseBindPosesPerSkinIndices();
+			GenerateLocalBindPosesPerSkinIndices();
+			SetLocalBindPosesForJoints();
+
+			BuildRootJoints();
+		}
+
 		if (!SetupObjects(AnimationsCount, TEXT("animations"), [this](const FJsonObject& Object) { SetupAnimation(Object); })) { return; }
 
 		if (!SetupObjects(ImageCount, TEXT("images"), [this, InResourcesPath, bInLoadImageData](const FJsonObject& Object) { SetupImage(Object, InResourcesPath, bInLoadImageData); })) { return; }
@@ -1127,14 +1138,7 @@ namespace GLTF
 			}
 		}
 
-		SetupNodesType();
-
-		GenerateInverseBindPosesPerSkinIndices();
-		GenerateLocalBindPosesPerSkinIndices();
-		SetLocalBindPosesForJoints();
-
 		ExtensionsHandler->SetupAssetExtensions(*JsonRoot);
-		BuildRootJoints();
 	}
 
 	bool FFileReader::CheckForErrors(int32 StartIndex) const
@@ -1291,12 +1295,7 @@ namespace GLTF
 					}
 					else
 					{
-						FTransform ParentGlobalTransform;
-						if (Skin.Skeleton != INDEX_NONE && Skin.Skeleton != CurrentNode.Index)
-						{
-							GenerateGlobalTransform(Asset->Nodes, CurrentNode.ParentIndex, ParentGlobalTransform, Skin.Skeleton);
-						}
-						FTransform LocalBindPose = CurrentNode.SkinIndexToGlobalInverseBindTransform[SkinIndex].Inverse() * ParentGlobalTransform.Inverse();
+						FTransform LocalBindPose = CurrentNode.SkinIndexToGlobalInverseBindTransform[SkinIndex].Inverse();
 						CurrentNode.SkinIndexToLocalBindPose.Add(SkinIndex, LocalBindPose);
 					}
 				}
@@ -1393,6 +1392,131 @@ namespace GLTF
 	}
 	void FFileReader::BuildRootJoints() const
 	{
+		//Fix Multi root problems (at least for the semi-trivial "same parent non-joint" scenario)
+		{
+			auto GetRootDistance = [&](const GLTF::FNode& Node)
+				{
+					int32 Distance = 0;
+					int32 CurrentNodeIndex = Node.Index;
+					while (Asset->Nodes[CurrentNodeIndex].ParentIndex != INDEX_NONE && Asset->Nodes.IsValidIndex(Asset->Nodes[CurrentNodeIndex].ParentIndex))
+					{
+						Distance++;
+						CurrentNodeIndex = Asset->Nodes[CurrentNodeIndex].ParentIndex;
+					}
+					return Distance;
+				};
+
+			struct FRootJoints {
+				TArray<int32> Indices;
+				int32 ParentDistanceFromRoot;
+			};
+			auto SortPredicate = [](const FRootJoints& GroupA, const FRootJoints& GroupB)
+				{
+					return GroupA.ParentDistanceFromRoot > GroupB.ParentDistanceFromRoot;
+				};
+
+			for (size_t SkinIndex = 0; SkinIndex < Asset->Skins.Num(); SkinIndex++)
+			{
+				GLTF::FSkinInfo& Skin = Asset->Skins[SkinIndex];
+
+				//0. Group RootJointNodes by ParentIndices
+				TMap<int32, FRootJoints> ParentToRootJointIndices;
+				for (size_t JointIndex = 0; JointIndex < Skin.Joints.Num(); JointIndex++)
+				{
+					if (!Asset->Nodes.IsValidIndex(Skin.Joints[JointIndex]))
+					{
+						continue;
+					}
+					GLTF::FNode& JointNode = Asset->Nodes[Skin.Joints[JointIndex]];
+					if (!Asset->Nodes.IsValidIndex(JointNode.ParentIndex) || Asset->Nodes[JointNode.ParentIndex].Type != GLTF::FNode::EType::Joint)
+					{
+						FRootJoints* RootJoints = ParentToRootJointIndices.Find(JointNode.ParentIndex);
+						if (RootJoints)
+						{
+							RootJoints->Indices.Add(JointNode.Index);
+						}
+						else
+						{
+							RootJoints = &ParentToRootJointIndices.Add(JointNode.ParentIndex);
+							RootJoints->Indices.Add(JointNode.Index);
+							RootJoints->ParentDistanceFromRoot = GetRootDistance(JointNode);
+						}
+					}
+				}
+				
+				//1. Sort groups based on distance from common root (furthest to closest)
+				ParentToRootJointIndices.ValueSort(SortPredicate);
+
+				//2. Validate Skin.Skeleton if it exists:
+				if (Skin.Skeleton != INDEX_NONE && Asset->Nodes.IsValidIndex(Skin.Skeleton))
+				{
+					const GLTF::FNode SkinSkeletonNode = Asset->Nodes[Skin.Skeleton];
+					int32 SkinSkeletonDistance = GetRootDistance(SkinSkeletonNode);
+					
+					for (TPair<int32, FRootJoints> Group : ParentToRootJointIndices)
+					{
+						if (Group.Value.Indices.Num() == 1 )
+						{
+							if (!Asset->Nodes.IsValidIndex(Group.Value.Indices[0]))
+							{
+								continue;
+							}
+							const GLTF::FNode RootNodeCandidate = Asset->Nodes[Group.Value.Indices[0]];
+							int32 RootNodeCandidateDistance = GetRootDistance(RootNodeCandidate);
+							if (SkinSkeletonDistance > RootNodeCandidateDistance)
+							{
+								Messages.Emplace(EMessageSeverity::Warning, TEXT("Skeleton node is not a common root."));
+								break;
+							}
+						}
+						else if (Group.Value.Indices.Num() == 0)
+						{
+							continue;
+						}
+						else
+						{
+							const GLTF::FNode RootNodeCandidate = Asset->Nodes[Group.Key];
+							int32 RootNodeCandidateDistance = GetRootDistance(RootNodeCandidate);
+							if (SkinSkeletonDistance > RootNodeCandidateDistance)
+							{
+								Messages.Emplace(EMessageSeverity::Warning, TEXT("Skeleton node is not a common root."));
+								break;
+							}
+						}
+						
+					}
+				}
+
+				//3. Introduce new true root per group
+				for (TPair<int32, FRootJoints> Group : ParentToRootJointIndices)
+				{
+					if (Group.Value.Indices.Num() < 2)
+					{
+						continue;
+					}
+					GLTF::FNode& OriginalNode = Asset->Nodes[Group.Key];
+
+					GLTF::FNode& Node = Asset->Nodes.Emplace_GetRef();
+					Node.Index = Asset->Nodes.Num() - 1;
+
+					Node.Name = OriginalNode.Name + "_ProxyTrueRootJoint";
+					Node.Transform = FTransform::Identity;
+					Node.Children = OriginalNode.Children;
+					OriginalNode.Children.Reset();
+					OriginalNode.Children.Add(Node.Index);
+					Node.Type = GLTF::FNode::EType::Joint;
+					Node.ParentIndex = OriginalNode.Index;
+
+					for (size_t RootJointIndex = 0; RootJointIndex < Group.Value.Indices.Num(); RootJointIndex++)
+					{
+						GLTF::FNode& JointNode = Asset->Nodes[Group.Value.Indices[RootJointIndex]];
+						JointNode.ParentIndex = Node.Index;
+					}
+				}
+			}
+		}
+
+		//Set up the RootJointIndices:
 		for (size_t Index = 0; Index < Asset->Nodes.Num(); Index++)
 		{
 			if (Asset->Nodes[Index].Type == GLTF::FNode::EType::Joint)
diff --git a/Engine/Plugins/Interchange/Runtime/Source/Pipelines/Private/InterchangeGenericScenesPipeline.cpp b/Engine/Plugins/Interchange/Runtime/Source/Pipelines/Private/InterchangeGenericScenesPipeline.cpp
--- a/Engine/Plugins/Interchange/Runtime/Source/Pipelines/Private/InterchangeGenericScenesPipeline.cpp
+++ b/Engine/Plugins/Interchange/Runtime/Source/Pipelines/Private/InterchangeGenericScenesPipeline.cpp
@@ -485,16 +485,8 @@ void UInterchangeGenericLevelPipeline::ExecuteSceneNodePreImport(const FTransfor
 	{
 		if (bRootJointNode)
 		{
-			LocalTransform = FTransform::Identity;
 			//LocalTransform of RootjointNode is already baked into the Skeletal and animation.
-			//due to that we acquire the Parent SceneNode and get its GlobalTransform:
-			if (!SceneNode->GetParentUid().IsEmpty())
-			{
-				if (const UInterchangeSceneNode* ParentSceneNode = Cast<UInterchangeSceneNode>(BaseNodeContainer->GetNode(SceneNode->GetParentUid())))
-				{
-					ParentSceneNode->GetCustomLocalTransform(LocalTransform);
-				}
-			}
+			LocalTransform = FTransform::Identity;
 		}
 
 		if (SceneNode->GetParentUid().IsEmpty())
diff --git a/Engine/Source/Runtime/Interchange/Core/Private/Nodes/InterchangeBaseNodeContainer.cpp b/Engine/Source/Runtime/Interchange/Core/Private/Nodes/InterchangeBaseNodeContainer.cpp
--- a/Engine/Source/Runtime/Interchange/Core/Private/Nodes/InterchangeBaseNodeContainer.cpp
+++ b/Engine/Source/Runtime/Interchange/Core/Private/Nodes/InterchangeBaseNodeContainer.cpp
@@ -324,3 +324,22 @@ UInterchangeBaseNode* UInterchangeBaseNodeContainer::GetNodeChildrenInternal(con
 	return nullptr;
 }
 
+bool UInterchangeBaseNodeContainer::GetIsAncestor(const FString& NodeUniqueID, const FString& AncestorUID) const
+{
+	FString CurrentNodeUID = NodeUniqueID;
+	
+	while (CurrentNodeUID != UInterchangeBaseNode::InvalidNodeUid() && Nodes.Contains(CurrentNodeUID))
+	{
+		UInterchangeBaseNode* Node = Nodes.FindChecked(CurrentNodeUID);
+		const FString& ParentUID = Node->GetParentUid();
+
+		if (AncestorUID == ParentUID)
+		{
+			return true;
+		}
+
+		CurrentNodeUID = ParentUID;
+	}
+
+	return false;
+}
diff --git a/Engine/Source/Runtime/Interchange/Core/Public/Nodes/InterchangeBaseNodeContainer.h b/Engine/Source/Runtime/Interchange/Core/Public/Nodes/InterchangeBaseNodeContainer.h
--- a/Engine/Source/Runtime/Interchange/Core/Public/Nodes/InterchangeBaseNodeContainer.h
+++ b/Engine/Source/Runtime/Interchange/Core/Public/Nodes/InterchangeBaseNodeContainer.h
@@ -206,6 +206,12 @@ public:
 		ChildrenCache.Reset();
 	}
 
+	/**
+	 * Checks if ParentNodeUID is an ancestor.
+	 */
+	UFUNCTION(BlueprintCallable, Category = "Interchange | Node")
+	INTERCHANGECORE_API bool GetIsAncestor(const FString& NodeUniqueID, const FString& AncestorUID) const;
+
 private:
 
 	INTERCHANGECORE_API UInterchangeBaseNode* GetNodeChildrenInternal(const FString& NodeUniqueID, int32 ChildIndex);
