fix UDIM parser incorrectly treating large numbers like 4096 as UDIM which then led to hard crashes making VT with too many pixels
UDIM valid range is 1001 to 1100 only

diff --git a/Engine/Plugins/Interchange/Runtime/Source/Import/Private/Texture/InterchangeTextureFactory.cpp b/Engine/Plugins/Interchange/Runtime/Source/Import/Private/Texture/InterchangeTextureFactory.cpp
--- a/Engine/Plugins/Interchange/Runtime/Source/Import/Private/Texture/InterchangeTextureFactory.cpp
+++ b/Engine/Plugins/Interchange/Runtime/Source/Import/Private/Texture/InterchangeTextureFactory.cpp
@@ -460,6 +460,11 @@ namespace UE::Interchange::Private::InterchangeTextureFactory
 	{
 		check(!InBlockSourcesData.IsEmpty());
 
+		// note if only one image in UDIM set is found,
+		//	then SetupTexture2DSourceDataFromBulkData_Editor will wind up discarding the UDIM parse
+		//	and instead import it as a single image.
+		// eg. if you have something like "test_1024.bmp" it will not be treated as UDIM with offsets of X=3 and Y=2
+
 		TArray<const TPair<int32, FString>*> UDIMsAndSourcesFileArray;
 		UDIMsAndSourcesFileArray.Reserve(InBlockSourcesData.Num());
 		int32 OrginalSourceDataIndex = INDEX_NONE;
@@ -790,7 +795,8 @@ namespace UE::Interchange::Private::InterchangeTextureFactory
 		}
 		else
 		{
-			//Import as a normal texture
+			// only 1 image in UDIM set found
+			//Import as a normal texture (not VT/not UDIM)
 			FImportImage Image;
 			Image.Format = BlockedImage.Format;
 			Image.CompressionSettings = BlockedImage.CompressionSettings;
@@ -801,6 +807,7 @@ namespace UE::Interchange::Private::InterchangeTextureFactory
 			Image.SizeX = Block.SizeX;
 			Image.SizeY = Block.SizeY;
 			Image.NumMips = Block.NumMips;
+			// note any UDIM block X,Y offset is discarded
 
 			SetupTextureSourceDataFromBulkData_Editor(Texture2D, Image, MoveTemp(BufferAndId), bIsReimport);
 		}
@@ -1093,6 +1100,9 @@ namespace UE::Interchange::Private::InterchangeTextureFactory
 
 	FGraphEventArray GenerateHashSourceFilesTasks(const UInterchangeSourceData* SourceData, TArray<FString>&& FilesToHash, TArray<FAssetImportInfo::FSourceFile>& OutSourceFiles)
 	{
+		// This hashing is quite slow (MD5 single threaded)
+		// it is needed for the auto-reimport feature that detects changed files by comparing their hashes
+
 		struct FHashSourceTaskBase
 		{
 			/**
diff --git a/Engine/Source/Runtime/TextureUtilitiesCommon/Private/UDIMUtilities.cpp b/Engine/Source/Runtime/TextureUtilitiesCommon/Private/UDIMUtilities.cpp
--- a/Engine/Source/Runtime/TextureUtilitiesCommon/Private/UDIMUtilities.cpp
+++ b/Engine/Source/Runtime/TextureUtilitiesCommon/Private/UDIMUtilities.cpp
@@ -45,13 +45,24 @@ namespace UE
 				// UDIM starts with 1001 as the origin
 				return INDEX_NONE;
 			}
+			if ( UdimValue > 1100 )
+			{
+				return INDEX_NONE;
+			}
 
 			return UdimValue;
 		}
 
 		int32 GetUDIMIndex(int32 BlockX, int32 BlockY)
 		{
-			return BlockY * 10 + BlockX + 1001;
+			check( BlockX >= 0 && BlockX <= 9 );
+			check( BlockY >= 0 && BlockY <= 9 );
+
+			int32 Ret = BlockY * 10 + BlockX + 1001;
+
+			check( Ret >= 1001 && Ret <= 1100 );
+
+			return Ret;
 		}
 
 		TMap<int32, FString> GetUDIMBlocksFromSourceFile(const FString& File, const FString& UdimRegexPattern, FString* OutFilenameWithoutUdimPatternAndExtension)
@@ -91,8 +102,15 @@ namespace UE
 
 		void ExtractUDIMCoordinates(int32 UDIMIndex, int32& OutBlockX, int32& OutBlockY)
 		{
+			check( UDIMIndex != INDEX_NONE );
+			check( UDIMIndex >= 1001 && UDIMIndex <= 1100 );
+
+			// note 1010 is X=9, Y=0 , that is the tens column is not Y
 			OutBlockX = (UDIMIndex - 1001) % 10;
 			OutBlockY = (UDIMIndex - 1001) / 10;
+
+			check( OutBlockX >= 0 && OutBlockX <= 9 );
+			check( OutBlockY >= 0 && OutBlockY <= 9 );
 		}
 	}
 }
diff --git a/Engine/Source/Runtime/TextureUtilitiesCommon/Public/UDIMUtilities.h b/Engine/Source/Runtime/TextureUtilitiesCommon/Public/UDIMUtilities.h
--- a/Engine/Source/Runtime/TextureUtilitiesCommon/Public/UDIMUtilities.h
+++ b/Engine/Source/Runtime/TextureUtilitiesCommon/Public/UDIMUtilities.h
@@ -12,6 +12,13 @@ namespace UE
 {
 	namespace TextureUtilitiesCommon
 	{
+		// UDIM = 1001 + X + Y*10
+		// valid UDIM is in "1001" to "1100"
+		//	X and Y are in [0,9]
+		//	note the tens digit is not Y ; eg. 1010 is X=9,Y=0
+
+		// this Regex is unnecessarily complex and broad
+		//	UDIM must start with "10" then two digits, or "1100"
 		constexpr const TCHAR* DefaultUdimRegexPattern = TEXT(R"((.+?)[._](\d{4})$)");
 
 		TEXTUREUTILITIESCOMMON_API uint32 ParseUDIMName(const FString& Name, const FString& UdimRegexPattern, FString& OutPrefixName, FString& OutPostfixName);
