Improve import of property animation in the FBX parser
- FbxAnimation:
-- Import uchar, char types and enums values in the ByteKeyValues array instead of the integer/StringKeyValues
-- Import bool type in the BoolKeyValues array instead of IntegerKeyValues

FbxScene:
- Create animation track nodes
- Import NodeAttributes and Custom attributes properties
- Support of Intensity, Color, bUseTemperature, IntensityUnits, bHidden, CurrentFocalLength, AspectRatioAxisConstraint porperties

GenericAnimationPipeline:
- Update dependencies on AnimationTrackNode instead of just TransformAnimationTrackNode

LevelSequenceFactory
- Fix array out of bounds by taking the channels min between the imported curves and the MovieSceneChannels
- Set the actor as movable for AnimationTrackNode otherwise the actor is set by default as static, hence no animations

diff --git a/Engine/Plugins/Interchange/Runtime/Source/Import/Private/Animation/InterchangeLevelSequenceFactory.cpp b/Engine/Plugins/Interchange/Runtime/Source/Import/Private/Animation/InterchangeLevelSequenceFactory.cpp
--- a/Engine/Plugins/Interchange/Runtime/Source/Import/Private/Animation/InterchangeLevelSequenceFactory.cpp
+++ b/Engine/Plugins/Interchange/Runtime/Source/Import/Private/Animation/InterchangeLevelSequenceFactory.cpp
@@ -533,41 +533,54 @@ namespace UE::Interchange::Private
 		for(const FMovieSceneChannelEntry& ChannelEntry : ChannelEntries)
 		{
 			const FName ChannelTypeName = ChannelEntry.GetChannelTypeName();
-			if(ChannelTypeName != DoubleChannelTypeName &&
-			   ChannelTypeName != FloatChannelTypeName &&
-			   ChannelTypeName != IntegerChannelTypeName &&
-			   ChannelTypeName != BoolChannelTypeName &&
-			   ChannelTypeName != EnumChannelTypeName)
+			const bool bIsBoolChannel = (ChannelTypeName == BoolChannelTypeName);
+			const bool bIsEnumChannel = (ChannelTypeName == EnumChannelTypeName);
+			const bool bIsIntegerChannel = (ChannelTypeName == IntegerChannelTypeName);
+			const bool bIsDoubleChannel = (ChannelTypeName == DoubleChannelTypeName);
+			const bool bIsFloatChannel = (ChannelTypeName == FloatChannelTypeName);
+
+			if(!bIsBoolChannel &&
+			   !bIsEnumChannel &&
+			   !bIsIntegerChannel &&
+			   !bIsDoubleChannel &&
+			   !bIsFloatChannel)
 			{
 				continue;
 			}
 
 			TArrayView<FMovieSceneChannel* const> Channels = ChannelEntry.GetChannels();
-			for(int32 Index = 0; Index < Channels.Num(); ++Index)
+			int32 NumChannels = (bIsBoolChannel || bIsEnumChannel || bIsIntegerChannel) ? PayloadData->StepCurves.Num() : PayloadData->Curves.Num();
+			NumChannels = FMath::Min(NumChannels, Channels.Num());
+			for(int32 Index = 0; Index < NumChannels; ++Index)
 			{
 				FMovieSceneChannelHandle Channel = ChannelProxy.MakeHandle(ChannelTypeName, Index);
-				if(ChannelTypeName == FMovieSceneBoolChannel::StaticStruct()->GetFName())
+				if(bIsBoolChannel)
 				{
 					UpdateStepChannel(*(Channel.Cast<FMovieSceneBoolChannel>().Get()), PayloadData->StepCurves[Index].KeyTimes, PayloadData->StepCurves[0].BooleanKeyValues.GetValue());
 				}
-				else if(ChannelTypeName == FMovieSceneByteChannel::StaticStruct()->GetFName())
+				else if(bIsEnumChannel)
 				{
 					UpdateStepChannel(*(Channel.Cast<FMovieSceneByteChannel>().Get()), PayloadData->StepCurves[Index].KeyTimes, PayloadData->StepCurves[0].ByteKeyValues.GetValue());
 				}
-				else if(ChannelTypeName == FMovieSceneIntegerChannel::StaticStruct()->GetFName())
+				else if(bIsIntegerChannel)
 				{
-					UpdateStepChannel(*(Channel.Cast<FMovieSceneBoolChannel>().Get()), PayloadData->StepCurves[Index].KeyTimes, PayloadData->StepCurves[0].BooleanKeyValues.GetValue());
+					UpdateStepChannel(*(Channel.Cast<FMovieSceneIntegerChannel>().Get()), PayloadData->StepCurves[Index].KeyTimes, PayloadData->StepCurves[0].IntegerKeyValues.GetValue());
 				}
-				else if(ChannelTypeName == FMovieSceneFloatChannel::StaticStruct()->GetFName())
+				else if(bIsFloatChannel)
 				{
 					CopyToChannel(Channel.Cast<FMovieSceneFloatChannel>().Get(), PayloadData->Curves[Index]);
 				}
-				else if(ChannelTypeName == FMovieSceneDoubleChannel::StaticStruct()->GetFName())
+				else if(bIsDoubleChannel)
 				{
 					CopyToChannel(Channel.Cast<FMovieSceneDoubleChannel>().Get(), PayloadData->Curves[Index]);
 				}
 			}
 		}
+
+		if(USceneComponent* SceneComp = Actor->GetRootComponent())
+		{
+			SceneComp->SetMobility(EComponentMobility::Movable);
+		}
 	}
 
 	AActor* FLevelSequenceHelper::GetActor(const UInterchangeAnimationTrackNode& TrackNode)
diff --git a/Engine/Plugins/Interchange/Runtime/Source/Parsers/Fbx/Private/FbxAnimation.cpp b/Engine/Plugins/Interchange/Runtime/Source/Parsers/Fbx/Private/FbxAnimation.cpp
--- a/Engine/Plugins/Interchange/Runtime/Source/Parsers/Fbx/Private/FbxAnimation.cpp
+++ b/Engine/Plugins/Interchange/Runtime/Source/Parsers/Fbx/Private/FbxAnimation.cpp
@@ -301,39 +301,38 @@ namespace UE::Interchange::Private
 		}
 	}
 
-	void ImportFloatStepCurve(const FbxAnimCurve* SourceCurves, FbxProperty& Property, FInterchangeStepCurve& DestinationCurve)
+	template<typename T>
+	void ImportStepCurve(const FbxAnimCurve* SourceCurves, FbxProperty& Property, FInterchangeStepCurve& DestinationCurve)
 	{
-		TArray<float> StepCurveValues;
-		FillStepCurveAttribute<float>(DestinationCurve.KeyTimes, StepCurveValues, SourceCurves, [&Property](const FbxAnimCurveKey* Key, const FbxTime* KeyTime)
+		TArray<T> StepCurveValues;
+		FillStepCurveAttribute<T>(DestinationCurve.KeyTimes, StepCurveValues, SourceCurves, [&Property](const FbxAnimCurveKey* Key, const FbxTime* KeyTime)
+		{
+			if(Key)
 			{
-				if (Key)
-				{
-					return Key->GetValue();
-				}
-				else
-				{
-					return Property.Get<float>();
-				}
-			});
-		DestinationCurve.FloatKeyValues = StepCurveValues;
-
-	}
-
-	void ImportIntegerStepCurve(const FbxAnimCurve* SourceCurves, FbxProperty& Property, FInterchangeStepCurve& DestinationCurve)
-	{
-		TArray<int32> StepCurveValues;
-		FillStepCurveAttribute<int32>(DestinationCurve.KeyTimes, StepCurveValues, SourceCurves, [&Property](const FbxAnimCurveKey* Key, const FbxTime* KeyTime)
+				return static_cast<T>(Key->GetValue());
+			}
+			else
 			{
-				if (Key)
-				{
-					return static_cast<int32>(Key->GetValue());
-				}
-				else
-				{
-					return static_cast<int32>(Property.Get<int32>());
-				}
-			});
-		DestinationCurve.IntegerKeyValues = StepCurveValues;
+				return static_cast<T>(Property.Get<T>());
+			}
+		});
+
+		if constexpr(std::is_same_v<T, bool>)
+		{
+			DestinationCurve.BooleanKeyValues = StepCurveValues;
+		}
+		else if constexpr(std::is_floating_point_v<T>)
+		{
+			check(false); //Float curve payload should be extract as FInterchangeCurve since we can interpolate them
+		}
+		else if constexpr(sizeof(T) == sizeof(uint8))
+		{
+			DestinationCurve.ByteKeyValues = StepCurveValues;
+		}
+		else if constexpr(std::is_integral_v<T> && sizeof(T) > sizeof(uint8))
+		{
+			DestinationCurve.IntegerKeyValues = StepCurveValues;
+		}
 	}
 
 	void ImportStringStepCurve(const FbxAnimCurve* SourceCurves, FbxProperty& Property, FInterchangeStepCurve& DestinationCurve)
@@ -356,34 +355,6 @@ namespace UE::Interchange::Private
 		DestinationCurve.StringKeyValues = StepCurveValues;
 	}
 
-	void ImportEnumStepCurve(const FbxAnimCurve* SourceCurves, FbxProperty& Property, FInterchangeStepCurve& DestinationCurve)
-	{
-		TArray<FString> StepCurveValues;
-		FillStepCurveAttribute<FString>(DestinationCurve.KeyTimes, StepCurveValues, SourceCurves, [&Property](const FbxAnimCurveKey* Key, const FbxTime* KeyTime)
-			{
-				int32 EnumIndex = -1;
-
-				if (KeyTime)
-				{
-					FbxPropertyValue& EvaluatedValue = Property.EvaluateValue(*KeyTime);
-					EvaluatedValue.Get(&EnumIndex, EFbxType::eFbxEnum);
-				}
-				else
-				{
-					EnumIndex = Property.Get<FbxEnum>();
-				}
-
-				if (EnumIndex < 0 || EnumIndex >= Property.GetEnumCount())
-				{
-					return FString();
-				}
-
-				const char* EnumValue = Property.GetEnumValue(EnumIndex);
-				return FString(UTF8_TO_TCHAR(EnumValue));
-			});
-		DestinationCurve.StringKeyValues = StepCurveValues;
-	}
-
 	bool ImportBakeTransforms(FNodeTransformFetchPayloadData& FetchPayloadData, FAnimationPayloadData& AnimationBakeTransformPayloadData, FbxScene* SDKScene)
 	{
 		if (!ensure(!FMath::IsNearlyZero(AnimationBakeTransformPayloadData.BakeFrequency)))
@@ -598,15 +569,18 @@ namespace UE::Interchange::Private
 							switch (FetchPayloadData.PropertyType)
 							{
 							case EFbxType::eFbxBool:
+								ImportStepCurve<bool>(CurrentAnimCurve, FetchPayloadData.Property, InterchangeStepCurves.AddDefaulted_GetRef());
 							case EFbxType::eFbxChar:
 							case EFbxType::eFbxUChar:
+							case EFbxType::eFbxEnum:
+								ImportStepCurve<uint8>(CurrentAnimCurve, FetchPayloadData.Property, InterchangeStepCurves.AddDefaulted_GetRef());
 							case EFbxType::eFbxShort:
 							case EFbxType::eFbxUShort:
 							case EFbxType::eFbxInt:
 							case EFbxType::eFbxUInt:
 							case EFbxType::eFbxLongLong:
 							case EFbxType::eFbxULongLong:
-								ImportIntegerStepCurve(CurrentAnimCurve, FetchPayloadData.Property, InterchangeStepCurves.AddDefaulted_GetRef());
+								ImportStepCurve<int32>(CurrentAnimCurve, FetchPayloadData.Property, InterchangeStepCurves.AddDefaulted_GetRef());
 								break;
 							case EFbxType::eFbxHalfFloat:
 							case EFbxType::eFbxFloat:
@@ -615,10 +589,6 @@ namespace UE::Interchange::Private
 							case EFbxType::eFbxDouble3:
 							case EFbxType::eFbxDouble4:
 								check(false); //Float curve payload should be extract as FInterchangeCurve since we can interpolate them
-								ImportFloatStepCurve(CurrentAnimCurve, FetchPayloadData.Property, InterchangeStepCurves.AddDefaulted_GetRef());
-								break;
-							case EFbxType::eFbxEnum:
-								ImportEnumStepCurve(CurrentAnimCurve, FetchPayloadData.Property, InterchangeStepCurves.AddDefaulted_GetRef());
 								break;
 							case EFbxType::eFbxString:
 								ImportStringStepCurve(CurrentAnimCurve, FetchPayloadData.Property, InterchangeStepCurves.AddDefaulted_GetRef());
diff --git a/Engine/Plugins/Interchange/Runtime/Source/Parsers/Fbx/Private/FbxScene.cpp b/Engine/Plugins/Interchange/Runtime/Source/Parsers/Fbx/Private/FbxScene.cpp
--- a/Engine/Plugins/Interchange/Runtime/Source/Parsers/Fbx/Private/FbxScene.cpp
+++ b/Engine/Plugins/Interchange/Runtime/Source/Parsers/Fbx/Private/FbxScene.cpp
@@ -19,6 +19,7 @@
 #include "Nodes/InterchangeSourceNode.h"
 #include "Nodes/InterchangeUserDefinedAttribute.h"
 #include "InterchangeAnimationTrackSetNode.h"
+#include "InterchangeAnimationDefinitions.h"
 
 #define LOCTEXT_NAMESPACE "InterchangeFbxScene"
 
@@ -352,9 +353,65 @@ namespace UE
 					}
 				}
 				
+				auto AddAnimationTrackNode = [&](const FName& PropertyTrack, const FString& CurveNodeName, const FString& PayloadKey, EInterchangeAnimationPayLoadType PayloadType)
+				{
+					UInterchangeAnimationTrackNode* AnimTrackNode = NewObject< UInterchangeAnimationTrackNode >(&NodeContainer);
+					const FString AnimTrackNodeName = FString::Printf(TEXT("%s"), *UnrealNode->GetDisplayLabel()) + CurveNodeName;
+					const FString AnimTrackNodeUid = TEXT("\\AnimationTrack\\") + AnimTrackNodeName;
+
+					AnimTrackNode->InitializeNode(AnimTrackNodeUid, AnimTrackNodeName, EInterchangeNodeContainerType::TranslatedAsset);
+					AnimTrackNode->SetCustomActorDependencyUid(*UnrealNode->GetUniqueID());
+					AnimTrackNode->SetCustomAnimationPayloadKey(PayloadKey, PayloadType);
+					AnimTrackNode->SetCustomPropertyTrack(PropertyTrack);
+					NodeContainer.AddNode(AnimTrackNode);
+				};
+
+				using namespace UE::Interchange::Animation;
+				static TMap<FString, TPair<FName, EInterchangeAnimationPayLoadType>> PropertyTracks
+				{
+					{TEXT("Intensity"), {PropertyTracks::Light::Intensity, EInterchangeAnimationPayLoadType::CURVE}},
+					{TEXT("Color"), {PropertyTracks::Light::Color, EInterchangeAnimationPayLoadType::CURVE}},
+					{TEXT("bUseTemperature"), {PropertyTracks::Light::UseTemperature, EInterchangeAnimationPayLoadType::STEPCURVE}},
+					{TEXT("IntensityUnits"), {PropertyTracks::Light::IntensityUnits, EInterchangeAnimationPayLoadType::STEPCURVE}},
+					{TEXT("bHidden"), {PropertyTracks::Visibility, EInterchangeAnimationPayLoadType::STEPCURVE}},
+					{TEXT("CurrentFocalLength"), {PropertyTracks::Camera::CurrentFocalLength, EInterchangeAnimationPayLoadType::CURVE}},
+					{TEXT("CurrentAperture"), {PropertyTracks::Camera::CurrentAperture, EInterchangeAnimationPayLoadType::CURVE}},
+					{TEXT("AspectRatioAxisConstraint"), {PropertyTracks::Camera::AspectRatioAxisConstraint, EInterchangeAnimationPayLoadType::CURVE}},
+				};
+
+				//Add all Node Attributes for the node
+				for(int32 i = 0, Count = Node->GetNodeAttributeCount(); i < Count; ++i)
+				{
+					FbxNodeAttribute* NodeAttribute = Node->GetNodeAttributeByIndex(i);
+					FbxProperty Property = NodeAttribute->GetFirstProperty();
+
+					while(Property.IsValid())
+					{
+						FbxAnimCurveNode* CurveNode = Property.GetCurveNode();
+						EFbxType PropertyType = Property.GetPropertyDataType().GetType();
+						if(CurveNode && CurveNode->IsAnimated() && FFbxAnimation::IsFbxPropertyTypeSupported(PropertyType))
+						{
+							TOptional<FString> PayloadKey;
+							//Attribute is animated, add the curves payload key that represent the attribute animation
+							FFbxAnimation::AddNodeAttributeCurvesAnimation(Parser, Node, Property, CurveNode, UnrealNode, PayloadContexts, PropertyType, PayloadKey);
+
+							if(PayloadKey.IsSet())
+							{
+								const char* CurveNodeName = CurveNode->GetName();
+								if(TPair<FName, EInterchangeAnimationPayLoadType>* Track = PropertyTracks.Find(CurveNodeName))
+								{
+									AddAnimationTrackNode(Track->Key, CurveNodeName, *PayloadKey, Track->Value);
+								}
+							}
+						}
+
+						Property = NodeAttribute->GetNextProperty(Property);
+					}
+				}
 
-				//Add all custom Attributes for the node
 				FbxProperty Property = Node->GetFirstProperty();
+
+				//Add all custom Attributes for the node
 				while (Property.IsValid())
 				{
 					EFbxType PropertyType =  Property.GetPropertyDataType().GetType();
@@ -368,6 +425,15 @@ namespace UE
 						{
 							//Attribute is animated, add the curves payload key that represent the attribute animation
 							FFbxAnimation::AddNodeAttributeCurvesAnimation(Parser, Node, Property, CurveNode, UnrealNode, PayloadContexts, PropertyType, PayloadKey);
+
+							if(PayloadKey.IsSet())
+							{
+								const char* CurveNodeName = CurveNode->GetName();
+								if(TPair<FName, EInterchangeAnimationPayLoadType>* Track = PropertyTracks.Find(CurveNodeName))
+								{
+									AddAnimationTrackNode(Track->Key, CurveNodeName, *PayloadKey, Track->Value);
+								}
+							}
 						}
 						switch (Property.GetPropertyDataType().GetType())
 						{
@@ -883,7 +949,7 @@ namespace UE
 				}
 
 				TArray<FString> TransformAnimTrackNodeUids;
-				NodeContainer.IterateNodesOfType<UInterchangeTransformAnimationTrackNode>([&](const FString& NodeUid, UInterchangeTransformAnimationTrackNode* TransformAnimationTrackNode)
+				NodeContainer.IterateNodesOfType<UInterchangeAnimationTrackNode>([&](const FString& NodeUid, UInterchangeAnimationTrackNode* TransformAnimationTrackNode)
 					{
 						TransformAnimTrackNodeUids.Add(NodeUid);
 					});
diff --git a/Engine/Plugins/Interchange/Runtime/Source/Pipelines/Private/InterchangeGenericAnimationPipeline.cpp b/Engine/Plugins/Interchange/Runtime/Source/Pipelines/Private/InterchangeGenericAnimationPipeline.cpp
--- a/Engine/Plugins/Interchange/Runtime/Source/Pipelines/Private/InterchangeGenericAnimationPipeline.cpp
+++ b/Engine/Plugins/Interchange/Runtime/Source/Pipelines/Private/InterchangeGenericAnimationPipeline.cpp
@@ -353,10 +353,10 @@ void UInterchangeGenericAnimationPipeline::CreateLevelSequenceFactoryNode(UInter
 		// Update factory's dependencies
 		if (const UInterchangeAnimationTrackBaseNode* TrackNode = Cast<UInterchangeAnimationTrackBaseNode>(BaseNodeContainer->GetNode(AnimationTrackUid)))
 		{
-			if (const UInterchangeTransformAnimationTrackNode* TransformTrackNode = Cast<UInterchangeTransformAnimationTrackNode>(TrackNode))
+			if (const UInterchangeAnimationTrackNode* AnimationTrackNode = Cast<UInterchangeAnimationTrackNode>(TrackNode))
 			{
 				FString ActorNodeUid;
-				if (TransformTrackNode->GetCustomActorDependencyUid(ActorNodeUid))
+				if (AnimationTrackNode->GetCustomActorDependencyUid(ActorNodeUid))
 				{
 					const FString ActorFactoryNodeUid = UInterchangeFactoryBaseNode::BuildFactoryNodeUid(ActorNodeUid);
 					FactoryNode->AddFactoryDependencyUid(ActorFactoryNodeUid);
