Interchange: Add bake pivot support

#rnx

#virtualized

diff --git a/Engine/Plugins/Interchange/Editor/Source/InterchangeEditor/Private/InterchangeFbxAssetImportDataConverter.cpp b/Engine/Plugins/Interchange/Editor/Source/InterchangeEditor/Private/InterchangeFbxAssetImportDataConverter.cpp
--- a/Engine/Plugins/Interchange/Editor/Source/InterchangeEditor/Private/InterchangeFbxAssetImportDataConverter.cpp
+++ b/Engine/Plugins/Interchange/Editor/Source/InterchangeEditor/Private/InterchangeFbxAssetImportDataConverter.cpp
@@ -70,6 +70,7 @@ namespace UE::Interchange::Private
 		MeshImportData->bImportMeshLODs = GenericAssetPipeline->CommonMeshesProperties->bImportLods;
 		MeshImportData->bReorderMaterialToFbxOrder = true;
 		MeshImportData->bTransformVertexToAbsolute = GenericAssetPipeline->CommonMeshesProperties->bBakeMeshes;
+		MeshImportData->bBakePivotInVertex = GenericAssetPipeline->CommonMeshesProperties->bBakePivotMeshes;
 
 		if (GenericAssetPipeline->CommonMeshesProperties->bUseMikkTSpace)
 		{
@@ -107,6 +108,7 @@ namespace UE::Interchange::Private
 		GenericAssetPipeline->CommonMeshesProperties->bComputeWeightedNormals = LegacyMeshImportData->bComputeWeightedNormals;
 		GenericAssetPipeline->CommonMeshesProperties->bImportLods = LegacyMeshImportData->bImportMeshLODs;
 		GenericAssetPipeline->CommonMeshesProperties->bBakeMeshes = LegacyMeshImportData->bTransformVertexToAbsolute;
+		GenericAssetPipeline->CommonMeshesProperties->bBakePivotMeshes = LegacyMeshImportData->bBakePivotInVertex;
 
 		if (LegacyMeshImportData->NormalGenerationMethod == EFBXNormalGenerationMethod::MikkTSpace)
 		{
diff --git a/Engine/Plugins/Interchange/Runtime/Source/FactoryNodes/Private/InterchangeCommonPipelineDataFactoryNode.cpp b/Engine/Plugins/Interchange/Runtime/Source/FactoryNodes/Private/InterchangeCommonPipelineDataFactoryNode.cpp
--- a/Engine/Plugins/Interchange/Runtime/Source/FactoryNodes/Private/InterchangeCommonPipelineDataFactoryNode.cpp
+++ b/Engine/Plugins/Interchange/Runtime/Source/FactoryNodes/Private/InterchangeCommonPipelineDataFactoryNode.cpp
@@ -60,7 +60,17 @@ bool UInterchangeCommonPipelineDataFactoryNode::GetBakeMeshes(bool& AttributeVal
 	IMPLEMENT_NODE_ATTRIBUTE_GETTER(BakeMeshes, bool);
 }
 
-bool UInterchangeCommonPipelineDataFactoryNode::SetBakeMeshes(const UInterchangeBaseNodeContainer* NodeContainer, const bool& AttributeValue)
+bool UInterchangeCommonPipelineDataFactoryNode::SetBakeMeshes(const bool& AttributeValue)
 {
 	IMPLEMENT_NODE_ATTRIBUTE_SETTER_NODELEGATE(BakeMeshes, bool);
+}
+
+bool UInterchangeCommonPipelineDataFactoryNode::GetBakePivotMeshes(bool& AttributeValue) const
+{
+	IMPLEMENT_NODE_ATTRIBUTE_GETTER(BakePivotMeshes, bool);
+}
+
+bool UInterchangeCommonPipelineDataFactoryNode::SetBakePivotMeshes(const bool& AttributeValue)
+{
+	IMPLEMENT_NODE_ATTRIBUTE_SETTER_NODELEGATE(BakePivotMeshes, bool);
 }
\ No newline at end of file
diff --git a/Engine/Plugins/Interchange/Runtime/Source/FactoryNodes/Public/InterchangeCommonPipelineDataFactoryNode.h b/Engine/Plugins/Interchange/Runtime/Source/FactoryNodes/Public/InterchangeCommonPipelineDataFactoryNode.h
--- a/Engine/Plugins/Interchange/Runtime/Source/FactoryNodes/Public/InterchangeCommonPipelineDataFactoryNode.h
+++ b/Engine/Plugins/Interchange/Runtime/Source/FactoryNodes/Public/InterchangeCommonPipelineDataFactoryNode.h
@@ -36,11 +36,20 @@ public:
 
 	/** Pipelines can set this Bake Meshes setting. Factories use this to identify whether they should apply global transforms to static meshes and skeletal meshes. */
 	UFUNCTION(BlueprintCallable, Category = "Interchange | Node | Common Pipeline Data")
-	bool SetBakeMeshes(const UInterchangeBaseNodeContainer* NodeContainer, const bool& AttributeValue);
+	bool SetBakeMeshes(const bool& AttributeValue);
+
+	/** Return the value of the Bake Meshes setting set by the pipelines. */
+	UFUNCTION(BlueprintCallable, Category = "Interchange | Node | Common Pipeline Data")
+	bool GetBakePivotMeshes(bool& AttributeValue) const;
+
+	/** Pipelines can set this Bake Meshes setting. Factories use this to identify whether they should apply global transforms to static meshes and skeletal meshes. */
+	UFUNCTION(BlueprintCallable, Category = "Interchange | Node | Common Pipeline Data")
+	bool SetBakePivotMeshes(const bool& AttributeValue);
 
 private:
 	UInterchangeCommonPipelineDataFactoryNode() {};
 
 	const UE::Interchange::FAttributeKey Macro_CustomGlobalOffsetTransformKey = UE::Interchange::FAttributeKey(TEXT("GlobalOffsetTransform"));
 	const UE::Interchange::FAttributeKey Macro_CustomBakeMeshesKey = UE::Interchange::FAttributeKey(TEXT("BakeMeshes"));
+	const UE::Interchange::FAttributeKey Macro_CustomBakePivotMeshesKey = UE::Interchange::FAttributeKey(TEXT("BakePivotMeshes"));
 };
\ No newline at end of file
diff --git a/Engine/Plugins/Interchange/Runtime/Source/Import/Private/Mesh/InterchangeMeshHelper.cpp b/Engine/Plugins/Interchange/Runtime/Source/Import/Private/Mesh/InterchangeMeshHelper.cpp
--- a/Engine/Plugins/Interchange/Runtime/Source/Import/Private/Mesh/InterchangeMeshHelper.cpp
+++ b/Engine/Plugins/Interchange/Runtime/Source/Import/Private/Mesh/InterchangeMeshHelper.cpp
@@ -1,6 +1,7 @@
 // Copyright Epic Games, Inc. All Rights Reserved. 
 #include "Mesh/InterchangeMeshHelper.h"
 
+#include "InterchangeSceneNode.h"
 #include "MeshDescription.h"
 #include "StaticMeshOperations.h"
 #include "GenericPlatform/GenericPlatformMisc.h"
@@ -42,4 +43,28 @@ namespace UE::Interchange::Private::MeshHelper
 			RemapPolygonGroup.Add(SourcePolygonGroupID, TargetMatchingID);
 		}
 	}
+
+	void AddSceneNodeGeometricAndPivotToGlobalTransform(FTransform& GlobalTransform, const UInterchangeSceneNode* SceneNode, const bool bBakeMeshes, const bool bBakePivotMeshes)
+	{
+		FTransform SceneNodeGeometricTransform;
+		SceneNode->GetCustomGeometricTransform(SceneNodeGeometricTransform);
+
+		if (!bBakeMeshes)
+		{
+			if (bBakePivotMeshes)
+			{
+				FTransform SceneNodePivotNodeTransform;
+				if (SceneNode->GetCustomPivotNodeTransform(SceneNodePivotNodeTransform))
+				{
+					SceneNodeGeometricTransform = SceneNodePivotNodeTransform * SceneNodeGeometricTransform;
+				}
+			}
+			else
+			{
+				SceneNodeGeometricTransform.SetIdentity();
+			}
+		}
+		GlobalTransform = bBakeMeshes ? SceneNodeGeometricTransform * GlobalTransform : SceneNodeGeometricTransform;
+	}
+
 } //ns UE::Interchange::Private::MeshHelper
\ No newline at end of file
diff --git a/Engine/Plugins/Interchange/Runtime/Source/Import/Private/Mesh/InterchangeStaticMeshFactory.cpp b/Engine/Plugins/Interchange/Runtime/Source/Import/Private/Mesh/InterchangeStaticMeshFactory.cpp
--- a/Engine/Plugins/Interchange/Runtime/Source/Import/Private/Mesh/InterchangeStaticMeshFactory.cpp
+++ b/Engine/Plugins/Interchange/Runtime/Source/Import/Private/Mesh/InterchangeStaticMeshFactory.cpp
@@ -864,10 +864,15 @@ TArray<UInterchangeStaticMeshFactory::FMeshPayload> UInterchangeStaticMeshFactor
 
 	FTransform GlobalOffsetTransform = FTransform::Identity;
 	bool bBakeMeshes = false;
+	bool bBakePivotMeshes = false;
 	if (UInterchangeCommonPipelineDataFactoryNode* CommonPipelineDataFactoryNode = UInterchangeCommonPipelineDataFactoryNode::GetUniqueInstance(Arguments.NodeContainer))
 	{
 		CommonPipelineDataFactoryNode->GetCustomGlobalOffsetTransform(GlobalOffsetTransform);
 		CommonPipelineDataFactoryNode->GetBakeMeshes(bBakeMeshes);
+		if (!bBakeMeshes)
+		{
+			CommonPipelineDataFactoryNode->GetBakePivotMeshes(bBakePivotMeshes);
+		}
 	}
 
 	for (const FString& MeshUid : MeshUids)
@@ -896,6 +901,8 @@ TArray<UInterchangeStaticMeshFactory::FMeshPayload> UInterchangeStaticMeshFactor
 				}
 			}
 
+			UE::Interchange::Private::MeshHelper::AddSceneNodeGeometricAndPivotToGlobalTransform(Payload.Transform, SceneNode, bBakeMeshes, bBakePivotMeshes);
+
 			// And get the mesh node which it references
 			FString MeshDependencyUid;
 			SceneNode->GetCustomAssetInstanceUid(MeshDependencyUid);
@@ -1768,11 +1775,17 @@ bool UInterchangeStaticMeshFactory::ImportSockets(const FImportAssetObjectParams
 	TSet<FName> ImportedSocketNames;
 
 	FTransform GlobalOffsetTransform = FTransform::Identity;
+
 	bool bBakeMeshes = false;
+	bool bBakePivotMeshes = false;
 	if (UInterchangeCommonPipelineDataFactoryNode* CommonPipelineDataFactoryNode = UInterchangeCommonPipelineDataFactoryNode::GetUniqueInstance(Arguments.NodeContainer))
 	{
 		CommonPipelineDataFactoryNode->GetCustomGlobalOffsetTransform(GlobalOffsetTransform);
 		CommonPipelineDataFactoryNode->GetBakeMeshes(bBakeMeshes);
+		if (!bBakeMeshes)
+		{
+			CommonPipelineDataFactoryNode->GetBakePivotMeshes(bBakePivotMeshes);
+		}
 	}
 
 	for (const FString& SocketUid : SocketUids)
@@ -1793,6 +1806,8 @@ bool UInterchangeStaticMeshFactory::ImportSockets(const FImportAssetObjectParams
 				SceneNode->GetCustomGlobalTransform(Arguments.NodeContainer, GlobalOffsetTransform, Transform);
 			}
 
+			UE::Interchange::Private::MeshHelper::AddSceneNodeGeometricAndPivotToGlobalTransform(Transform, SceneNode, bBakeMeshes, bBakePivotMeshes);
+
 			UStaticMeshSocket* Socket = StaticMesh->FindSocket(SocketName);
 			if (!Socket)
 			{
diff --git a/Engine/Plugins/Interchange/Runtime/Source/Import/Public/Mesh/InterchangeMeshHelper.h b/Engine/Plugins/Interchange/Runtime/Source/Import/Public/Mesh/InterchangeMeshHelper.h
--- a/Engine/Plugins/Interchange/Runtime/Source/Import/Public/Mesh/InterchangeMeshHelper.h
+++ b/Engine/Plugins/Interchange/Runtime/Source/Import/Public/Mesh/InterchangeMeshHelper.h
@@ -2,7 +2,12 @@
 #include "MeshDescription.h"
 #include "StaticMeshOperations.h"
 
+class UInterchangeSceneNode;
+
 namespace UE::Interchange::Private::MeshHelper
 {
 	INTERCHANGEIMPORT_API void RemapPolygonGroups(const FMeshDescription& SourceMesh, FMeshDescription& TargetMesh, PolygonGroupMap& RemapPolygonGroup);
+
+	/* return the result of the global transform with the geometric and pivot transform of the scene node. */
+	INTERCHANGEIMPORT_API void AddSceneNodeGeometricAndPivotToGlobalTransform(FTransform& GlobalTransform, const UInterchangeSceneNode* SceneNode, const bool bBakeMeshes, const bool bBakePivotMeshes);
 } //ns UE::Interchange::Private::MeshHelper
\ No newline at end of file
diff --git a/Engine/Plugins/Interchange/Runtime/Source/Nodes/Private/InterchangeSceneNode.cpp b/Engine/Plugins/Interchange/Runtime/Source/Nodes/Private/InterchangeSceneNode.cpp
--- a/Engine/Plugins/Interchange/Runtime/Source/Nodes/Private/InterchangeSceneNode.cpp
+++ b/Engine/Plugins/Interchange/Runtime/Source/Nodes/Private/InterchangeSceneNode.cpp
@@ -258,6 +258,16 @@ bool UInterchangeSceneNode::SetCustomGeometricTransform(const FTransform& Attrib
 	IMPLEMENT_NODE_ATTRIBUTE_SETTER_NODELEGATE(GeometricTransform, FTransform);
 }
 
+bool UInterchangeSceneNode::GetCustomPivotNodeTransform(FTransform& AttributeValue) const
+{
+	IMPLEMENT_NODE_ATTRIBUTE_GETTER(PivotNodeTransform, FTransform);
+}
+
+bool UInterchangeSceneNode::SetCustomPivotNodeTransform(const FTransform& AttributeValue)
+{
+	IMPLEMENT_NODE_ATTRIBUTE_SETTER_NODELEGATE(PivotNodeTransform, FTransform);
+}
+
 bool UInterchangeSceneNode::GetCustomAssetInstanceUid(FString& AttributeValue) const
 {
 	IMPLEMENT_NODE_ATTRIBUTE_GETTER(AssetInstanceUid, FString);
diff --git a/Engine/Plugins/Interchange/Runtime/Source/Nodes/Public/InterchangeSceneNode.h b/Engine/Plugins/Interchange/Runtime/Source/Nodes/Public/InterchangeSceneNode.h
--- a/Engine/Plugins/Interchange/Runtime/Source/Nodes/Public/InterchangeSceneNode.h
+++ b/Engine/Plugins/Interchange/Runtime/Source/Nodes/Public/InterchangeSceneNode.h
@@ -124,6 +124,14 @@ public:
 	UFUNCTION(BlueprintCallable, Category = "Interchange | Node | Scene")
 	bool SetCustomGeometricTransform(const FTransform& AttributeValue);
 
+	/** Get the node pivot geometric offset. */
+	UFUNCTION(BlueprintCallable, Category = "Interchange | Node | Scene")
+	bool GetCustomPivotNodeTransform(FTransform& AttributeValue) const;
+
+	/** Set the node pivot geometric offset. */
+	UFUNCTION(BlueprintCallable, Category = "Interchange | Node | Scene")
+	bool SetCustomPivotNodeTransform(const FTransform& AttributeValue);
+
 	/***********************************************************************************************
 	* Skeleton bind bone API Begin
 	* 
@@ -234,6 +242,9 @@ private:
 	//A scene node can have a transform apply to the mesh it reference.
 	const UE::Interchange::FAttributeKey Macro_CustomGeometricTransformKey = UE::Interchange::FAttributeKey(TEXT("GeometricTransform"));
 
+	//A scene node can have a pivot transform apply to the mesh it reference (use this pivot only if you are not baking the vertices of the mesh).
+	const UE::Interchange::FAttributeKey Macro_CustomPivotNodeTransformKey = UE::Interchange::FAttributeKey(TEXT("PivotNodeTransform"));
+
 	//A scene node can reference an asset. Asset can be Mesh, Light, camera...
 	const UE::Interchange::FAttributeKey Macro_CustomAssetInstanceUidKey = UE::Interchange::FAttributeKey(TEXT("AssetInstanceUid"));
 	
diff --git a/Engine/Plugins/Interchange/Runtime/Source/Parsers/Fbx/Private/FbxScene.cpp b/Engine/Plugins/Interchange/Runtime/Source/Parsers/Fbx/Private/FbxScene.cpp
--- a/Engine/Plugins/Interchange/Runtime/Source/Parsers/Fbx/Private/FbxScene.cpp
+++ b/Engine/Plugins/Interchange/Runtime/Source/Parsers/Fbx/Private/FbxScene.cpp
@@ -29,8 +29,7 @@ namespace UE
 	{
 		namespace Private
 		{
-
-			void FFbxScene::CreateMeshNodeReference(UInterchangeSceneNode* UnrealSceneNode, FbxNodeAttribute* NodeAttribute, UInterchangeBaseNodeContainer& NodeContainer, const FTransform& GeometricTransform)
+			void FFbxScene::CreateMeshNodeReference(UInterchangeSceneNode* UnrealSceneNode, FbxNodeAttribute* NodeAttribute, UInterchangeBaseNodeContainer& NodeContainer, const FTransform& GeometricTransform, const FTransform& PivotNodeTransform)
 			{
 				const UInterchangeMeshNode* MeshNode = nullptr;
 				if (NodeAttribute->GetAttributeType() == FbxNodeAttribute::eMesh)
@@ -56,6 +55,11 @@ namespace UE
 						UnrealSceneNode->SetCustomGeometricTransform(GeometricTransform);
 					}
 
+					if (!PivotNodeTransform.Equals(FTransform::Identity))
+					{
+						UnrealSceneNode->SetCustomPivotNodeTransform(PivotNodeTransform);
+					}
+
 					// @todo: Nothing is using the SceneInstanceUid in the MeshNode. Do we even need to support it?
 					// For the moment an ugly const_cast so we can mutate it (it was fetched from the NodeContainer and is hence const).
 					// Possible solutions:
@@ -304,6 +308,7 @@ namespace UE
 							//For Mesh attribute we add the fbx nodes materials
 							FFbxMaterial FbxMaterial(Parser);
 							FbxMaterial.AddAllNodeMaterials(UnrealNode, Node, NodeContainer);
+							
 							//Get the Geometric offset transform and set it in the mesh node
 							//The geometric offset is not part of the hierarchy transform, it is not inherited
 							FbxAMatrix Geometry;
@@ -315,7 +320,18 @@ namespace UE
 							Geometry.SetR(Rotation);
 							Geometry.SetS(Scaling);
 							FTransform GeometricTransform = GetConvertedTransform(Geometry);
-							CreateMeshNodeReference(UnrealNode, NodeAttribute, NodeContainer, GeometricTransform);
+
+							//Get the pivot geometry offset 
+							FbxAMatrix PivotGeometry;
+							FbxVector4 RotationPivot = Node->GetRotationPivot(FbxNode::eSourcePivot);
+							FbxVector4 FullPivot;
+							FullPivot[0] = -RotationPivot[0];
+							FullPivot[1] = -RotationPivot[1];
+							FullPivot[2] = -RotationPivot[2];
+							PivotGeometry.SetT(FullPivot);
+							FTransform PivotNodeTransform = GetConvertedTransform(PivotGeometry);
+
+							CreateMeshNodeReference(UnrealNode, NodeAttribute, NodeContainer, GeometricTransform, PivotNodeTransform);
 							break;
 						}
 						case FbxNodeAttribute::eLODGroup:
diff --git a/Engine/Plugins/Interchange/Runtime/Source/Parsers/Fbx/Private/FbxScene.h b/Engine/Plugins/Interchange/Runtime/Source/Parsers/Fbx/Private/FbxScene.h
--- a/Engine/Plugins/Interchange/Runtime/Source/Parsers/Fbx/Private/FbxScene.h
+++ b/Engine/Plugins/Interchange/Runtime/Source/Parsers/Fbx/Private/FbxScene.h
@@ -36,7 +36,7 @@ namespace UE
 				UInterchangeSceneNode* CreateTransformNode(UInterchangeBaseNodeContainer& NodeContainer, const FString& NodeName, const FString& NodeUniqueID);
 
 			protected:
-				void CreateMeshNodeReference(UInterchangeSceneNode* UnrealSceneNode, FbxNodeAttribute* NodeAttribute, UInterchangeBaseNodeContainer& NodeContainer, const FTransform& GeometricTransform);
+				void CreateMeshNodeReference(UInterchangeSceneNode* UnrealSceneNode, FbxNodeAttribute* NodeAttribute, UInterchangeBaseNodeContainer& NodeContainer, const FTransform& GeometricTransform, const FTransform& PivotNodeTransform);
 				void CreateCameraNodeReference(UInterchangeSceneNode* UnrealSceneNode, FbxNodeAttribute* NodeAttribute, UInterchangeBaseNodeContainer& NodeContainer);
 				void CreateLightNodeReference(UInterchangeSceneNode* UnrealSceneNode, FbxNodeAttribute* NodeAttribute, UInterchangeBaseNodeContainer& NodeContainer);
 				void AddHierarchyRecursively(UInterchangeSceneNode* UnrealParentNode
diff --git a/Engine/Plugins/Interchange/Runtime/Source/Pipelines/Private/InterchangeGenericAssetsPipeline.cpp b/Engine/Plugins/Interchange/Runtime/Source/Pipelines/Private/InterchangeGenericAssetsPipeline.cpp
--- a/Engine/Plugins/Interchange/Runtime/Source/Pipelines/Private/InterchangeGenericAssetsPipeline.cpp
+++ b/Engine/Plugins/Interchange/Runtime/Source/Pipelines/Private/InterchangeGenericAssetsPipeline.cpp
@@ -1121,8 +1121,8 @@ void UInterchangeGenericAssetsPipeline::ExecutePipeline(UInterchangeBaseNodeCont
 		UInterchangeCommonPipelineDataFactoryNode* CommonPipelineDataFactoryNode = UInterchangeCommonPipelineDataFactoryNode::FindOrCreateUniqueInstance(InBaseNodeContainer);
 		CommonPipelineDataFactoryNode->SetCustomGlobalOffsetTransform(InBaseNodeContainer, ImportOffsetTransform);
 
-		// In case all mesh types are forced to Static/Skeletal we bake the scene instance hierarchy transforms
-		CommonPipelineDataFactoryNode->SetBakeMeshes(InBaseNodeContainer, CommonMeshesProperties->ForceAllMeshAsType != EInterchangeForceMeshType::IFMT_None || CommonMeshesProperties->bBakeMeshes);
+		CommonPipelineDataFactoryNode->SetBakeMeshes(CommonMeshesProperties->bBakeMeshes);
+		CommonPipelineDataFactoryNode->SetBakePivotMeshes(CommonMeshesProperties->bBakePivotMeshes);
 	}
 
 	if (MaterialPipeline)
diff --git a/Engine/Plugins/Interchange/Runtime/Source/Pipelines/Private/InterchangeGenericMeshPipeline.cpp b/Engine/Plugins/Interchange/Runtime/Source/Pipelines/Private/InterchangeGenericMeshPipeline.cpp
--- a/Engine/Plugins/Interchange/Runtime/Source/Pipelines/Private/InterchangeGenericMeshPipeline.cpp
+++ b/Engine/Plugins/Interchange/Runtime/Source/Pipelines/Private/InterchangeGenericMeshPipeline.cpp
@@ -51,6 +51,7 @@ void UInterchangeGenericMeshPipeline::AdjustSettingsForContext(EInterchangePipel
 			CommonMeshesProperties->ForceAllMeshAsType = EInterchangeForceMeshType::IFMT_SkeletalMesh;
 			CommonMeshesProperties->bAutoDetectMeshType = false;
 			CommonMeshesProperties->bBakeMeshes = true;
+			CommonMeshesProperties->bBakePivotMeshes = false;
 			CommonMeshesProperties->bImportLods = false;
 			CommonMeshesProperties->bKeepSectionsSeparate = false;
 			CommonMeshesProperties->VertexColorImportOption = EInterchangeVertexColorImportOption::IVCIO_Ignore;
@@ -227,7 +228,7 @@ UInterchangePipelineMeshesUtilities* UInterchangeGenericMeshPipeline::CreateMesh
 	DataContext.bConvertSkeletalMeshToStaticMesh = (Pipeline->CommonMeshesProperties->ForceAllMeshAsType == EInterchangeForceMeshType::IFMT_StaticMesh);
 	DataContext.bConvertStaticsWithMorphTargetsToSkeletals = Pipeline->CommonSkeletalMeshesAndAnimationsProperties->bConvertStaticsWithMorphTargetsToSkeletals;
 	DataContext.bImportMeshesInBoneHierarchy = Pipeline->CommonSkeletalMeshesAndAnimationsProperties->bImportMeshesInBoneHierarchy;
-	DataContext.bQueryGeometryOnlyIfNoInstance = Pipeline->CommonMeshesProperties->bBakeMeshes;
+	DataContext.bQueryGeometryOnlyIfNoInstance = Pipeline->CommonMeshesProperties->bBakeMeshes || Pipeline->CommonMeshesProperties->bBakePivotMeshes;
 	CreatedPipelineMeshesUtilities->SetContext(DataContext);
 	return CreatedPipelineMeshesUtilities;
 }
diff --git a/Engine/Plugins/Interchange/Runtime/Source/Pipelines/Private/InterchangeGenericStaticMeshPipeline.cpp b/Engine/Plugins/Interchange/Runtime/Source/Pipelines/Private/InterchangeGenericStaticMeshPipeline.cpp
--- a/Engine/Plugins/Interchange/Runtime/Source/Pipelines/Private/InterchangeGenericStaticMeshPipeline.cpp
+++ b/Engine/Plugins/Interchange/Runtime/Source/Pipelines/Private/InterchangeGenericStaticMeshPipeline.cpp
@@ -241,7 +241,7 @@ void UInterchangeGenericMeshPipeline::ExecutePreImportPipelineStaticMesh()
 			// Do not combine static meshes
 
 			bool bFoundMeshes = false;
-			if(CommonMeshesProperties->bBakeMeshes)
+			if(CommonMeshesProperties->bBakeMeshes || CommonMeshesProperties->bBakePivotMeshes)
 			{
 				TArray<FString> MeshUids;
 				PipelineMeshesUtilities->GetAllStaticMeshInstance(MeshUids);
diff --git a/Engine/Plugins/Interchange/Runtime/Source/Pipelines/Public/InterchangeGenericAssetsPipelineSharedSettings.h b/Engine/Plugins/Interchange/Runtime/Source/Pipelines/Public/InterchangeGenericAssetsPipelineSharedSettings.h
--- a/Engine/Plugins/Interchange/Runtime/Source/Pipelines/Public/InterchangeGenericAssetsPipelineSharedSettings.h
+++ b/Engine/Plugins/Interchange/Runtime/Source/Pipelines/Public/InterchangeGenericAssetsPipelineSharedSettings.h
@@ -72,6 +72,10 @@ public:
 	UPROPERTY(EditAnywhere, BlueprintReadWrite, Category = "Common Meshes")
 	bool bBakeMeshes = true;
 	
+	/** If enabled, the inverse node rotation pivot will be apply to the mesh vertices. The pivot from the DCC will then be the origin of the mesh.*/
+	UPROPERTY(EditAnywhere, BlueprintReadWrite, Category = "Common Meshes", meta = (editcondition = "!bBakeMeshes"))
+	bool bBakePivotMeshes = false;
+	
 	/** If checked, sections with matching materials are kept separate and will not get combined. */
 	UPROPERTY(EditAnywhere, BlueprintReadWrite, Category = "Common Meshes")
 	bool bKeepSectionsSeparate = false;
@@ -126,10 +130,6 @@ public:
 		{
 			return true;
 		}
-		if (PropertyChangedEvent.Property->GetName() == GET_MEMBER_NAME_CHECKED(UInterchangeGenericCommonMeshesProperties, bBakeMeshes))
-		{
-			return true;
-		}
 		return Super::IsPropertyChangeNeedRefresh(PropertyChangedEvent);
 	}
 #endif //WITH_EDITOR
