Interchange fix a dead lock if we chain blueprint or python import with the FImportAssetParameters::OnImportDone delegate.

#rnx

diff --git a/Engine/Plugins/Interchange/Runtime/Source/Import/Private/Mesh/InterchangeStaticMeshFactory.cpp b/Engine/Plugins/Interchange/Runtime/Source/Import/Private/Mesh/InterchangeStaticMeshFactory.cpp
--- a/Engine/Plugins/Interchange/Runtime/Source/Import/Private/Mesh/InterchangeStaticMeshFactory.cpp
+++ b/Engine/Plugins/Interchange/Runtime/Source/Import/Private/Mesh/InterchangeStaticMeshFactory.cpp
@@ -149,6 +149,10 @@ UInterchangeFactoryBase::FImportAssetResult UInterchangeStaticMeshFactory::Begin
 {
 	TRACE_CPUPROFILER_EVENT_SCOPE(UInterchangeStaticMeshFactory::BeginImportAsset_GameThread);
 
+	//We must ensure we use the same settings until the import is finish, FApp::IsGame() can return a different
+	//value during an asynchronous import
+	ImportAssetObjectData.bIsAppGame = FApp::IsGame();
+
 	FImportAssetResult ImportAssetResult;
 	UStaticMesh* StaticMesh = nullptr;
 	if (!Arguments.AssetNode || !Arguments.AssetNode->GetObjectClass()->IsChildOf(GetFactoryClass()))
@@ -209,7 +213,7 @@ UInterchangeFactoryBase::FImportAssetResult UInterchangeStaticMeshFactory::Begin
 	StaticMesh->CreateBodySetup();
 	
 #if WITH_EDITOR
-	if (!FApp::IsGame())
+	if (!ImportAssetObjectData.bIsAppGame)
 	{
 		StaticMesh->PreEditChange(nullptr);
 	}
@@ -529,7 +533,7 @@ UInterchangeFactoryBase::FImportAssetResult UInterchangeStaticMeshFactory::EndIm
 		return ImportAssetResult;
 	}
 
-	if (FApp::IsGame())
+	if (ImportAssetObjectData.bIsAppGame)
 	{
 		if (!Arguments.ReimportObject)
 		{
@@ -728,7 +732,7 @@ UInterchangeFactoryBase::FImportAssetResult UInterchangeStaticMeshFactory::EndIm
 void UInterchangeStaticMeshFactory::CommitMeshDescriptions(UStaticMesh& StaticMesh)
 {
 #if WITH_EDITOR
-	if (FApp::IsGame())
+	if (ImportAssetObjectData.bIsAppGame)
 	{
 		return;
 	}
@@ -751,7 +755,7 @@ void UInterchangeStaticMeshFactory::CommitMeshDescriptions(UStaticMesh& StaticMe
 
 void UInterchangeStaticMeshFactory::BuildFromMeshDescriptions(UStaticMesh& StaticMesh)
 {
-	if (!FApp::IsGame())
+	if (!ImportAssetObjectData.bIsAppGame)
 	{
 		return;
 	}
diff --git a/Engine/Plugins/Interchange/Runtime/Source/Import/Public/Mesh/InterchangeStaticMeshFactory.h b/Engine/Plugins/Interchange/Runtime/Source/Import/Public/Mesh/InterchangeStaticMeshFactory.h
--- a/Engine/Plugins/Interchange/Runtime/Source/Import/Public/Mesh/InterchangeStaticMeshFactory.h
+++ b/Engine/Plugins/Interchange/Runtime/Source/Import/Public/Mesh/InterchangeStaticMeshFactory.h
@@ -72,6 +72,7 @@ private:
 #endif
 	struct FImportAssetObjectData
 	{
+		bool bIsAppGame = false;
 		bool bImportedCustomCollision = false;
 		bool bImportCollision = false;
 		TArray<FMeshDescription> LodMeshDescriptions;
diff --git a/Engine/Source/Runtime/Interchange/Engine/Private/InterchangeManager.cpp b/Engine/Source/Runtime/Interchange/Engine/Private/InterchangeManager.cpp
--- a/Engine/Source/Runtime/Interchange/Engine/Private/InterchangeManager.cpp
+++ b/Engine/Source/Runtime/Interchange/Engine/Private/InterchangeManager.cpp
@@ -1477,13 +1477,20 @@ void UInterchangeManager::StartQueuedTasks(bool bCancelAllTasks /*= false*/)
 	}
 }
 
-bool UInterchangeManager::ImportAsset(const FString& ContentPath, const UInterchangeSourceData* SourceData, const FImportAssetParameters& ImportAssetParameters)
+bool UInterchangeManager::ImportAsset(const FString& ContentPath, const UInterchangeSourceData* SourceData, const FImportAssetParameters& ImportAssetParameters, TArray<UObject*>& OutImportedObjects)
 {
 	UE::Interchange::FAssetImportResultRef InterchangeResult = ImportAssetAsync(ContentPath, SourceData, ImportAssetParameters);
 	InterchangeResult->WaitUntilDone();
+	OutImportedObjects = InterchangeResult->GetImportedObjects();
 	return InterchangeResult->IsValid();
 }
 
+bool UInterchangeManager::ImportAsset(const FString& ContentPath, const UInterchangeSourceData* SourceData, const FImportAssetParameters& ImportAssetParameters)
+{
+	TArray<UObject*> ImportedObjects;
+	return ImportAsset(ContentPath, SourceData, ImportAssetParameters, ImportedObjects);
+}
+
 UE::Interchange::FAssetImportResultRef UInterchangeManager::ImportAssetAsync(const FString& ContentPath, const UInterchangeSourceData* SourceData, const FImportAssetParameters& ImportAssetParameters)
 {
 	return ImportInternal(ContentPath, SourceData, ImportAssetParameters, UE::Interchange::EImportType::ImportType_Asset).Get<0>();
diff --git a/Engine/Source/Runtime/Interchange/Engine/Public/InterchangeManager.h b/Engine/Source/Runtime/Interchange/Engine/Public/InterchangeManager.h
--- a/Engine/Source/Runtime/Interchange/Engine/Public/InterchangeManager.h
+++ b/Engine/Source/Runtime/Interchange/Engine/Public/InterchangeManager.h
@@ -547,7 +547,9 @@ public:
 	 * @return true if the import succeeds, or false otherwise.
 	 */
 	UFUNCTION(BlueprintCallable, Category = "Interchange | Import Manager")
+	INTERCHANGEENGINE_API bool ImportAsset(const FString& ContentPath, const UInterchangeSourceData* SourceData, const FImportAssetParameters& ImportAssetParameters, TArray<UObject*>& OutImportedObjects);
 	INTERCHANGEENGINE_API bool ImportAsset(const FString& ContentPath, const UInterchangeSourceData* SourceData, const FImportAssetParameters& ImportAssetParameters);
+
 	INTERCHANGEENGINE_API UE::Interchange::FAssetImportResultRef ImportAssetAsync(const FString& ContentPath, const UInterchangeSourceData* SourceData, const FImportAssetParameters& ImportAssetParameters);
 
 	/**
