[Interchange] Custom attributes are not imported with level import fix.

#rnx

diff --git a/Engine/Plugins/Interchange/Runtime/Source/Pipelines/Private/InterchangeGenericStaticMeshPipeline.cpp b/Engine/Plugins/Interchange/Runtime/Source/Pipelines/Private/InterchangeGenericStaticMeshPipeline.cpp
--- a/Engine/Plugins/Interchange/Runtime/Source/Pipelines/Private/InterchangeGenericStaticMeshPipeline.cpp
+++ b/Engine/Plugins/Interchange/Runtime/Source/Pipelines/Private/InterchangeGenericStaticMeshPipeline.cpp
@@ -331,7 +331,7 @@ void UInterchangeGenericMeshPipeline::ExecutePreImportPipelineStaticMesh()
 							}
 						}
 
-						UInterchangeStaticMeshFactoryNode* StaticMeshFactoryNode = CreateStaticMeshFactoryNode(MeshUidsPerLodIndex);
+						UInterchangeStaticMeshFactoryNode* StaticMeshFactoryNode = CreateStaticMeshFactoryNode(MeshUidsPerLodIndex, MeshGeometry.ReferencingMeshInstanceUids);
 						StaticMeshFactoryNodes.Add(StaticMeshFactoryNode);
 					}
 				}
@@ -396,7 +396,7 @@ bool UInterchangeGenericMeshPipeline::MakeMeshFactoryNodeUidAndDisplayLabel(cons
 	return false;
 }
 
-UInterchangeStaticMeshFactoryNode* UInterchangeGenericMeshPipeline::CreateStaticMeshFactoryNode(const TMap<int32, TArray<FString>>& MeshUidsPerLodIndex)
+UInterchangeStaticMeshFactoryNode* UInterchangeGenericMeshPipeline::CreateStaticMeshFactoryNode(const TMap<int32, TArray<FString>>& MeshUidsPerLodIndex, const TArray<FString>& ReferencingMeshInstanceUids)
 {
 	check(CommonMeshesProperties.IsValid());
 	if (MeshUidsPerLodIndex.Num() == 0)
@@ -475,6 +475,19 @@ UInterchangeStaticMeshFactoryNode* UInterchangeGenericMeshPipeline::CreateStatic
 	StaticMeshFactoryNode->SetCustomMaxLumenMeshCards(MaxLumenMeshCards);
 	StaticMeshFactoryNode->SetCustomBuildNanite(bBuildNanite);
 
+	if (ReferencingMeshInstanceUids.Num() > 0)
+	{
+		constexpr bool bAddSourceNodeName = true;
+
+		for (const FString& ReferencingMeshInstanceUid : ReferencingMeshInstanceUids)
+		{
+			if (const UInterchangeSceneNode* SceneNode = Cast<UInterchangeSceneNode>(BaseNodeContainer->GetNode(ReferencingMeshInstanceUid)))
+			{
+				UInterchangeUserDefinedAttributesAPI::DuplicateAllUserDefinedAttribute(SceneNode, StaticMeshFactoryNode, bAddSourceNodeName);
+			}
+		}
+	}
+
 	return StaticMeshFactoryNode;
 }
 
diff --git a/Engine/Plugins/Interchange/Runtime/Source/Pipelines/Public/InterchangeGenericMeshPipeline.h b/Engine/Plugins/Interchange/Runtime/Source/Pipelines/Public/InterchangeGenericMeshPipeline.h
--- a/Engine/Plugins/Interchange/Runtime/Source/Pipelines/Public/InterchangeGenericMeshPipeline.h
+++ b/Engine/Plugins/Interchange/Runtime/Source/Pipelines/Public/InterchangeGenericMeshPipeline.h
@@ -314,7 +314,7 @@ private:
 	 * This function can create a UInterchangeStaticMeshFactoryNode
 	 * @param MeshUidsPerLodIndex - The MeshUids can represent a SceneNode pointing on a MeshNode or directly a MeshNode
 	 */
-	UInterchangeStaticMeshFactoryNode* CreateStaticMeshFactoryNode(const TMap<int32, TArray<FString>>& MeshUidsPerLodIndex);
+	UInterchangeStaticMeshFactoryNode* CreateStaticMeshFactoryNode(const TMap<int32, TArray<FString>>& MeshUidsPerLodIndex, const TArray<FString>& ReferencingMeshInstanceUids = TArray<FString>());
 
 	/** This function can create a UInterchangeStaticMeshLodDataNode which represents the LOD data needed by the factory to create a lod mesh */
 	UInterchangeStaticMeshLodDataNode* CreateStaticMeshLodDataNode(const FString& NodeName, const FString& NodeUniqueID);
