[On behalf of danny.kabrane] Fixed input names of material factory nodes in FBX and OBJ translators
Reason: Otherwise, the FBX/OBJ materials imported as material instances would fail

#rnx

diff --git a/Engine/Plugins/Interchange/Runtime/Source/Import/Private/Mesh/InterchangeOBJTranslator.cpp b/Engine/Plugins/Interchange/Runtime/Source/Import/Private/Mesh/InterchangeOBJTranslator.cpp
--- a/Engine/Plugins/Interchange/Runtime/Source/Import/Private/Mesh/InterchangeOBJTranslator.cpp
+++ b/Engine/Plugins/Interchange/Runtime/Source/Import/Private/Mesh/InterchangeOBJTranslator.cpp
@@ -1107,8 +1107,28 @@ namespace ObjTranslatorUtils
 
 		if (TextureSampleShaderNode)
 		{
-			TextureSampleShaderNode->SetDisplayLabel(InputName);
-			UInterchangeShaderPortsAPI::ConnectDefaultOuputToInput(ShaderGraphNode, InputName, TextureSampleShaderNode->GetUniqueID());
+			using namespace UE::Interchange::Materials::Standard::Nodes;
+
+			const FString MultiplierNodeName = InputName + TEXT("Multiply");
+			UInterchangeShaderNode* MultiplierNode = UInterchangeShaderNode::Create(&BaseNodeContainer, MultiplierNodeName, ShaderGraphNode->GetUniqueID());
+			MultiplierNode->SetCustomShaderType(Multiply::Name.ToString());
+
+			UInterchangeShaderPortsAPI::ConnectDefaultOuputToInput(MultiplierNode, Multiply::Inputs::A.ToString(), TextureSampleShaderNode->GetUniqueID());
+
+			const FString WeightNodeName = InputName + TEXT("MapWeight");
+			UInterchangeShaderNode* WeightNode = UInterchangeShaderNode::Create(&BaseNodeContainer, WeightNodeName, MultiplierNode->GetUniqueID());
+			WeightNode->SetCustomShaderType(ScalarParameter::Name.ToString());
+
+			const float WeightValue = 1.0f;
+			WeightNode->AddFloatAttribute(UInterchangeShaderPortsAPI::MakeInputParameterKey(ScalarParameter::Attributes::DefaultValue.ToString()), WeightValue);
+
+			UInterchangeShaderPortsAPI::ConnectDefaultOuputToInput(MultiplierNode, Multiply::Inputs::B.ToString(), WeightNode->GetUniqueID());
+
+			UInterchangeShaderPortsAPI::ConnectDefaultOuputToInput(MultiplierNode, Multiply::Inputs::A.ToString(), TextureSampleShaderNode->GetUniqueID());
+
+			const FString InputNameLabel = InputName + TEXT("Map");
+			TextureSampleShaderNode->SetDisplayLabel(InputNameLabel);
+			UInterchangeShaderPortsAPI::ConnectDefaultOuputToInput(ShaderGraphNode, InputName, MultiplierNode->GetUniqueID());
 		}
 		else if (IsColorInitialized(Color))
 		{
@@ -1128,19 +1148,24 @@ namespace ObjTranslatorUtils
 			return nullptr;
 		}
 
-		if (IsScalarInitialized(Weight))
-		{
-			using namespace UE::Interchange::Materials::Standard::Nodes;
+		using namespace UE::Interchange::Materials::Standard::Nodes;
 
-			const FString MultiplierNodeName = MapName + TEXT("Multiply");
-			UInterchangeShaderNode* MultiplierNode = UInterchangeShaderNode::Create(&BaseNodeContainer, MultiplierNodeName, ShaderGraphNode->GetUniqueID());
-			MultiplierNode->SetCustomShaderType(Multiply::Name.ToString());
+		const FString MultiplierNodeName = MapName + TEXT("Multiply");
+		UInterchangeShaderNode* MultiplierNode = UInterchangeShaderNode::Create(&BaseNodeContainer, MultiplierNodeName, ShaderGraphNode->GetUniqueID());
+		MultiplierNode->SetCustomShaderType(Multiply::Name.ToString());
 
-			UInterchangeShaderPortsAPI::ConnectDefaultOuputToInput(MultiplierNode, Multiply::Inputs::A.ToString(), TextureSampleShaderNode->GetUniqueID());
-			MultiplierNode->AddFloatAttribute(UInterchangeShaderPortsAPI::MakeInputValueKey( Multiply::Inputs::B.ToString() ), Weight);
-			return MultiplierNode;
-		}
-		return TextureSampleShaderNode;
+		UInterchangeShaderPortsAPI::ConnectDefaultOuputToInput(MultiplierNode, Multiply::Inputs::A.ToString(), TextureSampleShaderNode->GetUniqueID());
+
+		const FString WeightNodeName = MapName + TEXT("MapWeight");
+		UInterchangeShaderNode* WeightNode = UInterchangeShaderNode::Create(&BaseNodeContainer, WeightNodeName, MultiplierNode->GetUniqueID());
+		WeightNode->SetCustomShaderType(ScalarParameter::Name.ToString());
+
+		const float WeightValue = IsScalarInitialized(Weight) ? Weight : 1.0f;
+		WeightNode->AddFloatAttribute(UInterchangeShaderPortsAPI::MakeInputParameterKey(ScalarParameter::Attributes::DefaultValue.ToString()), WeightValue);
+
+		UInterchangeShaderPortsAPI::ConnectDefaultOuputToInput(MultiplierNode, Multiply::Inputs::B.ToString(), WeightNode->GetUniqueID());
+
+		return MultiplierNode;
 	}
 
 	static bool AddTexturedWeightedInput(UInterchangeBaseNodeContainer& BaseNodeContainer, UInterchangeShaderGraphNode* ShaderGraphNode, const FString& InputName, const float& Weight, UInterchangeShaderNode* TextureSampleShader) 
diff --git a/Engine/Plugins/Interchange/Runtime/Source/Parsers/Fbx/Private/FbxMaterial.cpp b/Engine/Plugins/Interchange/Runtime/Source/Parsers/Fbx/Private/FbxMaterial.cpp
--- a/Engine/Plugins/Interchange/Runtime/Source/Parsers/Fbx/Private/FbxMaterial.cpp
+++ b/Engine/Plugins/Interchange/Runtime/Source/Parsers/Fbx/Private/FbxMaterial.cpp
@@ -55,16 +55,16 @@ namespace UE
 					return TextureNode;
 				}
 
-				UInterchangeTexture2DNode* NewTextureNode = UInterchangeTexture2DNode::Create(&NodeContainer, TextureName);
+				UInterchangeTexture2DNode* NewTextureNode = UInterchangeTexture2DNode::Create(&NodeContainer, TextureNodeID);
+				NewTextureNode->SetDisplayLabel(TextureName);
 
 				//All texture translator expect a file as the payload key
-				
 				NewTextureNode->SetPayLoadKey(NormalizeFilePath);
 
 				return NewTextureNode;
 			}
 
-			const UInterchangeShaderNode* FFbxMaterial::CreateTextureSampler(FbxFileTexture* FbxTexture, UInterchangeBaseNodeContainer& NodeContainer, const FString& ShaderUniqueID)
+			const UInterchangeShaderNode* FFbxMaterial::CreateTextureSampler(FbxFileTexture* FbxTexture, UInterchangeBaseNodeContainer& NodeContainer, const FString& ShaderUniqueID, const FString& InputName)
 			{
 				using namespace Materials::Standard::Nodes;
 
@@ -74,8 +74,7 @@ namespace UE
 				}
 
 				const FString TextureFilename = FbxTexture ? UTF8_TO_TCHAR(FbxTexture->GetFileName()) : TEXT("");
-				const FString TextureName = FPaths::GetBaseFilename(TextureFilename);
-				const FString NodeName = TEXT("Sampler_") + TextureName;
+				const FString NodeName = TEXT("Sampler_") + InputName;
 
 				// Return already created node if applicable.
 				const FString SamplerNodeUid = UInterchangeShaderNode::MakeNodeUid(NodeName, ShaderUniqueID);
@@ -85,7 +84,7 @@ namespace UE
 				}
 
 				UInterchangeShaderNode* TextureSampleShader = UInterchangeShaderNode::Create(&NodeContainer, NodeName, ShaderUniqueID);
-				TextureSampleShader->SetDisplayLabel(TextureName);
+				TextureSampleShader->SetDisplayLabel(InputName);
 				TextureSampleShader->SetCustomShaderType(TextureSample::Name.ToString());
 
 				// Return incomplete texture sampler if texture file does not exist
@@ -108,7 +107,7 @@ namespace UE
 
 				if (!FMath::IsNearlyEqual(FbxTexture->GetScaleU(), 1.0) || !FMath::IsNearlyEqual(FbxTexture->GetScaleV(), 1.0))
 				{
-					UInterchangeShaderNode* TextureCoordinateShader = UInterchangeShaderNode::Create(&NodeContainer, TextureName + TEXT("_Coordinate"), ShaderUniqueID);
+					UInterchangeShaderNode* TextureCoordinateShader = UInterchangeShaderNode::Create(&NodeContainer, InputName + TEXT("_Coordinate"), ShaderUniqueID);
 					TextureCoordinateShader->SetCustomShaderType(TextureCoordinate::Name.ToString());
 
 					TextureCoordinateShader->AddFloatAttribute(UInterchangeShaderPortsAPI::MakeInputValueKey(TextureCoordinate::Inputs::UTiling.ToString()), (float)FbxTexture->GetScaleU());
@@ -121,7 +120,7 @@ namespace UE
 			}
 
 			bool FFbxMaterial::ConvertPropertyToShaderNode(UInterchangeBaseNodeContainer& NodeContainer, UInterchangeShaderGraphNode* ShaderGraphNode, FbxProperty& Property, float Factor, FName InputName,
-				TVariant<FLinearColor, float> DefaultValue, bool bInverse)
+														   const TVariant<FLinearColor, float>& DefaultValue, bool bInverse)
 			{
 				using namespace Materials::Standard::Nodes;
 
@@ -172,7 +171,6 @@ namespace UE
 					InputToConnectTo = OneMinus::Inputs::Input.ToString();
 				}
 
-				if (!FMath::IsNearlyEqual(Factor, 1.f))
 				{
 					FString LerpNodeName = InputName.ToString() + TEXT("Lerp");
 					UInterchangeShaderNode* LerpNode = UInterchangeShaderNode::Create(&NodeContainer, LerpNodeName, ShaderGraphNode->GetUniqueID());
@@ -186,9 +184,15 @@ namespace UE
 					{
 						LerpNode->AddLinearColorAttribute(UInterchangeShaderPortsAPI::MakeInputValueKey(Lerp::Inputs::B.ToString()), DefaultValue.Get<FLinearColor>());
 					}
+					
+					const FString WeightNodeName = InputName.ToString() + TEXT("MapWeight");
+					UInterchangeShaderNode* WeightNode = UInterchangeShaderNode::Create(&NodeContainer, WeightNodeName, LerpNode->GetUniqueID());
+					WeightNode->SetCustomShaderType(ScalarParameter::Name.ToString());
 
 					const float InverseFactor = 1.f - Factor; // We lerp from A to B and prefer to put the strongest input in A so we need to flip the lerp factor
-					LerpNode->AddFloatAttribute(UInterchangeShaderPortsAPI::MakeInputValueKey(Lerp::Inputs::Factor.ToString()), InverseFactor);
+					WeightNode->AddFloatAttribute(UInterchangeShaderPortsAPI::MakeInputParameterKey(ScalarParameter::Attributes::DefaultValue.ToString()), InverseFactor);
+
+					UInterchangeShaderPortsAPI::ConnectDefaultOuputToInput(LerpNode, Lerp::Inputs::Factor.ToString() , WeightNode->GetUniqueID());
 
 					UInterchangeShaderPortsAPI::ConnectDefaultOuputToInput(NodeToConnectTo, InputToConnectTo, LerpNode->GetUniqueID());
 
@@ -198,7 +202,7 @@ namespace UE
 
 				// Handles max one texture per property.
 				FbxFileTexture* FbxTexture = Property.GetSrcObject<FbxFileTexture>(0);
-				if (const UInterchangeShaderNode* TextureSampleShader = CreateTextureSampler(FbxTexture, NodeContainer, ShaderGraphNode->GetUniqueID()))
+				if (const UInterchangeShaderNode* TextureSampleShader = CreateTextureSampler(FbxTexture, NodeContainer, ShaderGraphNode->GetUniqueID(), InputName.ToString() + TEXT("Map")))
 				{
 					UInterchangeShaderPortsAPI::ConnectDefaultOuputToInput(NodeToConnectTo, InputToConnectTo, TextureSampleShader->GetUniqueID());
 				}
@@ -234,7 +238,7 @@ namespace UE
 				if (MaterialProperty.GetSrcObjectCount<FBXSDK_NAMESPACE::FbxTexture>() > 0)
 				{
 					FbxFileTexture* FbxTexture = MaterialProperty.GetSrcObject<FbxFileTexture>(0);
-					if (const UInterchangeShaderNode* TextureSampleShader = CreateTextureSampler(FbxTexture, NodeContainer, ShaderGraphNode->GetUniqueID()))
+					if (const UInterchangeShaderNode* TextureSampleShader = CreateTextureSampler(FbxTexture, NodeContainer, ShaderGraphNode->GetUniqueID(), TEXT("ShininessMap")))
 					{
 						FString MultiplyNodeName = Phong::Parameters::Shininess.ToString() + TEXT("_Multiply");
 						UInterchangeShaderNode* MultiplyNode = UInterchangeShaderNode::Create(&NodeContainer, MultiplyNodeName, ShaderGraphNode->GetUniqueID());
@@ -370,7 +374,8 @@ namespace UE
 										const FString TextureName = FPaths::GetBaseFilename(TexturePath);
 
 										// NormalFromHeightmap needs TextureObject(not just a sample as it takes multiple samples from it)
-										UInterchangeShaderNode* TextureObjectNode = UInterchangeShaderNode::Create(&NodeContainer, TextureName, ShaderGraphNode->GetUniqueID());
+										UInterchangeShaderNode* TextureObjectNode = UInterchangeShaderNode::Create(&NodeContainer, TEXT("NormalMap"), ShaderGraphNode->GetUniqueID());
+										TextureObjectNode->SetDisplayLabel(TEXT("NormalMap"));
 										TextureObjectNode->SetCustomShaderType(TextureObject::Name.ToString());
 										TextureObjectNode->AddStringAttribute(UInterchangeShaderPortsAPI::MakeInputValueKey(TextureObject::Inputs::Texture.ToString()), TextureNode->GetUniqueID());
 
diff --git a/Engine/Plugins/Interchange/Runtime/Source/Parsers/Fbx/Private/FbxMaterial.h b/Engine/Plugins/Interchange/Runtime/Source/Parsers/Fbx/Private/FbxMaterial.h
--- a/Engine/Plugins/Interchange/Runtime/Source/Parsers/Fbx/Private/FbxMaterial.h
+++ b/Engine/Plugins/Interchange/Runtime/Source/Parsers/Fbx/Private/FbxMaterial.h
@@ -55,12 +55,12 @@ namespace UE
 				UInterchangeShaderGraphNode* CreateShaderGraphNode(UInterchangeBaseNodeContainer& NodeContainer, const FString& NodeUID, const FString& NodeName);
 				const UInterchangeTexture2DNode* CreateTexture2DNode(UInterchangeBaseNodeContainer& NodeContainer, const FString& TextureFilePath);
 				bool ConvertPropertyToShaderNode(UInterchangeBaseNodeContainer& NodeContainer, UInterchangeShaderGraphNode* ShaderGraphNode,
-					FbxProperty& Property, float Factor, FName PropertyName, TVariant<FLinearColor, float> DefaultValue, bool bInverse = false);
+												 FbxProperty& Property, float Factor, FName PropertyName, const TVariant<FLinearColor, float>& DefaultValue, bool bInverse = false);
 
 			private:
 				const UInterchangeShaderGraphNode* AddShaderGraphNode(FbxSurfaceMaterial* SurfaceMaterial, UInterchangeBaseNodeContainer& NodeContainer);
 				void ConvertShininessToShaderNode(FbxSurfaceMaterial& SurfaceMaterial, UInterchangeBaseNodeContainer& NodeContainer, UInterchangeShaderGraphNode* ShaderGraphNode);
-				const UInterchangeShaderNode* CreateTextureSampler(FbxFileTexture* FbxTexture, UInterchangeBaseNodeContainer& NodeContainer, const FString& ShaderUniqueID);
+				const UInterchangeShaderNode* CreateTextureSampler(FbxFileTexture* FbxTexture, UInterchangeBaseNodeContainer& NodeContainer, const FString& ShaderUniqueID, const FString& InputName);
 
 				FFbxParser& Parser;
 			};
