[Interchange] Import custom properties expansions (Meshes/Animations/Materials)

#rnx

diff --git a/Engine/Plugins/Interchange/Runtime/Source/Parsers/Fbx/Private/FbxHelper.cpp b/Engine/Plugins/Interchange/Runtime/Source/Parsers/Fbx/Private/FbxHelper.cpp
--- a/Engine/Plugins/Interchange/Runtime/Source/Parsers/Fbx/Private/FbxHelper.cpp
+++ b/Engine/Plugins/Interchange/Runtime/Source/Parsers/Fbx/Private/FbxHelper.cpp
@@ -3,9 +3,12 @@
 #include "FbxHelper.h"
 
 #include "CoreMinimal.h"
+#include "FbxAPI.h"
 #include "FbxConvert.h"
 #include "FbxInclude.h"
 #include "InterchangeHelper.h"
+#include "Nodes/InterchangeBaseNode.h"
+#include "Nodes/InterchangeUserDefinedAttribute.h"
 
 #define GeneratedLODNameSuffix "_GeneratedLOD_"
 
@@ -182,6 +185,142 @@ namespace UE
 				FormatArguments.Add(TEXT("UniqueID"), UniqueID);
 				return FString::Format(TEXT("{UniqueID}"), FormatArguments);
 			}
+
+			void ProcessCustomAttribute(FFbxParser& Parser, UInterchangeBaseNode* UnrealNode, FbxProperty Property, const TOptional<FString>& PayloadKey)
+			{
+				FString PropertyName = Parser.GetFbxHelper()->GetFbxPropertyName(Property);
+
+				switch (Property.GetPropertyDataType().GetType())
+				{
+					case EFbxType::eFbxBool:
+						{
+							bool PropertyValue = Property.Get<bool>();
+							UInterchangeUserDefinedAttributesAPI::CreateUserDefinedAttribute(UnrealNode, PropertyName, PropertyValue, PayloadKey);
+						}
+						break;
+					case EFbxType::eFbxChar:
+						{
+							int8 PropertyValue = Property.Get<int8>();
+							UInterchangeUserDefinedAttributesAPI::CreateUserDefinedAttribute(UnrealNode, PropertyName, PropertyValue, PayloadKey);
+						}
+						break;
+					case EFbxType::eFbxUChar:
+						{
+							uint8 PropertyValue = Property.Get<uint8>();
+							UInterchangeUserDefinedAttributesAPI::CreateUserDefinedAttribute(UnrealNode, PropertyName, PropertyValue, PayloadKey);
+						}
+						break;
+					case EFbxType::eFbxShort:
+						{
+							int16 PropertyValue = Property.Get<int16>();
+							UInterchangeUserDefinedAttributesAPI::CreateUserDefinedAttribute(UnrealNode, PropertyName, PropertyValue, PayloadKey);
+						}
+						break;
+					case EFbxType::eFbxUShort:
+						{
+							uint16 PropertyValue = Property.Get<uint16>();
+							UInterchangeUserDefinedAttributesAPI::CreateUserDefinedAttribute(UnrealNode, PropertyName, PropertyValue, PayloadKey);
+						}
+						break;
+					case EFbxType::eFbxInt:
+						{
+							int32 PropertyValue = Property.Get<int32>();
+							UInterchangeUserDefinedAttributesAPI::CreateUserDefinedAttribute(UnrealNode, PropertyName, PropertyValue, PayloadKey);
+						}
+						break;
+					case EFbxType::eFbxUInt:
+						{
+							uint32 PropertyValue = Property.Get<uint32>();
+							UInterchangeUserDefinedAttributesAPI::CreateUserDefinedAttribute(UnrealNode, PropertyName, PropertyValue, PayloadKey);
+						}
+						break;
+					case EFbxType::eFbxLongLong:
+						{
+							int64 PropertyValue = Property.Get<int64>();
+							UInterchangeUserDefinedAttributesAPI::CreateUserDefinedAttribute(UnrealNode, PropertyName, PropertyValue, PayloadKey);
+						}
+						break;
+					case EFbxType::eFbxULongLong:
+						{
+							uint64 PropertyValue = Property.Get<uint64>();
+							UInterchangeUserDefinedAttributesAPI::CreateUserDefinedAttribute(UnrealNode, PropertyName, PropertyValue, PayloadKey);
+						}
+						break;
+					case EFbxType::eFbxHalfFloat:
+						{
+							FbxHalfFloat HalfFloat = Property.Get<FbxHalfFloat>();
+							FFloat16 PropertyValue = FFloat16(HalfFloat.value());
+							UInterchangeUserDefinedAttributesAPI::CreateUserDefinedAttribute(UnrealNode, PropertyName, PropertyValue, PayloadKey);
+						}
+						break;
+					case EFbxType::eFbxFloat:
+						{
+							float PropertyValue = Property.Get<float>();
+							UInterchangeUserDefinedAttributesAPI::CreateUserDefinedAttribute(UnrealNode, PropertyName, PropertyValue, PayloadKey);
+						}
+						break;
+					case EFbxType::eFbxDouble:
+						{
+							double PropertyValue = Property.Get<double>();
+							UInterchangeUserDefinedAttributesAPI::CreateUserDefinedAttribute(UnrealNode, PropertyName, PropertyValue, PayloadKey);
+						}
+						break;
+					case EFbxType::eFbxDouble2:
+						{
+							FbxDouble2 Vec = Property.Get<FbxDouble2>();
+							FVector2D PropertyValue = FVector2D(Vec[0], Vec[1]);
+							UInterchangeUserDefinedAttributesAPI::CreateUserDefinedAttribute(UnrealNode, PropertyName, PropertyValue, PayloadKey);
+						}
+						break;
+					case EFbxType::eFbxDouble3:
+						{
+							FbxDouble3 Vec = Property.Get<FbxDouble3>();
+							FVector3d PropertyValue = FVector3d(Vec[0], Vec[1], Vec[2]);
+							UInterchangeUserDefinedAttributesAPI::CreateUserDefinedAttribute(UnrealNode, PropertyName, PropertyValue, PayloadKey);
+						}
+						break;
+					case EFbxType::eFbxDouble4:
+						{
+							FbxDouble4 Vec = Property.Get<FbxDouble4>();
+							FVector4d PropertyValue = FVector4d(Vec[0], Vec[1], Vec[2], Vec[3]);
+							UInterchangeUserDefinedAttributesAPI::CreateUserDefinedAttribute(UnrealNode, PropertyName, PropertyValue, PayloadKey);
+						}
+						break;
+					case EFbxType::eFbxEnum:
+						{
+							//Convert enum to uint8
+							FbxEnum EnumValue = Property.Get<FbxEnum>();
+							uint8 PropertyValue = static_cast<uint8>(EnumValue);
+							UInterchangeUserDefinedAttributesAPI::CreateUserDefinedAttribute(UnrealNode, PropertyName, PropertyValue, PayloadKey);
+						}
+						break;
+					case EFbxType::eFbxString:
+						{
+							FbxString StringValue = Property.Get<FbxString>();
+							FString PropertyValue = FFbxConvert::MakeString(StringValue.Buffer());
+							UInterchangeUserDefinedAttributesAPI::CreateUserDefinedAttribute(UnrealNode, PropertyName, PropertyValue, PayloadKey);
+						}
+						break;
+				}
+			}
+
+			void ProcessCustomAttributes(FFbxParser& Parser, FbxObject* Object, UInterchangeBaseNode* UnrealNode)
+			{
+				FbxProperty Property = Object->GetFirstProperty();
+
+				//Add all custom Attributes for the node
+				while (Property.IsValid())
+				{
+					EFbxType PropertyType = Property.GetPropertyDataType().GetType();
+					if (Property.GetFlag(FbxPropertyFlags::eUserDefined))
+					{
+						ProcessCustomAttribute(Parser, UnrealNode, Property);
+					}
+
+					//Inspect next node property
+					Property = Object->GetNextProperty(Property);
+				}
+			}
 		}//ns Private
 	}//ns Interchange
 }//ns UE
diff --git a/Engine/Plugins/Interchange/Runtime/Source/Parsers/Fbx/Private/FbxHelper.h b/Engine/Plugins/Interchange/Runtime/Source/Parsers/Fbx/Private/FbxHelper.h
--- a/Engine/Plugins/Interchange/Runtime/Source/Parsers/Fbx/Private/FbxHelper.h
+++ b/Engine/Plugins/Interchange/Runtime/Source/Parsers/Fbx/Private/FbxHelper.h
@@ -6,6 +6,8 @@
 #include "Containers/UnrealString.h"
 #include "FbxInclude.h"
 
+class UInterchangeBaseNode;
+
 namespace UE::Interchange
 {
 #if WITH_ENGINE
@@ -62,5 +64,10 @@ namespace UE::Interchange
 		private:
 			mutable TMap<FString, const FbxObject*> MaterialNameClashMap;
 		};
+
+		void ProcessCustomAttribute(FFbxParser& Parser, UInterchangeBaseNode* UnrealNode, FbxProperty Property, const TOptional<FString>& PayloadKey = TOptional<FString>());
+		
+		//Without AnimationProcessing
+		void ProcessCustomAttributes(FFbxParser& Parser, FbxObject* Object, UInterchangeBaseNode* UnrealNode);
 	}//ns Private
 }//ns UE::Interchange
diff --git a/Engine/Plugins/Interchange/Runtime/Source/Parsers/Fbx/Private/FbxMaterial.cpp b/Engine/Plugins/Interchange/Runtime/Source/Parsers/Fbx/Private/FbxMaterial.cpp
--- a/Engine/Plugins/Interchange/Runtime/Source/Parsers/Fbx/Private/FbxMaterial.cpp
+++ b/Engine/Plugins/Interchange/Runtime/Source/Parsers/Fbx/Private/FbxMaterial.cpp
@@ -290,6 +290,8 @@ namespace UE
 					return nullptr;
 				}
 
+				ProcessCustomAttributes(Parser, SurfaceMaterial, ShaderGraphNode);
+
 				TFunction<bool(FBXSDK_NAMESPACE::FbxProperty&)> ShouldConvertProperty = [&](FBXSDK_NAMESPACE::FbxProperty& MaterialProperty) -> bool
 				{
 					bool bShouldConvertProperty = false;
diff --git a/Engine/Plugins/Interchange/Runtime/Source/Parsers/Fbx/Private/FbxMesh.cpp b/Engine/Plugins/Interchange/Runtime/Source/Parsers/Fbx/Private/FbxMesh.cpp
--- a/Engine/Plugins/Interchange/Runtime/Source/Parsers/Fbx/Private/FbxMesh.cpp
+++ b/Engine/Plugins/Interchange/Runtime/Source/Parsers/Fbx/Private/FbxMesh.cpp
@@ -1766,6 +1766,9 @@ void FFbxMesh::AddAllMeshes(FbxScene* SDKScene, FbxGeometryConverter* SDKGeometr
 		}
 
 		MeshNode = CreateMeshNode(NodeContainer, MeshName, MeshUniqueID);
+
+		ProcessCustomAttributes(Parser, Mesh, MeshNode);
+
 		if (Geometry->GetDeformerCount(FbxDeformer::eSkin) > 0)
 		{
 			if (ExtractSkinnedMeshNodeJoints(SDKScene, NodeContainer, Mesh, MeshNode))
diff --git a/Engine/Plugins/Interchange/Runtime/Source/Parsers/Fbx/Private/FbxScene.cpp b/Engine/Plugins/Interchange/Runtime/Source/Parsers/Fbx/Private/FbxScene.cpp
--- a/Engine/Plugins/Interchange/Runtime/Source/Parsers/Fbx/Private/FbxScene.cpp
+++ b/Engine/Plugins/Interchange/Runtime/Source/Parsers/Fbx/Private/FbxScene.cpp
@@ -433,8 +433,6 @@ namespace UE
 					EFbxType PropertyType =  Property.GetPropertyDataType().GetType();
 					if (Property.GetFlag(FbxPropertyFlags::eUserDefined) && FFbxAnimation::IsFbxPropertyTypeSupported(PropertyType))
 					{
-						FString PropertyName = Parser.GetFbxHelper()->GetFbxPropertyName(Property);
-
 						FbxAnimCurveNode* CurveNode = Property.GetCurveNode();
 						TOptional<FString> PayloadKey;
 						if (CurveNode && CurveNode->IsAnimated())
@@ -451,118 +449,8 @@ namespace UE
 								}
 							}
 						}
-						switch (Property.GetPropertyDataType().GetType())
-						{
-							case EFbxType::eFbxBool:
-								{
-									bool PropertyValue = Property.Get<bool>();
-									UInterchangeUserDefinedAttributesAPI::CreateUserDefinedAttribute(UnrealNode, PropertyName, PropertyValue, PayloadKey);
-								}
-								break;
-							case EFbxType::eFbxChar:
-								{
-									int8 PropertyValue = Property.Get<int8>();
-									UInterchangeUserDefinedAttributesAPI::CreateUserDefinedAttribute(UnrealNode, PropertyName, PropertyValue, PayloadKey);
-								}
-								break;
-							case EFbxType::eFbxUChar:
-								{
-									uint8 PropertyValue = Property.Get<uint8>();
-									UInterchangeUserDefinedAttributesAPI::CreateUserDefinedAttribute(UnrealNode, PropertyName, PropertyValue, PayloadKey);
-								}
-								break;
-							case EFbxType::eFbxShort:
-								{
-									int16 PropertyValue = Property.Get<int16>();
-									UInterchangeUserDefinedAttributesAPI::CreateUserDefinedAttribute(UnrealNode, PropertyName, PropertyValue, PayloadKey);
-								}
-								break;
-							case EFbxType::eFbxUShort:
-								{
-									uint16 PropertyValue = Property.Get<uint16>();
-									UInterchangeUserDefinedAttributesAPI::CreateUserDefinedAttribute(UnrealNode, PropertyName, PropertyValue, PayloadKey);
-								}
-								break;
-							case EFbxType::eFbxInt:
-								{
-									int32 PropertyValue = Property.Get<int32>();
-									UInterchangeUserDefinedAttributesAPI::CreateUserDefinedAttribute(UnrealNode, PropertyName, PropertyValue, PayloadKey);
-								}
-								break;
-							case EFbxType::eFbxUInt:
-								{
-									uint32 PropertyValue = Property.Get<uint32>();
-									UInterchangeUserDefinedAttributesAPI::CreateUserDefinedAttribute(UnrealNode, PropertyName, PropertyValue, PayloadKey);
-								}
-								break;
-							case EFbxType::eFbxLongLong:
-								{
-									int64 PropertyValue = Property.Get<int64>();
-									UInterchangeUserDefinedAttributesAPI::CreateUserDefinedAttribute(UnrealNode, PropertyName, PropertyValue, PayloadKey);
-								}
-								break;
-							case EFbxType::eFbxULongLong:
-								{
-									uint64 PropertyValue = Property.Get<uint64>();
-									UInterchangeUserDefinedAttributesAPI::CreateUserDefinedAttribute(UnrealNode, PropertyName, PropertyValue, PayloadKey);
-								}
-								break;
-							case EFbxType::eFbxHalfFloat:
-								{
-									FbxHalfFloat HalfFloat = Property.Get<FbxHalfFloat>();
-									FFloat16 PropertyValue = FFloat16(HalfFloat.value());
-									UInterchangeUserDefinedAttributesAPI::CreateUserDefinedAttribute(UnrealNode, PropertyName, PropertyValue, PayloadKey);
-								}
-								break;
-							case EFbxType::eFbxFloat:
-								{
-									float PropertyValue = Property.Get<float>();
-									UInterchangeUserDefinedAttributesAPI::CreateUserDefinedAttribute(UnrealNode, PropertyName, PropertyValue, PayloadKey);
-								}
-								break;
-							case EFbxType::eFbxDouble:
-								{
-									double PropertyValue = Property.Get<double>();
-									UInterchangeUserDefinedAttributesAPI::CreateUserDefinedAttribute(UnrealNode, PropertyName, PropertyValue, PayloadKey);
-								}
-								break;
-							case EFbxType::eFbxDouble2:
-								{
-									FbxDouble2 Vec = Property.Get<FbxDouble2>();
-									FVector2D PropertyValue = FVector2D(Vec[0], Vec[1]);
-									UInterchangeUserDefinedAttributesAPI::CreateUserDefinedAttribute(UnrealNode, PropertyName, PropertyValue, PayloadKey);
-								}
-								break;
-							case EFbxType::eFbxDouble3:
-								{
-									FbxDouble3 Vec = Property.Get<FbxDouble3>();
-									FVector3d PropertyValue = FVector3d(Vec[0], Vec[1], Vec[2]);
-									UInterchangeUserDefinedAttributesAPI::CreateUserDefinedAttribute(UnrealNode, PropertyName, PropertyValue, PayloadKey);
-								}
-								break;
-							case EFbxType::eFbxDouble4:
-								{
-									FbxDouble4 Vec = Property.Get<FbxDouble4>();
-									FVector4d PropertyValue = FVector4d(Vec[0], Vec[1], Vec[2], Vec[3]);
-									UInterchangeUserDefinedAttributesAPI::CreateUserDefinedAttribute(UnrealNode, PropertyName, PropertyValue, PayloadKey);
-								}
-								break;
-							case EFbxType::eFbxEnum:
-								{
-									//Convert enum to uint8
-									FbxEnum EnumValue = Property.Get<FbxEnum>();
-									uint8 PropertyValue = static_cast<uint8>(EnumValue);
-									UInterchangeUserDefinedAttributesAPI::CreateUserDefinedAttribute(UnrealNode, PropertyName, PropertyValue, PayloadKey);
-								}
-								break;
-							case EFbxType::eFbxString:
-								{
-									FbxString StringValue = Property.Get<FbxString>();
-									FString PropertyValue = FFbxConvert::MakeString(StringValue.Buffer());
-									UInterchangeUserDefinedAttributesAPI::CreateUserDefinedAttribute(UnrealNode, PropertyName, PropertyValue, PayloadKey);
-								}
-								break;
-						}
+
+						ProcessCustomAttribute(Parser, UnrealNode, Property, PayloadKey);
 					}
 					//Inspect next node property
 					Property = Node->GetNextProperty(Property);
@@ -694,13 +582,12 @@ namespace UE
 						const FString TransformAnimTrackNodeUid = TEXT("\\AnimationTrack\\") + TransformAnimTrackNodeName;
 
 						TransformAnimTrackNode->InitializeNode(TransformAnimTrackNodeUid, TransformAnimTrackNodeName, EInterchangeNodeContainerType::TranslatedAsset);
-
 						TransformAnimTrackNode->SetCustomActorDependencyUid(*UnrealNode->GetUniqueID());
-
 						TransformAnimTrackNode->SetCustomAnimationPayloadKey(PayloadKey.GetValue(), EInterchangeAnimationPayLoadType::CURVE);
-
 						TransformAnimTrackNode->SetCustomUsedChannels(UsedChannels);
 
+						ProcessCustomAttributes(Parser, Node, TransformAnimTrackNode);
+
 						NodeContainer.AddNode(TransformAnimTrackNode);
 					}
 				}
@@ -781,13 +668,19 @@ namespace UE
 
 							SkeletalAnimationTrackNode->SetCustomAnimationStartTime(TimeSpan.GetStart().GetSecondDouble());
 							SkeletalAnimationTrackNode->SetCustomAnimationStopTime(TimeSpan.GetStop().GetSecondDouble());
+
+							return true;
 						}
+
+						return false;
 					};
 
 					bool bIsNodeContainJointAttribute = false;
 
 					int32 AttributeCount = Node->GetNodeAttributeCount();
 
+					bool bNewSkeltalAnimationStarted = false;
+
 					for (int32 AttributeIndex = 0; AttributeIndex < AttributeCount && !HasSkeletonAttribute; ++AttributeIndex)
 					{
 						FbxNodeAttribute* NodeAttribute = Node->GetNodeAttributeByIndex(AttributeIndex);
@@ -801,7 +694,7 @@ namespace UE
 							}
 						case FbxNodeAttribute::eSkeleton:
 							bIsNodeContainJointAttribute = true;
-							ApplySkeletonAttribute();
+							bNewSkeltalAnimationStarted = ApplySkeletonAttribute() || bNewSkeltalAnimationStarted;
 							break;
 						default:
 							break;
@@ -810,16 +703,20 @@ namespace UE
 
 					if (!bIsNodeContainJointAttribute)
 					{
-						//Make sure to threat the node like a joint if it's in the ForcejointNodes array
+						//Make sure to treat the node like a joint if it's in the ForcejointNodes array
 						if (ForceJointNodes.Contains(Node))
 						{
-							ApplySkeletonAttribute();
+							bNewSkeltalAnimationStarted = ApplySkeletonAttribute() || bNewSkeltalAnimationStarted;
 						}
 						else if (!bIsRootNode && DoesTheParentOrChildrenHierarchyContainJoints(Node))
 						{
-							ApplySkeletonAttribute();
+							bNewSkeltalAnimationStarted = ApplySkeletonAttribute() || bNewSkeltalAnimationStarted;
 						}
+					}
 
+					if (bNewSkeltalAnimationStarted)
+					{
+						ProcessCustomAttributes(Parser, Node, SkeletalAnimationTrackNode);
 					}
 
 					if (!HasSkeletonAttribute)
@@ -1097,6 +994,8 @@ namespace UE
 						SkeletalAnimationTrackNode->SetCustomAnimationStartTime(TimeSpan.GetStart().GetSecondDouble());
 						SkeletalAnimationTrackNode->SetCustomAnimationStopTime(TimeSpan.GetStop().GetSecondDouble());
 
+						ProcessCustomAttributes(Parser, CurrentAnimationStack, SkeletalAnimationTrackNode);
+
 						NodeContainer.AddNode(SkeletalAnimationTrackNode);
 						
 						for (const FMorphTargetAnimationBuildingData& MorphTargetAnimationBuildingDataPerSkeletonPerAnimationIndex : MorphTargetAnimationsBuildingDataPerSkeleton.Value)
diff --git a/Engine/Plugins/Interchange/Runtime/Source/Pipelines/Private/InterchangeGenericMaterialPipeline.cpp b/Engine/Plugins/Interchange/Runtime/Source/Pipelines/Private/InterchangeGenericMaterialPipeline.cpp
--- a/Engine/Plugins/Interchange/Runtime/Source/Pipelines/Private/InterchangeGenericMaterialPipeline.cpp
+++ b/Engine/Plugins/Interchange/Runtime/Source/Pipelines/Private/InterchangeGenericMaterialPipeline.cpp
@@ -3012,6 +3012,8 @@ UInterchangeMaterialFactoryNode* UInterchangeGenericMaterialPipeline::CreateMate
 {
 	UInterchangeMaterialFactoryNode* MaterialFactoryNode = Cast<UInterchangeMaterialFactoryNode>( CreateBaseMaterialFactoryNode(ShaderGraphNode, UInterchangeMaterialFactoryNode::StaticClass()) );
 
+	UInterchangeUserDefinedAttributesAPI::DuplicateAllUserDefinedAttribute(ShaderGraphNode, MaterialFactoryNode, false);
+
 	if(HandleSubstrate(ShaderGraphNode, MaterialFactoryNode))
 	{
 		return MaterialFactoryNode;
@@ -3064,6 +3066,8 @@ UInterchangeMaterialInstanceFactoryNode* UInterchangeGenericMaterialPipeline::Cr
 	UInterchangeMaterialInstanceFactoryNode* MaterialInstanceFactoryNode =
 		Cast<UInterchangeMaterialInstanceFactoryNode>( CreateBaseMaterialFactoryNode(ShaderGraphNode, UInterchangeMaterialInstanceFactoryNode::StaticClass()) );
 
+	UInterchangeUserDefinedAttributesAPI::DuplicateAllUserDefinedAttribute(ShaderGraphNode, MaterialInstanceFactoryNode, false);
+
 	TFunction<void(const FString&)> ChooseParent = [this, MaterialInstanceFactoryNode, ShaderGraphNode](const FString& Model) -> void
 	{
 		FString ParentRootName;
