[Interchange - glTF] Add StaticSwitches to Parent Materials in order to decrease compiled shader complexity.

#rnx
#virtualized

diff --git a/Engine/Plugins/Interchange/Runtime/Source/Import/Private/Gltf/InterchangeGLTFMaterial.cpp b/Engine/Plugins/Interchange/Runtime/Source/Import/Private/Gltf/InterchangeGLTFMaterial.cpp
--- a/Engine/Plugins/Interchange/Runtime/Source/Import/Private/Gltf/InterchangeGLTFMaterial.cpp
+++ b/Engine/Plugins/Interchange/Runtime/Source/Import/Private/Gltf/InterchangeGLTFMaterial.cpp
@@ -339,7 +339,7 @@ namespace UE::Interchange::GLTFMaterials
 			}
 		}
 
-		void SetMap(const FString& Name, const GLTF::FTextureMap& TextureMap)
+		bool SetMap(const FString& Name, const GLTF::FTextureMap& TextureMap)
 		{
 			if (Textures.IsValidIndex(TextureMap.TextureIndex))
 			{
@@ -368,33 +368,55 @@ namespace UE::Interchange::GLTFMaterials
 
 				//Set the TilingMethod:
 				SetVec4(Name + Inputs::PostFix::TilingMethod, GetTilingMethod(Textures[TextureMap.TextureIndex].Sampler), FVector4f(0, 0, 0, 0));
-			}
-
-			if (TextureMap.bHasTextureTransform)
-			{
-				{//Material specific settings:
-					SetScalar(Name + Inputs::PostFix::OffsetX, TextureMap.TextureTransform.Offset[0], 0.f, EProcessType::MATERIAL);
-					SetScalar(Name + Inputs::PostFix::OffsetY, TextureMap.TextureTransform.Offset[1], 0.f, EProcessType::MATERIAL);
 
-					SetScalar(Name + Inputs::PostFix::ScaleX, TextureMap.TextureTransform.Scale[0], 1.f, EProcessType::MATERIAL);
-					SetScalar(Name + Inputs::PostFix::ScaleY, TextureMap.TextureTransform.Scale[1], 1.f, EProcessType::MATERIAL);
-				}
+				if (TextureMap.bHasTextureTransform)
+				{
+					{//Material specific settings:
+						SetScalar(Name + Inputs::PostFix::OffsetX, TextureMap.TextureTransform.Offset[0], 0.f, EProcessType::MATERIAL);
+						SetScalar(Name + Inputs::PostFix::OffsetY, TextureMap.TextureTransform.Offset[1], 0.f, EProcessType::MATERIAL);
 
-				{//MaterialInstance specific settings:
-					FVector4f OffsetScale(TextureMap.TextureTransform.Offset[0], TextureMap.TextureTransform.Offset[1], TextureMap.TextureTransform.Scale[0], TextureMap.TextureTransform.Scale[1]);
-					SetVec4(Name + Inputs::PostFix::OffsetScale, OffsetScale, FVector4f(0, 0, 1, 1), EProcessType::MATERIALINSTANCE);
-				}
+						SetScalar(Name + Inputs::PostFix::ScaleX, TextureMap.TextureTransform.Scale[0], 1.f, EProcessType::MATERIAL);
+						SetScalar(Name + Inputs::PostFix::ScaleY, TextureMap.TextureTransform.Scale[1], 1.f, EProcessType::MATERIAL);
+					}
 
-				if (!FMath::IsNearlyZero(TextureMap.TextureTransform.Rotation))
-				{
-					float AngleRadians = TextureMap.TextureTransform.Rotation;
+					{//MaterialInstance specific settings:
+						FVector4f OffsetScale(TextureMap.TextureTransform.Offset[0], TextureMap.TextureTransform.Offset[1], TextureMap.TextureTransform.Scale[0], TextureMap.TextureTransform.Scale[1]);
+						SetVec4(Name + Inputs::PostFix::OffsetScale, OffsetScale, FVector4f(0, 0, 1, 1), EProcessType::MATERIALINSTANCE);
+					}
 
-					if (AngleRadians < 0.0f)
+					if (!FMath::IsNearlyZero(TextureMap.TextureTransform.Rotation))
 					{
-						AngleRadians = TWO_PI - AngleRadians;
+						float AngleRadians = TextureMap.TextureTransform.Rotation;
+
+						if (AngleRadians < 0.0f)
+						{
+							AngleRadians = TWO_PI - AngleRadians;
+						}
+
+						SetScalar(Name + Inputs::PostFix::Rotation, AngleRadians, 0.f);
 					}
+				}
+
+				return true;
+			}
+
+			return false;
+		}
 
-					SetScalar(Name + Inputs::PostFix::Rotation, AngleRadians, 0.f);
+		// Only setting it for MaterialInstances on default for now, the glTF pipeline is using MaterialInstances only.
+		//	However if someone disregards the usage of the glTF pipeline then the MF_XYZ_Body functions will be used which don't have configuration inputs.
+		void SetBool(const FString& Name, const bool& Value, const bool& DefaultValue = true, int ProcessType = (EProcessType::MATERIALINSTANCE))
+		{
+			if (Value != DefaultValue)
+			{
+				if (ProcessType & EProcessType::MATERIAL)
+				{
+					MaterialNode->AddBooleanAttribute(GetMaterialInputName(Name), Value);
+				}
+
+				if (ProcessType & EProcessType::MATERIALINSTANCE)
+				{
+					MaterialInstanceNode->AddBooleanAttribute(GetMaterialInstanceInputName(Name), Value);
 				}
 			}
 		}
@@ -415,10 +437,13 @@ namespace UE::Interchange::GLTFMaterials
 				//BaseColorTexture_Rotation
 				//BaseColorTexture_TexCoord
 				//BaseColorTexture_TilingMethod
-				SetMap(Inputs::BaseColorTexture, GltfMaterial.BaseColor);
+				bool bHasBaseColorTexture = SetMap(Inputs::BaseColorTexture, GltfMaterial.BaseColor);
 
 				//BaseColorFactor
 				SetColor(Inputs::BaseColorFactor, GltfMaterial.BaseColorFactor, FVector4f(1, 1, 1, 1));
+
+				//Configuration
+				SetBool(Inputs::Configuration::bHasBaseColorTexture, bHasBaseColorTexture);
 				return;
 			}
 
@@ -428,10 +453,13 @@ namespace UE::Interchange::GLTFMaterials
 				//NormalTexture_Rotation
 				//NormalTexture_TexCoord
 				//NormalTexture_TilingMethod
-				SetMap(Inputs::NormalTexture, GltfMaterial.Normal);
+				bool bHasNormalTexture = SetMap(Inputs::NormalTexture, GltfMaterial.Normal);
 
 				//NormalScale
 				SetScalar(Inputs::NormalScale, GltfMaterial.NormalScale, 1.f);
+
+				//Configuration
+				SetBool(Inputs::Configuration::bHasNormalTexture, bHasNormalTexture);
 			}
 
 			if (!GltfMaterial.bHasTransmission)
@@ -441,13 +469,21 @@ namespace UE::Interchange::GLTFMaterials
 				//EmissiveTexture_Rotation
 				//EmissiveTexture_TexCoord
 				//EmissiveTexture_TilingMethod
-				SetMap(Inputs::EmissiveTexture, GltfMaterial.Emissive);
+				bool bHasEmissiveTexture = SetMap(Inputs::EmissiveTexture, GltfMaterial.Emissive);
 
 				//EmissiveFactor
 				SetVec3(Inputs::EmissiveFactor, GltfMaterial.EmissiveFactor, FVector3f(0, 0, 0));
 
 				//EmissiveStrength
 				SetScalar(Inputs::EmissiveStrength, GltfMaterial.bHasEmissiveStrength ? GltfMaterial.EmissiveStrength : 1.f, 1.f);
+
+				//Configuration
+				SetBool(Inputs::Configuration::bHasEmissiveTexture, bHasEmissiveTexture);
+			}
+			else
+			{
+				//Configuration
+				SetBool(Inputs::Configuration::bHasEmissiveTexture, false);
 			}
 
 			{
@@ -456,10 +492,13 @@ namespace UE::Interchange::GLTFMaterials
 				//OcclusionTexture_Rotation
 				//OcclusionTexture_TexCoord
 				//OcclusionTexture_TilingMethod
-				SetMap(Inputs::OcclusionTexture, GltfMaterial.Occlusion);
+				bool bHasOcclusionTexture = SetMap(Inputs::OcclusionTexture, GltfMaterial.Occlusion);
 
 				//OcclusionStrength
 				SetScalar(Inputs::OcclusionStrength, GltfMaterial.OcclusionStrength, 1.f);
+
+				//Configuration
+				SetBool(Inputs::Configuration::bHasOcclusionTexture, bHasOcclusionTexture);
 			}
 
 			if (GltfMaterial.ShadingModel == GLTF::FMaterial::EShadingModel::SpecularGlossiness)
@@ -468,13 +507,14 @@ namespace UE::Interchange::GLTFMaterials
 				//specular glossiness:
 				//////////////////////
 
+				bool bHasDiffuseSpecGlossTexture = false;
 				{
 					//DiffuseTexture
 					//DiffuseTexture_OffsetScale
 					//DiffuseTexture_Rotation
 					//DiffuseTexture_TexCoord
 					//DiffuseTexture_TilingMethod
-					SetMap(Inputs::DiffuseTexture, GltfMaterial.BaseColor);
+					bHasDiffuseSpecGlossTexture = SetMap(Inputs::DiffuseTexture, GltfMaterial.BaseColor);
 
 					//DiffuseFactor
 					SetColor(Inputs::DiffuseFactor, GltfMaterial.BaseColorFactor, FVector4f(1, 1, 1, 1));
@@ -486,7 +526,7 @@ namespace UE::Interchange::GLTFMaterials
 					//SpecularGlossinessTexture_Rotation
 					//SpecularGlossinessTexture_TexCoord
 					//SpecularGlossinessTexture_TilingMethod
-					SetMap(Inputs::SpecularGlossinessTexture, GltfMaterial.SpecularGlossiness.Map);
+					bHasDiffuseSpecGlossTexture = SetMap(Inputs::SpecularGlossinessTexture, GltfMaterial.SpecularGlossiness.Map) || bHasDiffuseSpecGlossTexture;
 
 					//SpecFactor
 					FVector3f SpecularFactor(GltfMaterial.SpecularGlossiness.SpecularFactor[0], GltfMaterial.SpecularGlossiness.SpecularFactor[1], GltfMaterial.SpecularGlossiness.SpecularFactor[2]);
@@ -496,6 +536,9 @@ namespace UE::Interchange::GLTFMaterials
 					SetScalar(Inputs::GlossinessFactor, GltfMaterial.SpecularGlossiness.GlossinessFactor, 1.f);
 				}
 
+				//Configuration
+				SetBool(Inputs::Configuration::bHasDiffuseSpecGlossTexture, bHasDiffuseSpecGlossTexture);
+
 				return;
 			}
 			else if (GltfMaterial.ShadingModel == GLTF::FMaterial::EShadingModel::MetallicRoughness)
@@ -506,10 +549,13 @@ namespace UE::Interchange::GLTFMaterials
 					//BaseColorTexture_Rotation
 					//BaseColorTexture_TexCoord
 					//BaseColorTexture_TilingMethod
-					SetMap(Inputs::BaseColorTexture, GltfMaterial.BaseColor);
+					bool bHasBaseColorTexture = SetMap(Inputs::BaseColorTexture, GltfMaterial.BaseColor);
 
 					//BaseColorFactor
 					SetColor(Inputs::BaseColorFactor, GltfMaterial.BaseColorFactor, FVector4f(1, 1, 1, 1));
+
+					//Configuration
+					SetBool(Inputs::Configuration::bHasBaseColorTexture, bHasBaseColorTexture);
 				}
 
 				{
@@ -518,13 +564,16 @@ namespace UE::Interchange::GLTFMaterials
 					//MetallicRoughnessTexture_Rotation
 					//MetallicRoughnessTexture_TexCoord
 					//MetallicRoughnessTexture_TilingMethod
-					SetMap(Inputs::MetallicRoughnessTexture, GltfMaterial.MetallicRoughness.Map);
+					bool bHasMetallicRoughnessTexture = SetMap(Inputs::MetallicRoughnessTexture, GltfMaterial.MetallicRoughness.Map);
 
 					//MetallicFactor
 					SetScalar(Inputs::MetallicFactor, GltfMaterial.MetallicRoughness.MetallicFactor, 1.f);
 
 					//RoughnessFactor
 					SetScalar(Inputs::RoughnessFactor, GltfMaterial.MetallicRoughness.RoughnessFactor, 1.f);
+
+					//Configuration
+					SetBool(Inputs::Configuration::bHasMetallicRoughnessTexture, bHasMetallicRoughnessTexture);
 				}
 
 				if (GltfMaterial.bHasSpecular)
@@ -534,10 +583,18 @@ namespace UE::Interchange::GLTFMaterials
 					//SpecularTexture_Rotation
 					//SpecularTexture_TexCoord
 					//SpecularTexture_TilingMethod
-					SetMap(Inputs::SpecularTexture, GltfMaterial.Specular.SpecularMap);
+					bool bHasSpecularTexture = SetMap(Inputs::SpecularTexture, GltfMaterial.Specular.SpecularMap);
 
 					//SpecularFactor
 					SetScalar(Inputs::SpecularFactor, GltfMaterial.Specular.SpecularFactor, 0.5); //(UE Specular default is 0.5).
+
+					//Configuration
+					SetBool(Inputs::Configuration::bHasSpecularTexture, bHasSpecularTexture);
+				}
+				else
+				{
+					//Configuration
+					SetBool(Inputs::Configuration::bHasSpecularTexture, false);
 				}
 
 				if (GltfMaterial.Iridescence.bHasIridescence)
@@ -550,7 +607,7 @@ namespace UE::Interchange::GLTFMaterials
 					//IridescenceTexture_Rotation
 					//IridescenceTexture_TexCoord
 					//IridescenceTexture_TilingMethod
-					SetMap(Inputs::IridescenceTexture, GltfMaterial.Iridescence.Texture);
+					bool bHasIridescenceTexture = SetMap(Inputs::IridescenceTexture, GltfMaterial.Iridescence.Texture);
 
 					//IridescenceFactor
 					SetScalar(Inputs::IridescenceFactor, GltfMaterial.Iridescence.Factor, 0.0f);
@@ -560,12 +617,20 @@ namespace UE::Interchange::GLTFMaterials
 					//IridescenceThicknessTexture_Rotation
 					//IridescenceThicknessTexture_TexCoord
 					//IridescenceThicknessTexture_TilingMethod
-					SetMap(Inputs::IridescenceThicknessTexture, GltfMaterial.Iridescence.Thickness.Texture);
+					bHasIridescenceTexture = SetMap(Inputs::IridescenceThicknessTexture, GltfMaterial.Iridescence.Thickness.Texture) || bHasIridescenceTexture;
 
 					//IridescenceThicknessMinimum
 					SetScalar(Inputs::IridescenceThicknessMinimum, GltfMaterial.Iridescence.Thickness.Minimum, 100.0f);
 					//IridescenceThicknessMaximum
 					SetScalar(Inputs::IridescenceThicknessMaximum, GltfMaterial.Iridescence.Thickness.Maximum, 400.0f);
+
+					//Configuration
+					SetBool(Inputs::Configuration::bHasIridescenceTexture, bHasIridescenceTexture);
+				}
+				else
+				{
+					//Configuration
+					SetBool(Inputs::Configuration::bHasIridescenceTexture, false);
 				}
 			}
 			else
@@ -579,6 +644,8 @@ namespace UE::Interchange::GLTFMaterials
 				SetScalar(Inputs::IOR, GltfMaterial.IOR, 1.5f);
 			}
 
+			//Transmission ClearCoat and Sheen are Model specific
+			// Meaning we don't have to do exclusion configs.
 			if (GltfMaterial.bHasTransmission)
 			{
 				{
@@ -587,24 +654,28 @@ namespace UE::Interchange::GLTFMaterials
 					//TransmissionTexture_Rotation
 					//TransmissionTexture_TexCoord
 					//TransmissionTexture_TilingMethod
-					SetMap(Inputs::TransmissionTexture, GltfMaterial.Transmission.TransmissionMap);
+					bool bHasTransmissionTexture = SetMap(Inputs::TransmissionTexture, GltfMaterial.Transmission.TransmissionMap);
 
 					//TransmissionFactor
 					SetScalar(Inputs::TransmissionFactor, GltfMaterial.Transmission.TransmissionFactor, 0.f);
 
 					//AlphaMode
 					SetScalar(Inputs::AlphaMode, AlphaMode, EAlphaMode::Blend);
+
+					//Configuration
+					SetBool(Inputs::Configuration::bHasTransmissionTexture, bHasTransmissionTexture);
 				}
 			}
 			else if (GltfMaterial.bHasClearCoat)
 			{
+				bool bHasClearCoatTexture = false;
 				{
 					//ClearCoatTexture
 					//ClearCoatTexture_OffsetScale
 					//ClearCoatTexture_Rotation
 					//ClearCoatTexture_TexCoord
 					//ClearCoatTexture_TilingMethod
-					SetMap(Inputs::ClearCoatTexture, GltfMaterial.ClearCoat.ClearCoatMap);
+					bHasClearCoatTexture = SetMap(Inputs::ClearCoatTexture, GltfMaterial.ClearCoat.ClearCoatMap);
 
 					//ClearCoatFactor
 					SetScalar(Inputs::ClearCoatFactor, GltfMaterial.ClearCoat.ClearCoatFactor, 0.f);
@@ -616,7 +687,7 @@ namespace UE::Interchange::GLTFMaterials
 					//ClearCoatRoughnessTexture_Rotation
 					//ClearCoatRoughnessTexture_TexCoord
 					//ClearCoatRoughnessTexture_TilingMethod
-					SetMap(Inputs::ClearCoatRoughnessTexture, GltfMaterial.ClearCoat.RoughnessMap);
+					bHasClearCoatTexture = SetMap(Inputs::ClearCoatRoughnessTexture, GltfMaterial.ClearCoat.RoughnessMap) || bHasClearCoatTexture;
 
 					//ClearCoatRoughnessFactor
 					SetScalar(Inputs::ClearCoatRoughnessFactor, GltfMaterial.ClearCoat.Roughness, 0.f);
@@ -628,11 +699,14 @@ namespace UE::Interchange::GLTFMaterials
 					//ClearCoatNormalTexture_Rotation
 					//ClearCoatNormalTexture_TexCoord
 					//ClearCoatNormalTexture_TilingMethod
-					SetMap(Inputs::ClearCoatNormalTexture, GltfMaterial.ClearCoat.NormalMap);
+					bHasClearCoatTexture = SetMap(Inputs::ClearCoatNormalTexture, GltfMaterial.ClearCoat.NormalMap) || bHasClearCoatTexture;
 
 					//ClearCoatNormalFactor
 					SetScalar(Inputs::ClearCoatNormalScale, GltfMaterial.ClearCoat.NormalMapUVScale, 1.f);
 				}
+
+				//Configuration
+				SetBool(Inputs::Configuration::bHasClearCoatTexture, bHasClearCoatTexture);
 			}
 			else if (GltfMaterial.bHasSheen)
 			{
@@ -642,7 +716,7 @@ namespace UE::Interchange::GLTFMaterials
 					//SheenColorTexture_Rotation
 					//SheenColorTexture_TexCoord
 					//SheenColorTexture_TilingMethod
-					SetMap(Inputs::SheenColorTexture, GltfMaterial.Sheen.SheenColorMap);
+					bool bHasSheenTexture = SetMap(Inputs::SheenColorTexture, GltfMaterial.Sheen.SheenColorMap);
 
 					//SheenColorFactor
 					FVector3f SheenColorFactor(GltfMaterial.Sheen.SheenColorFactor[0], GltfMaterial.Sheen.SheenColorFactor[1], GltfMaterial.Sheen.SheenColorFactor[2]);
@@ -653,10 +727,13 @@ namespace UE::Interchange::GLTFMaterials
 					//SheenRoughnessTexture_Rotation
 					//SheenRoughnessTexture_TexCoord
 					//SheenRoughnessTexture_TilingMethod
-					SetMap(Inputs::SheenRoughnessTexture, GltfMaterial.Sheen.SheenRoughnessMap);
+					bHasSheenTexture = SetMap(Inputs::SheenRoughnessTexture, GltfMaterial.Sheen.SheenRoughnessMap) || bHasSheenTexture;
 
 					//SheenRoughnessFactor
 					SetScalar(Inputs::SheenRoughnessFactor, GltfMaterial.Sheen.SheenRoughnessFactor, 0.f);
+
+					//Configuration
+					SetBool(Inputs::Configuration::bHasSheenTexture, bHasSheenTexture);
 				}
 			}
 		}
diff --git a/Engine/Plugins/Interchange/Runtime/Source/Import/Public/Gltf/InterchangeGLTFMaterial.h b/Engine/Plugins/Interchange/Runtime/Source/Import/Public/Gltf/InterchangeGLTFMaterial.h
--- a/Engine/Plugins/Interchange/Runtime/Source/Import/Public/Gltf/InterchangeGLTFMaterial.h
+++ b/Engine/Plugins/Interchange/Runtime/Source/Import/Public/Gltf/InterchangeGLTFMaterial.h
@@ -124,6 +124,23 @@ namespace UE::Interchange::GLTFMaterials
 		DECLARE_INTERCHANGE_GLTF_MI_MAP(IridescenceThickness)
 		const FString IridescenceThicknessMinimum = TEXT("IridescenceThicknessMinimum");
 		const FString IridescenceThicknessMaximum = TEXT("IridescenceThicknessMaximum");
+
+		namespace Configuration
+		{
+			//For StaticSwitch optimizations:
+			// Note: only usable in Editor, the Default values are true so that Runtime will still use the complete graph (without optimization).
+			const FString bHasBaseColorTexture = TEXT("bHasBaseColorTexture");
+			const FString bHasMetallicRoughnessTexture = TEXT("bHasMetallicRoughnessTexture");
+			const FString bHasDiffuseSpecGlossTexture = TEXT("bHasDiffuseSpecGlossTexture");
+			const FString bHasEmissiveTexture = TEXT("bHasEmissiveTexture");
+			const FString bHasNormalTexture = TEXT("bHasNormalTexture");
+			const FString bHasOcclusionTexture = TEXT("bHasOcclusionTexture");
+			const FString bHasSpecularTexture = TEXT("bHasSpecularTexture");
+			const FString bHasClearCoatTexture = TEXT("bHasClearCoatTexture");
+			const FString bHasSheenTexture = TEXT("bHasSheenTexture");
+			const FString bHasTransmissionTexture = TEXT("bHasTransmissionTexture");
+			const FString bHasIridescenceTexture = TEXT("bHasIridescenceTexture");
+		}
 	}
 
 	enum EShadingModel : uint8
diff --git a/Engine/Plugins/Interchange/Runtime/Source/Pipelines/Private/InterchangeglTFPipeline.cpp b/Engine/Plugins/Interchange/Runtime/Source/Pipelines/Private/InterchangeglTFPipeline.cpp
--- a/Engine/Plugins/Interchange/Runtime/Source/Pipelines/Private/InterchangeglTFPipeline.cpp
+++ b/Engine/Plugins/Interchange/Runtime/Source/Pipelines/Private/InterchangeglTFPipeline.cpp
@@ -234,9 +234,18 @@ void UGLTFPipelineSettings::BuildMaterialInstance(const UInterchangeShaderGraphN
 
 		FString InputValueKey = UInterchangeShaderPortsAPI::MakeInputValueKey(GltfAttributeKey.Key);
 
-		//we are only using 3 attribute types:
+		//we are only using 4 attribute types for now:
 		switch (AttributeType)
 		{
+		case UE::Interchange::EAttributeTypes::Bool:
+		{
+			bool Value;
+			if (ShaderGraphNode->GetBooleanAttribute(GltfAttributeKey.Value.Key, Value))
+			{
+				MaterialInstanceFactoryNode->AddBooleanAttribute(InputValueKey, Value);
+			}
+		}
+		break;
 		case UE::Interchange::EAttributeTypes::Float:
 		{
 			float Value;
