Interchange:
- Add Texture Color Space support
- Set default texture colorspace setting to sRGB / Rec709 when sRGB is set on the UTexture

diff --git a/Engine/Plugins/Interchange/Runtime/Source/FactoryNodes/InterchangeFactoryNodes.Build.cs b/Engine/Plugins/Interchange/Runtime/Source/FactoryNodes/InterchangeFactoryNodes.Build.cs
--- a/Engine/Plugins/Interchange/Runtime/Source/FactoryNodes/InterchangeFactoryNodes.Build.cs
+++ b/Engine/Plugins/Interchange/Runtime/Source/FactoryNodes/InterchangeFactoryNodes.Build.cs
@@ -25,6 +25,7 @@ namespace UnrealBuildTool.Rules
 				new string[]
 				{
 					"LevelSequence",
+					"ColorManagement",
 				}
 			);
 		}
diff --git a/Engine/Plugins/Interchange/Runtime/Source/FactoryNodes/Private/InterchangeTextureFactoryNode.cpp b/Engine/Plugins/Interchange/Runtime/Source/FactoryNodes/Private/InterchangeTextureFactoryNode.cpp
new file mode 100644
--- /dev/null
+++ b/Engine/Plugins/Interchange/Runtime/Source/FactoryNodes/Private/InterchangeTextureFactoryNode.cpp
@@ -0,0 +1,51 @@
+// Copyright Epic Games, Inc. All Rights Reserved.
+
+#include "InterchangeTextureFactoryNode.h"
+
+#include "ColorSpace.h"
+
+#include UE_INLINE_GENERATED_CPP_BY_NAME(InterchangeTextureFactoryNode)
+
+
+bool UInterchangeTextureFactoryNode::GetCustomColorSpace(ETextureColorSpace& AttributeValue) const
+{
+	IMPLEMENT_NODE_ATTRIBUTE_GETTER(ColorSpace, ETextureColorSpace);
+}
+
+bool UInterchangeTextureFactoryNode::SetCustomColorSpace(ETextureColorSpace AttributeValue, bool bAddApplyDelegate)
+{
+#if WITH_EDITORONLY_DATA
+	IMPLEMENT_NODE_ATTRIBUTE_SETTER_WITH_CUSTOM_DELEGATE_WITH_CUSTOM_CLASS(UInterchangeTextureFactoryNode, ColorSpace, ETextureColorSpace, UTexture)
+#else
+	IMPLEMENT_NODE_ATTRIBUTE_SETTER_NODELEGATE(ColorSpace, ETextureColorSpace)
+#endif
+}
+
+#if WITH_EDITORONLY_DATA
+bool UInterchangeTextureFactoryNode::ApplyCustomColorSpaceToAsset(UObject* Asset) const
+{
+	if(UTexture* Texture = Cast<UTexture>(Asset))
+	{
+		if(ETextureColorSpace ValueData; GetCustomColorSpace(ValueData))
+		{
+			Texture->SourceColorSettings.ColorSpace = ValueData;
+			Texture->SourceColorSettings.UpdateColorSpaceChromaticities();
+
+			return true;
+		}
+	}
+	return false;
+}
+
+bool UInterchangeTextureFactoryNode::FillCustomColorSpaceFromAsset(UObject* Asset)
+{
+	if(UTexture* TypedObject = Cast<UTexture>(Asset))
+	{
+		if(SetCustomColorSpace(TypedObject->SourceColorSettings.ColorSpace, false))
+		{
+			return true;
+		}
+	}
+	return false;
+}
+#endif // WITH_EDITORONLY_DATA
\ No newline at end of file
diff --git a/Engine/Plugins/Interchange/Runtime/Source/FactoryNodes/Public/InterchangeTextureFactoryNode.h b/Engine/Plugins/Interchange/Runtime/Source/FactoryNodes/Public/InterchangeTextureFactoryNode.h
--- a/Engine/Plugins/Interchange/Runtime/Source/FactoryNodes/Public/InterchangeTextureFactoryNode.h
+++ b/Engine/Plugins/Interchange/Runtime/Source/FactoryNodes/Public/InterchangeTextureFactoryNode.h
@@ -528,6 +528,13 @@ public:
 		IMPLEMENT_NODE_ATTRIBUTE_SETTER(UInterchangeTextureFactoryNode, VirtualTextureStreaming, bool, UTexture)
 	}
 
+	/** Return false if the Attribute was not set previously.*/
+	UFUNCTION(BlueprintCallable, Category = "Interchange | Node | Texture")
+	bool GetCustomColorSpace(ETextureColorSpace& AttributeValue) const;
+
+	UFUNCTION(BlueprintCallable, Category = "Interchange | Node | Texture")
+	bool SetCustomColorSpace(ETextureColorSpace AttributeValue, bool bAddApplyDelegate = true);
+
 	//////////////////////////////////////////////////////////////////////////
 	//Level of Detail
 
@@ -725,6 +732,7 @@ public:
 			COPY_NODE_DELEGATES(TextureFactoryNode, MipGenSettings, uint8, UTexture)
 			COPY_NODE_DELEGATES(TextureFactoryNode, CompositeTextureMode, uint8, UTexture)
 			COPY_NODE_DELEGATES(TextureFactoryNode, CompositePower, float, UTexture)
+			COPY_NODE_DELEGATES_WITH_CUSTOM_DELEGATE(TextureFactoryNode, UInterchangeTextureFactoryNode, ColorSpace, ETextureColorSpace, UTexture::StaticClass())
 #endif
 			COPY_NODE_DELEGATES(TextureFactoryNode, VirtualTextureStreaming, bool, UTexture)
 			COPY_NODE_DELEGATES(TextureFactoryNode, CompressionSettings, uint8, UTexture)
@@ -779,6 +787,7 @@ private:
 	IMPLEMENT_NODE_ATTRIBUTE_KEY(SRGB)
 	IMPLEMENT_NODE_ATTRIBUTE_KEY(bUseLegacyGamma)
 	IMPLEMENT_NODE_ATTRIBUTE_KEY(VirtualTextureStreaming)
+	IMPLEMENT_NODE_ATTRIBUTE_KEY(ColorSpace)
 
 	// Level of Detail
 	IMPLEMENT_NODE_ATTRIBUTE_KEY(bPreserveBorder)
@@ -841,4 +850,10 @@ private:
 		}
 		return false;
 	}
+
+#if WITH_EDITORONLY_DATA
+	bool ApplyCustomColorSpaceToAsset(UObject* Asset) const;
+
+	bool FillCustomColorSpaceFromAsset(UObject* Asset);
+#endif // WITH_EDITORONLY_DATA
 };
diff --git a/Engine/Plugins/Interchange/Runtime/Source/Nodes/Public/InterchangeTextureNode.h b/Engine/Plugins/Interchange/Runtime/Source/Nodes/Public/InterchangeTextureNode.h
--- a/Engine/Plugins/Interchange/Runtime/Source/Nodes/Public/InterchangeTextureNode.h
+++ b/Engine/Plugins/Interchange/Runtime/Source/Nodes/Public/InterchangeTextureNode.h
@@ -41,6 +41,27 @@ enum class EInterchangeTextureFilterMode : uint8
 	Default
 };
 
+UENUM(BlueprintType)
+enum class EInterchangeTextureColorSpace : uint8
+{
+	TCS_None = 0 UMETA(DisplayName = "None", ToolTip = "No explicit color space definition."),
+	TCS_sRGB = 1 UMETA(DisplayName = "sRGB / Rec709", ToolTip = "sRGB / Rec709 (BT.709) color primaries, with D65 white point."),
+	TCS_Rec2020 = 2 UMETA(DisplayName = "Rec2020", ToolTip = "Rec2020 (BT.2020) primaries with D65 white point."),
+	TCS_ACESAP0 = 3 UMETA(DIsplayName = "ACES AP0", ToolTip = "ACES AP0 wide gamut primaries, with D60 white point."),
+	TCS_ACESAP1 = 4 UMETA(DIsplayName = "ACES AP1 / ACEScg", ToolTip = "ACES AP1 / ACEScg wide gamut primaries, with D60 white point."),
+	TCS_P3DCI = 5 UMETA(DisplayName = "P3DCI", ToolTip = "P3 (Theater) primaries, with DCI Calibration white point."),
+	TCS_P3D65 = 6 UMETA(DisplayName = "P3D65", ToolTip = "P3 (Display) primaries, with D65 white point."),
+	TCS_REDWideGamut = 7 UMETA(DisplayName = "RED Wide Gamut", ToolTip = "RED Wide Gamut primaries, with D65 white point."),
+	TCS_SonySGamut3 = 8 UMETA(DisplayName = "Sony S-Gamut3", ToolTip = "Sony S-Gamut/S-Gamut3 primaries, with D65 white point."),
+	TCS_SonySGamut3Cine = 9 UMETA(DisplayName = "Sony S-Gamut3 Cine", ToolTip = "Sony S-Gamut3 Cine primaries, with D65 white point."),
+	TCS_AlexaWideGamut = 10 UMETA(DisplayName = "Alexa Wide Gamut", ToolTip = "Alexa Wide Gamut primaries, with D65 white point."),
+	TCS_CanonCinemaGamut = 11 UMETA(DisplayName = "Canon Cinema Gamut", ToolTip = "Canon Cinema Gamut primaries, with D65 white point."),
+	TCS_GoProProtuneNative = 12 UMETA(DisplayName = "GoPro Protune Native", ToolTip = "GoPro Protune Native primaries, with D65 white point."),
+	TCS_PanasonicVGamut = 13 UMETA(DisplayName = "Panasonic V-Gamut", ToolTip = "Panasonic V-Gamut primaries, with D65 white point."),
+	TCS_Custom = 99 UMETA(DisplayName = "Custom", ToolTip = "User defined color space and white point."),
+	TCS_MAX,
+};
+
 UCLASS(BlueprintType, Abstract)
 class INTERCHANGENODES_API UInterchangeTextureNode : public UInterchangeBaseNode
 {
@@ -148,9 +169,22 @@ public:
 		IMPLEMENT_NODE_ATTRIBUTE_SETTER_NODELEGATE(Filter, EInterchangeTextureFilterMode)
 	}
 
+	UFUNCTION(BlueprintCallable, Category = "Interchange | Node | Texture")
+	bool GetCustomColorSpace(EInterchangeTextureColorSpace& AttributeValue) const
+	{
+		IMPLEMENT_NODE_ATTRIBUTE_GETTER(ColorSpace, EInterchangeTextureColorSpace);
+	}
+
+	UFUNCTION(BlueprintCallable, Category = "Interchange | Node | Texture")
+	bool SetCustomColorSpace(const EInterchangeTextureColorSpace& AttributeValue)
+	{
+		IMPLEMENT_NODE_ATTRIBUTE_SETTER_NODELEGATE(ColorSpace, EInterchangeTextureColorSpace)
+	}
+
 private:
 
 	IMPLEMENT_NODE_ATTRIBUTE_KEY(SRGB)
 	IMPLEMENT_NODE_ATTRIBUTE_KEY(bFlipGreenChannel)
 	IMPLEMENT_NODE_ATTRIBUTE_KEY(Filter)
+	IMPLEMENT_NODE_ATTRIBUTE_KEY(ColorSpace)
 };
diff --git a/Engine/Plugins/Interchange/Runtime/Source/Pipelines/Private/InterchangeGenericTexturePipeline.cpp b/Engine/Plugins/Interchange/Runtime/Source/Pipelines/Private/InterchangeGenericTexturePipeline.cpp
--- a/Engine/Plugins/Interchange/Runtime/Source/Pipelines/Private/InterchangeGenericTexturePipeline.cpp
+++ b/Engine/Plugins/Interchange/Runtime/Source/Pipelines/Private/InterchangeGenericTexturePipeline.cpp
@@ -335,9 +335,20 @@ UInterchangeTextureFactoryNode* UInterchangeGenericTexturePipeline::CreateTextur
 		}
 	}
 
+	EInterchangeTextureColorSpace ColorSpace;
+	const bool bHasColorSpace = TextureNode->GetCustomColorSpace(ColorSpace);
+	if(bHasColorSpace)
+	{
+		TextureFactoryNode->SetCustomColorSpace(ETextureColorSpace(ColorSpace));
+	}
+
 	if (bool bSRGB; TextureNode->GetCustomSRGB(bSRGB))
 	{
 		TextureFactoryNode->SetCustomSRGB(bSRGB);
+		if(bSRGB && !bHasColorSpace)
+		{
+			TextureFactoryNode->SetCustomColorSpace(ETextureColorSpace::TCS_sRGB);
+		}
 	}
 	if (bool bFlipGreenChannel; TextureNode->GetCustombFlipGreenChannel(bFlipGreenChannel))
 	{
diff --git a/Engine/Source/Runtime/Engine/Classes/Engine/Texture.h b/Engine/Source/Runtime/Engine/Classes/Engine/Texture.h
--- a/Engine/Source/Runtime/Engine/Classes/Engine/Texture.h
+++ b/Engine/Source/Runtime/Engine/Classes/Engine/Texture.h
@@ -1147,6 +1147,11 @@ struct FTextureSourceColorSettings
 	/** Chromatic adaption method applied if the source white point differs from the working color space white point. */
 	UPROPERTY(EditAnywhere, BlueprintReadWrite, Category = ColorManagement)
 	ETextureChromaticAdaptationMethod ChromaticAdaptationMethod;
+
+#if WITH_EDITORONLY_DATA
+	/** Update the chromaticity coordinates member variables based on the color space choice (unless custom). */
+	ENGINE_API void UpdateColorSpaceChromaticities();
+#endif
 };
 
 UCLASS(abstract, MinimalAPI, BlueprintType)
diff --git a/Engine/Source/Runtime/Engine/Private/Texture.cpp b/Engine/Source/Runtime/Engine/Private/Texture.cpp
--- a/Engine/Source/Runtime/Engine/Private/Texture.cpp
+++ b/Engine/Source/Runtime/Engine/Private/Texture.cpp
@@ -837,12 +837,7 @@ PRAGMA_ENABLE_DEPRECATION_WARNINGS
 #if WITH_EDITORONLY_DATA
 		else if (PropertyName == SourceColorSpaceName)
 		{
-			// Update the chromaticity coordinates member variables based on the color space choice (unless custom).
-			if (SourceColorSettings.ColorSpace != ETextureColorSpace::TCS_Custom)
-			{
-				UE::Color::FColorSpace ColorSpace(static_cast<UE::Color::EColorSpace>(SourceColorSettings.ColorSpace));
-				ColorSpace.GetChromaticities(SourceColorSettings.RedChromaticityCoordinate, SourceColorSettings.GreenChromaticityCoordinate, SourceColorSettings.BlueChromaticityCoordinate, SourceColorSettings.WhiteChromaticityCoordinate);
-			}
+			SourceColorSettings.UpdateColorSpaceChromaticities();
 		}
 		else if (PropertyName == CompressionQualityName)
 		{
@@ -1179,6 +1174,15 @@ PRAGMA_DISABLE_DEPRECATION_WARNINGS
 PRAGMA_ENABLE_DEPRECATION_WARNINGS
 }
 
+void FTextureSourceColorSettings::UpdateColorSpaceChromaticities()
+{
+	if(ColorSpace != ETextureColorSpace::TCS_Custom)
+	{
+		UE::Color::FColorSpace ColorSpacChromaticities(static_cast<UE::Color::EColorSpace>(ColorSpace));
+		ColorSpacChromaticities.GetChromaticities(RedChromaticityCoordinate, GreenChromaticityCoordinate, BlueChromaticityCoordinate, WhiteChromaticityCoordinate);
+	}
+}
+
 #endif // #if WITH_EDITORONLY_DATA
 
 void UTexture::PostInitProperties()
