Interchange FBX:
- Add FBX project settings
- Add data structure to support custom attributes
- Add Read-only mapping property tracks to Interchange

diff --git a/Engine/Plugins/Interchange/Runtime/Source/Nodes/Public/InterchangeAnimationDefinitions.h b/Engine/Plugins/Interchange/Runtime/Source/Nodes/Public/InterchangeAnimationDefinitions.h
--- a/Engine/Plugins/Interchange/Runtime/Source/Nodes/Public/InterchangeAnimationDefinitions.h
+++ b/Engine/Plugins/Interchange/Runtime/Source/Nodes/Public/InterchangeAnimationDefinitions.h
@@ -4,6 +4,40 @@
 
 #include "UObject/NameTypes.h"
 
+#include "InterchangeAnimationDefinitions.generated.h"
+
+UENUM()
+enum class EInterchangePropertyTracks : int32
+{
+	/** Common Properties*/
+	Visibility,
+
+	/** Light Properties */
+	LightColor,
+	LightIntensity,
+	LightIntensityUnits,
+	LightTemperature,
+	LightUseTemperature,
+
+	/** Camera Properties*/
+	CameraAutoActivate,
+	CameraAspectRatioAxisConstraint,
+	CameraConstrainAspectRatio,
+	CameraCurrentAperture,
+	CameraCurrentFocalLength,
+	CameraCustomNearClippingPlane,
+	CameraFieldOfView,
+	CameraMobility,
+	CameraOrthoFarClipPlane,
+	CameraOrthoWidth,
+	CameraPostProcessBlendWeight,
+	CameraProjectionMode,
+	CameraShouldUpdatePhysicsVolume,
+	CameraUseFieldOfViewForLOD,
+
+	None = -1 UMETA(hidden),
+};
+
 namespace UE
 {
 	namespace Interchange
diff --git a/Engine/Plugins/Interchange/Runtime/Source/Parsers/Fbx/InterchangeFbxParser.Build.cs b/Engine/Plugins/Interchange/Runtime/Source/Parsers/Fbx/InterchangeFbxParser.Build.cs
--- a/Engine/Plugins/Interchange/Runtime/Source/Parsers/Fbx/InterchangeFbxParser.Build.cs
+++ b/Engine/Plugins/Interchange/Runtime/Source/Parsers/Fbx/InterchangeFbxParser.Build.cs
@@ -34,6 +34,7 @@ namespace UnrealBuildTool.Rules
 			PrivateDependencyModuleNames.AddRange(
 				new string[]
 				{
+					"DeveloperSettings",
 					"InterchangeCommonParser",
 					"InterchangeMessages",
 					"InterchangeNodes",
diff --git a/Engine/Plugins/Interchange/Runtime/Source/Parsers/Fbx/Private/FbxScene.cpp b/Engine/Plugins/Interchange/Runtime/Source/Parsers/Fbx/Private/FbxScene.cpp
--- a/Engine/Plugins/Interchange/Runtime/Source/Parsers/Fbx/Private/FbxScene.cpp
+++ b/Engine/Plugins/Interchange/Runtime/Source/Parsers/Fbx/Private/FbxScene.cpp
@@ -20,6 +20,7 @@
 #include "Nodes/InterchangeUserDefinedAttribute.h"
 #include "InterchangeAnimationTrackSetNode.h"
 #include "InterchangeAnimationDefinitions.h"
+#include "InterchangeFbxSettings.h"
 
 #define LOCTEXT_NAMESPACE "InterchangeFbxScene"
 
@@ -382,18 +383,7 @@ namespace UE
 					NodeContainer.AddNode(AnimTrackNode);
 				};
 
-				using namespace UE::Interchange::Animation;
-				static TMap<FString, TPair<FName, EInterchangeAnimationPayLoadType>> PropertyTracks
-				{
-					{TEXT("Intensity"), {PropertyTracks::Light::Intensity, EInterchangeAnimationPayLoadType::CURVE}},
-					{TEXT("Color"), {PropertyTracks::Light::Color, EInterchangeAnimationPayLoadType::CURVE}},
-					{TEXT("bUseTemperature"), {PropertyTracks::Light::UseTemperature, EInterchangeAnimationPayLoadType::STEPCURVE}},
-					{TEXT("IntensityUnits"), {PropertyTracks::Light::IntensityUnits, EInterchangeAnimationPayLoadType::STEPCURVE}},
-					{TEXT("bHidden"), {PropertyTracks::Visibility, EInterchangeAnimationPayLoadType::STEPCURVE}},
-					{TEXT("CurrentFocalLength"), {PropertyTracks::Camera::CurrentFocalLength, EInterchangeAnimationPayLoadType::CURVE}},
-					{TEXT("CurrentAperture"), {PropertyTracks::Camera::CurrentAperture, EInterchangeAnimationPayLoadType::CURVE}},
-					{TEXT("AspectRatioAxisConstraint"), {PropertyTracks::Camera::AspectRatioAxisConstraint, EInterchangeAnimationPayLoadType::CURVE}},
-				};
+				const UInterchangeFbxSettings* InterchangeFbxSettings = GetDefault<UInterchangeFbxSettings>();
 
 				//Add all Node Attributes for the node
 				for(int32 i = 0, Count = Node->GetNodeAttributeCount(); i < Count; ++i)
@@ -413,10 +403,10 @@ namespace UE
 
 							if(PayloadKey.IsSet())
 							{
-								const char* CurveNodeName = CurveNode->GetName();
-								if(TPair<FName, EInterchangeAnimationPayLoadType>* Track = PropertyTracks.Find(CurveNodeName))
+								const char* CurveNodeName = CurveNode->GetName();								
+								if(FInterchangePropertyTracksSettings PropertyTrack = InterchangeFbxSettings->GetPropertyTrack(CurveNodeName); PropertyTrack.Property != EInterchangePropertyTracks::None)
 								{
-									AddAnimationTrackNode(Track->Key, CurveNodeName, *PayloadKey, Track->Value);
+									AddAnimationTrackNode(PropertyTrack.GetPropertyName(), CurveNodeName, *PayloadKey, PropertyTrack.Type);
 								}
 							}
 						}
@@ -443,9 +433,9 @@ namespace UE
 							if(PayloadKey.IsSet())
 							{
 								const char* CurveNodeName = CurveNode->GetName();
-								if(TPair<FName, EInterchangeAnimationPayLoadType>* Track = PropertyTracks.Find(CurveNodeName))
+								if(FInterchangePropertyTracksSettings PropertyTrack = InterchangeFbxSettings->GetPropertyTrack(CurveNodeName); PropertyTrack.Property != EInterchangePropertyTracks::None)
 								{
-									AddAnimationTrackNode(Track->Key, CurveNodeName, *PayloadKey, Track->Value);
+									AddAnimationTrackNode(PropertyTrack.GetPropertyName(), CurveNodeName, *PayloadKey, PropertyTrack.Type);
 								}
 							}
 						}
@@ -1009,4 +999,4 @@ namespace UE
 	} //ns Interchange
 }//ns UE
 
-#undef LOCTEXT_NAMESPACE
+#undef LOCTEXT_NAMESPACE
\ No newline at end of file
diff --git a/Engine/Plugins/Interchange/Runtime/Source/Parsers/Fbx/Private/InterchangeFbxSettings.cpp b/Engine/Plugins/Interchange/Runtime/Source/Parsers/Fbx/Private/InterchangeFbxSettings.cpp
new file mode 100644
--- /dev/null
+++ b/Engine/Plugins/Interchange/Runtime/Source/Parsers/Fbx/Private/InterchangeFbxSettings.cpp
@@ -0,0 +1,114 @@
+// Copyright Epic Games, Inc.All Rights Reserved.
+
+#include "InterchangeFbxSettings.h"
+
+UInterchangeFbxSettings::UInterchangeFbxSettings()
+{
+	if(HasAnyFlags(EObjectFlags::RF_ClassDefaultObject))
+	{
+		PredefinedPropertyTracks =
+		{
+			{ TEXT("Intensity"), { EInterchangePropertyTracks::LightIntensity, EInterchangeAnimationPayLoadType::CURVE }	},
+			{ TEXT("Color"), {EInterchangePropertyTracks::LightColor, EInterchangeAnimationPayLoadType::CURVE} },
+			{ TEXT("bUseTemperature"), {EInterchangePropertyTracks::LightUseTemperature, EInterchangeAnimationPayLoadType::STEPCURVE} },
+			{ TEXT("IntensityUnits"), {EInterchangePropertyTracks::LightIntensityUnits, EInterchangeAnimationPayLoadType::STEPCURVE} },
+			{ TEXT("bHidden"), {EInterchangePropertyTracks::Visibility, EInterchangeAnimationPayLoadType::STEPCURVE} },
+			{ TEXT("CurrentFocalLength"), {EInterchangePropertyTracks::CameraCurrentFocalLength, EInterchangeAnimationPayLoadType::CURVE} },
+			{ TEXT("CurrentAperture"), {EInterchangePropertyTracks::CameraCurrentAperture, EInterchangeAnimationPayLoadType::CURVE} },
+			{ TEXT("AspectRatioAxisConstraint"), {EInterchangePropertyTracks::CameraAspectRatioAxisConstraint, EInterchangeAnimationPayLoadType::CURVE} },
+		};
+	}
+}
+
+FInterchangePropertyTracksSettings UInterchangeFbxSettings::GetPropertyTrack(const FString& PropertyName) const
+{
+	FInterchangePropertyTracksSettings Result;
+	if(const FInterchangePropertyTracksSettings* Property = PredefinedPropertyTracks.Find(PropertyName))
+	{
+		Result = *Property;
+	}
+	else if(Property = CustomPropertyTracks.Find(PropertyName); Property != nullptr)
+	{
+		Result = *Property;
+	}
+
+	return Result;
+}
+
+FName FInterchangePropertyTracksSettings::GetPropertyName()
+{
+	FName Name;
+
+	switch(Property)
+	{
+	case EInterchangePropertyTracks::Visibility :
+		Name = UE::Interchange::Animation::PropertyTracks::Visibility;
+		break;
+
+	case EInterchangePropertyTracks::CameraAspectRatioAxisConstraint:
+		Name = UE::Interchange::Animation::PropertyTracks::Camera::AspectRatioAxisConstraint;
+		break;
+
+	case EInterchangePropertyTracks::CameraAutoActivate:
+		Name = UE::Interchange::Animation::PropertyTracks::Camera::AutoActivate;
+		break;
+
+	case EInterchangePropertyTracks::CameraConstrainAspectRatio:
+		Name = UE::Interchange::Animation::PropertyTracks::Camera::ConstrainAspectRatio;
+		break;
+
+	case EInterchangePropertyTracks::CameraCurrentAperture:
+		Name = UE::Interchange::Animation::PropertyTracks::Camera::CurrentAperture;
+		break;
+
+	case EInterchangePropertyTracks::CameraCurrentFocalLength:
+		Name = UE::Interchange::Animation::PropertyTracks::Camera::CurrentFocalLength;
+		break;
+
+	case EInterchangePropertyTracks::CameraCustomNearClippingPlane:
+		Name = UE::Interchange::Animation::PropertyTracks::Camera::CustomNearClippingPlane;
+		break;
+
+	case EInterchangePropertyTracks::CameraFieldOfView:
+		Name = UE::Interchange::Animation::PropertyTracks::Camera::FieldOfView;
+		break;
+
+	case EInterchangePropertyTracks::CameraMobility:
+		Name = UE::Interchange::Animation::PropertyTracks::Camera::Mobility;
+		break;
+
+	case EInterchangePropertyTracks::CameraOrthoFarClipPlane:
+		Name = UE::Interchange::Animation::PropertyTracks::Camera::OrthoFarClipPlane;
+		break;
+
+	case EInterchangePropertyTracks::CameraOrthoWidth:
+		Name = UE::Interchange::Animation::PropertyTracks::Camera::OrthoWidth;
+		break;
+
+	case EInterchangePropertyTracks::LightColor:
+		Name = UE::Interchange::Animation::PropertyTracks::Light::Color;
+		break;
+
+	case EInterchangePropertyTracks::LightIntensity:
+		Name = UE::Interchange::Animation::PropertyTracks::Light::Intensity;break;
+		break;
+
+	case EInterchangePropertyTracks::LightIntensityUnits:
+		Name = UE::Interchange::Animation::PropertyTracks::Light::IntensityUnits;
+		break;
+
+	case EInterchangePropertyTracks::LightTemperature:
+		Name = UE::Interchange::Animation::PropertyTracks::Light::Temperature;
+		break;
+
+	case EInterchangePropertyTracks::LightUseTemperature:
+		Name = UE::Interchange::Animation::PropertyTracks::Light::UseTemperature;
+		break;
+
+	case EInterchangePropertyTracks::None:
+	default:
+		checkf(false, TEXT("Unknown Interchange Property Track"));
+	}
+
+	return Name;
+}
diff --git a/Engine/Plugins/Interchange/Runtime/Source/Parsers/Fbx/Public/InterchangeFbxSettings.h b/Engine/Plugins/Interchange/Runtime/Source/Parsers/Fbx/Public/InterchangeFbxSettings.h
new file mode 100644
--- /dev/null
+++ b/Engine/Plugins/Interchange/Runtime/Source/Parsers/Fbx/Public/InterchangeFbxSettings.h
@@ -0,0 +1,47 @@
+// Copyright Epic Games, Inc.All Rights Reserved.
+
+#pragma once
+
+#include "CoreMinimal.h"
+
+#include "InterchangeAnimationDefinitions.h"
+#include "InterchangeAnimationTrackSetNode.h"
+
+
+#include "Engine/DeveloperSettings.h"
+
+#include "InterchangeFbxSettings.generated.h"
+
+USTRUCT(BlueprintType)
+struct FInterchangePropertyTracksSettings
+{
+	GENERATED_BODY()
+
+	UPROPERTY(EditAnywhere, config, Category = "FBX | Property Tracks")
+	EInterchangePropertyTracks Property = EInterchangePropertyTracks::None;
+
+	UPROPERTY(EditAnywhere, config, Category = "FBX | Property Tracks")
+	EInterchangeAnimationPayLoadType Type = EInterchangeAnimationPayLoadType::NONE;
+
+	/** Return the property name of the enum as a string, see InterchangeAnimationDefinitions.h */
+	FName GetPropertyName();
+};
+
+UCLASS(config = Interchange, meta = (DisplayName = "Interchange FBX"))
+class INTERCHANGEFBXPARSER_API UInterchangeFbxSettings : public UDeveloperSettings
+{
+	GENERATED_BODY()
+
+public:
+
+	UInterchangeFbxSettings();
+
+	/** Search for a predefined property track, if the property has been found it returns it, otherwise we search for a custom property track*/
+	FInterchangePropertyTracksSettings GetPropertyTrack(const FString& PropertyName) const;
+
+	UPROPERTY(EditAnywhere, config, Category = "FBX | Property Tracks")
+	TMap<FString, FInterchangePropertyTracksSettings> CustomPropertyTracks;
+
+private:
+	TMap<FString, FInterchangePropertyTracksSettings> PredefinedPropertyTracks;
+};
\ No newline at end of file
