[Interchange] Adding Extra Information category to Import dialog.

#rnx

diff --git a/Engine/Plugins/Interchange/Editor/Source/Pipelines/Private/InterchangeEditorPipelineDetails.cpp b/Engine/Plugins/Interchange/Editor/Source/Pipelines/Private/InterchangeEditorPipelineDetails.cpp
--- a/Engine/Plugins/Interchange/Editor/Source/Pipelines/Private/InterchangeEditorPipelineDetails.cpp
+++ b/Engine/Plugins/Interchange/Editor/Source/Pipelines/Private/InterchangeEditorPipelineDetails.cpp
@@ -119,6 +119,12 @@ void FInterchangePipelineBaseDetailsCustomization::SetConflictsInfo(TArray<FInte
 	}
 }
 
+static TMap<FString, FString> InterchangeAssetExtraInformation;
+void FInterchangePipelineBaseDetailsCustomization::SetExtraInformation(const TMap<FString, FString>& InExtraInformation)
+{
+	InterchangeAssetExtraInformation = InExtraInformation;
+}
+
 void FInterchangePipelineBaseDetailsCustomization::SetTextComboBoxWidget(IDetailPropertyRow& PropertyRow, const TSharedPtr<IPropertyHandle>& Handle, const TArray<FString>& PossibleValues)
 {
 	if (!Handle.IsValid() || !Handle->IsValidHandle() || PossibleValues.Num() < 1)
@@ -335,6 +341,36 @@ void FInterchangePipelineBaseDetailsCustomization::AddConflictSection()
 	}
 }
 
+void FInterchangePipelineBaseDetailsCustomization::AddExtraInformationSection()
+{
+	if (InterchangeAssetExtraInformation.Num() > 0)
+	{
+		const FText ExtraInformationCategoryText = LOCTEXT("ExtraInformationCategoryName", "Extra Information");
+		const FName ExtraInformationCategoryName = FName(TEXT("Extra Information"));
+		IDetailCategoryBuilder& AttributeCategoryBuilder = CachedDetailBuilder->EditCategory(ExtraInformationCategoryName, ExtraInformationCategoryText);
+		
+		for (const TPair<FString, FString>& Pair : InterchangeAssetExtraInformation)
+		{
+			FText AttributeName = FText::FromString(Pair.Key);
+			FText AttributeValue = FText::FromString(Pair.Value);
+			FDetailWidgetRow& CustomRow = AttributeCategoryBuilder.AddCustomRow(AttributeName)
+				.NameContent()
+				[
+					SNew(STextBlock)
+						.Text(AttributeName)
+						.Font(IDetailLayoutBuilder::GetDetailFont())
+				]
+				.ValueContent()
+				[
+					SNew(STextBlock)
+						.Text(AttributeValue)
+						.Font(IDetailLayoutBuilder::GetDetailFont())
+				];
+		}
+
+		InterchangeAssetExtraInformation.Reset();
+	}
+}
 
 void FInterchangePipelineBaseDetailsCustomization::CustomizeDetails(IDetailLayoutBuilder& DetailBuilder)
 {
@@ -362,9 +398,11 @@ void FInterchangePipelineBaseDetailsCustomization::CustomizeDetails(IDetailLayou
 	CachedDetailBuilder->GetCategoryNames(AllCategoryNames);
 	TMap<FName, TArray<FName>> PropertiesPerCategorys;
 	InternalGetPipelineProperties(InterchangePipeline.Get(), AllCategoryNames, PropertiesPerCategorys);
-
+	
 	AddConflictSection();
 	
+	AddExtraInformationSection();
+
 	for (const TPair<FName, TArray<FName>>& CategoryAndProperties : PropertiesPerCategorys)
 	{
 		//Category meta value Subgroup data
diff --git a/Engine/Plugins/Interchange/Editor/Source/Pipelines/Private/InterchangeEditorPipelineDetails.h b/Engine/Plugins/Interchange/Editor/Source/Pipelines/Private/InterchangeEditorPipelineDetails.h
--- a/Engine/Plugins/Interchange/Editor/Source/Pipelines/Private/InterchangeEditorPipelineDetails.h
+++ b/Engine/Plugins/Interchange/Editor/Source/Pipelines/Private/InterchangeEditorPipelineDetails.h
@@ -36,6 +36,9 @@ public:
 	// Temporary: Used to pass data to the next instance created.
 	static void SetConflictsInfo(TArray<FInterchangeConflictInfo>& ConflictInfos);
 
+	// Temporary: Used to pass ExtraInformation to the next instance created.
+	static void SetExtraInformation(const TMap<FString, FString>& InExtraInformation);
+
 	/** IDetailCustomization interface */
 	virtual void CustomizeDetails(IDetailLayoutBuilder& DetailBuilder) override;
 	/** End IDetailCustomization interface */
@@ -51,6 +54,7 @@ private:
 	void AddSubCategory(IDetailLayoutBuilder& DetailBuilder, TMap<FName, TMap<FName, TArray<FInternalPropertyData>>>& SubCategoriesPropertiesPerMainCategory);
 	void InternalGetPipelineProperties(const UInterchangePipelineBase* Pipeline, const TArray<FName>& AllCategoryNames, TMap<FName, TArray<FName>>& PropertiesPerCategorys) const;
 	void AddConflictSection();
+	void AddExtraInformationSection();
 	FReply ShowConflictDialog(FInterchangeConflictInfo ConflictInfo);
 
 	void SetTextComboBoxWidget(IDetailPropertyRow& PropertyRow, const TSharedPtr<IPropertyHandle>& Handle, const TArray<FString>& PossibleValues);
diff --git a/Engine/Plugins/Interchange/Editor/Source/Pipelines/Private/SInterchangePipelineConfigurationDialog.cpp b/Engine/Plugins/Interchange/Editor/Source/Pipelines/Private/SInterchangePipelineConfigurationDialog.cpp
--- a/Engine/Plugins/Interchange/Editor/Source/Pipelines/Private/SInterchangePipelineConfigurationDialog.cpp
+++ b/Engine/Plugins/Interchange/Editor/Source/Pipelines/Private/SInterchangePipelineConfigurationDialog.cpp
@@ -15,6 +15,7 @@
 #include "InterchangePipelineConfigurationBase.h"
 #include "InterchangeProjectSettings.h"
 #include "InterchangeTranslatorBase.h"
+#include "Nodes/InterchangeSourceNode.h"
 #include "Interfaces/IMainFrameModule.h"
 #include "Misc/App.h"
 #include "Misc/ConfigCacheIni.h"
@@ -195,6 +196,15 @@ void SInterchangePipelineConfigurationDialog::SetEditPipeline(FInterchangePipeli
 			PipelineItemToEdit->ConflictInfos = PipelineItemToEdit->Pipeline->GetConflictInfos(PipelineItemToEdit->ReimportObject, PipelineItemToEdit->Container, PipelineItemToEdit->SourceData);
 		}
 		FInterchangePipelineBaseDetailsCustomization::SetConflictsInfo(PipelineItemToEdit->ConflictInfos);
+
+		//Acquire ExtraInformation from SourceNode and pass it to FInterchangePipelineBaseDetailsCustomization:
+		TMap<FString, FString> ExtraInformation;
+		const UInterchangeSourceNode* SourceNode = UInterchangeSourceNode::GetUniqueInstance(PipelineItemToEdit->Container);
+		if (SourceNode)
+		{
+			SourceNode->GetExtraInformation(ExtraInformation);
+		}
+		FInterchangePipelineBaseDetailsCustomization::SetExtraInformation(ExtraInformation);
 	}
 	PipelineConfigurationDetailsView->SetObjects(ObjectsToEdit);
 }
diff --git a/Engine/Plugins/Interchange/Runtime/Source/Import/Private/Gltf/InterchangeGltfTranslator.cpp b/Engine/Plugins/Interchange/Runtime/Source/Import/Private/Gltf/InterchangeGltfTranslator.cpp
--- a/Engine/Plugins/Interchange/Runtime/Source/Import/Private/Gltf/InterchangeGltfTranslator.cpp
+++ b/Engine/Plugins/Interchange/Runtime/Source/Import/Private/Gltf/InterchangeGltfTranslator.cpp
@@ -407,13 +407,25 @@ bool UInterchangeGLTFTranslator::Translate( UInterchangeBaseNodeContainer& NodeC
 		SendAnalytics(TranslationResult::INPUT_FILE_NOTFOUND);
 		return false;
 	}
-	
+
 	GLTF::FFileReader GltfFileReader;
 
 	const bool bLoadImageData = false;
-	const bool bLoadMetaData = false;
+	const bool bLoadMetaData = true;
 	GltfFileReader.ReadFile( FilePath, bLoadImageData, bLoadMetaData, const_cast< UInterchangeGLTFTranslator* >( this )->GltfAsset );
 
+	UInterchangeSourceNode* SourceNode = UInterchangeSourceNode::FindOrCreateUniqueInstance(&NodeContainer);
+	SourceNode->SetExtraInformation(TEXT("File Units"), TEXT("meter"));
+	SourceNode->SetExtraInformation(TEXT("File Axis Direction"), TEXT("Y-UP (RH)"));
+	if (GltfAsset.Metadata.GeneratorName.Len() > 0)
+	{
+		SourceNode->SetExtraInformation(TEXT("Generator Name"), GltfAsset.Metadata.GeneratorName);
+	}
+	for (const GLTF::FMetadata::FExtraData& Extra : GltfAsset.Metadata.Extras)
+	{
+		SourceNode->SetExtraInformation(Extra.Name, Extra.Value);
+	}
+
 	const FString FileName = GltfAsset.Name;
 
 	//Required Extension Check:
diff --git a/Engine/Plugins/Interchange/Runtime/Source/Parsers/Fbx/Private/FbxAPI.cpp b/Engine/Plugins/Interchange/Runtime/Source/Parsers/Fbx/Private/FbxAPI.cpp
--- a/Engine/Plugins/Interchange/Runtime/Source/Parsers/Fbx/Private/FbxAPI.cpp
+++ b/Engine/Plugins/Interchange/Runtime/Source/Parsers/Fbx/Private/FbxAPI.cpp
@@ -16,6 +16,7 @@
 #include "Mesh/InterchangeMeshPayload.h"
 #endif
 #include "Nodes/InterchangeBaseNodeContainer.h"
+#include "Nodes/InterchangeSourceNode.h"
 #include "Misc/SecureHash.h"
 
 #define LOCTEXT_NAMESPACE "InterchangeFbxParser"
@@ -130,9 +131,32 @@ namespace UE
 				bool bStatus = SDKImporter->Import(SDKScene);
 
 				//We always convert scene to UE axis and units
-				FFbxConvert::ConvertScene(SDKScene, bConvertScene, bForceFrontXAxis, bConvertSceneUnit);
+				FFbxConvert::ConvertScene(SDKScene, bConvertScene, bForceFrontXAxis, bConvertSceneUnit, FileDetails.AxisDirection, FileDetails.UnitSystem);
 
 				FrameRate = FbxTime::GetFrameRate(SDKScene->GetGlobalSettings().GetTimeMode());
+				FileDetails.FrameRate = FString::Printf(TEXT("%.2f"), FrameRate);
+
+				// Get the version number of the FBX file format.
+				int32 FileMajor, FileMinor, FileRevision;
+				SDKImporter->GetFileVersion(FileMajor, FileMinor, FileRevision);
+				FileDetails.FbxFileVersion = FString::Printf(TEXT("%d.%d.%d"), FileMajor, FileMinor, FileRevision);
+
+				// Get The Creator of the FBX File.
+				FileDetails.FbxFileCreator = UTF8_TO_TCHAR(SDKImporter->GetFileHeaderInfo()->mCreator.Buffer());
+
+				FbxDocumentInfo* DocInfo = SDKImporter->GetSceneInfo();
+				if (DocInfo)
+				{
+					FString LastSavedVendor(UTF8_TO_TCHAR(DocInfo->LastSaved_ApplicationVendor.Get().Buffer()));
+					FString LastSavedAppName(UTF8_TO_TCHAR(DocInfo->LastSaved_ApplicationName.Get().Buffer()));
+					FString LastSavedAppVersion(UTF8_TO_TCHAR(DocInfo->LastSaved_ApplicationVersion.Get().Buffer()));
+
+					FileDetails.FbxFileCreatorApplication = LastSavedVendor + TEXT(" ") + LastSavedAppName + TEXT(" ") + LastSavedAppVersion;
+				}
+				else
+				{
+					FileDetails.FbxFileCreatorApplication = TEXT("");
+				}
 
 				return true;
 			}
@@ -158,6 +182,8 @@ namespace UE
 				FbxScene.AddHierarchy(SDKScene, NodeContainer, PayloadContexts);
 				FbxScene.AddAnimation(SDKScene, NodeContainer, PayloadContexts);
 				FbxScene.AddMorphTargetAnimations(SDKScene, NodeContainer, PayloadContexts, FbxMesh.GetMorphTargetAnimationsBuildingData());
+
+				ProcessExtraInformation(NodeContainer);
 			}
 
 			bool FFbxParser::FetchPayloadData(const FString& PayloadKey, const FString& PayloadFilepath)
@@ -291,6 +317,18 @@ namespace UE
 					MakeFbxObjectNameUnique(Mesh, MeshNames);
 				}
 			}
+
+			void FFbxParser::ProcessExtraInformation(UInterchangeBaseNodeContainer& NodeContainer)
+			{
+				UInterchangeSourceNode* SourceNode = UInterchangeSourceNode::FindOrCreateUniqueInstance(&NodeContainer);
+
+				SourceNode->SetExtraInformation(TEXT("File Version"), FileDetails.FbxFileVersion);
+				SourceNode->SetExtraInformation(TEXT("File Creator"), FileDetails.FbxFileCreator);
+				SourceNode->SetExtraInformation(TEXT("File Creator Application"), FileDetails.FbxFileCreatorApplication);
+				SourceNode->SetExtraInformation(TEXT("File Units"), FileDetails.UnitSystem);
+				SourceNode->SetExtraInformation(TEXT("File Axis Direction"), FileDetails.AxisDirection);
+				SourceNode->SetExtraInformation(TEXT("File Frame Rate"), FileDetails.FrameRate);
+			}
 		} //ns Private
 	} //ns Interchange
 } //ns UE
diff --git a/Engine/Plugins/Interchange/Runtime/Source/Parsers/Fbx/Private/FbxAPI.h b/Engine/Plugins/Interchange/Runtime/Source/Parsers/Fbx/Private/FbxAPI.h
--- a/Engine/Plugins/Interchange/Runtime/Source/Parsers/Fbx/Private/FbxAPI.h
+++ b/Engine/Plugins/Interchange/Runtime/Source/Parsers/Fbx/Private/FbxAPI.h
@@ -109,6 +109,7 @@ namespace UE
 			private:
 
 				void CleanupFbxData();
+				void ProcessExtraInformation(UInterchangeBaseNodeContainer& NodeContainer);
 
 				TWeakObjectPtr<UInterchangeResultsContainer> ResultsContainer;
 				FbxManager* SDKManager = nullptr;
@@ -127,6 +128,17 @@ namespace UE
 				bool bConvertScene = true;
 				bool bForceFrontXAxis = false;
 				bool bConvertSceneUnit = true;
+
+				struct FileDetails
+				{
+					FString FbxFileVersion;
+					FString FbxFileCreator;
+					FString FbxFileCreatorApplication;
+					FString UnitSystem;
+					FString AxisDirection;
+					FString FrameRate;
+				} FileDetails;
+				
 			};
 		}//ns Private
 	}//ns Interchange
diff --git a/Engine/Plugins/Interchange/Runtime/Source/Parsers/Fbx/Private/FbxConvert.cpp b/Engine/Plugins/Interchange/Runtime/Source/Parsers/Fbx/Private/FbxConvert.cpp
--- a/Engine/Plugins/Interchange/Runtime/Source/Parsers/Fbx/Private/FbxConvert.cpp
+++ b/Engine/Plugins/Interchange/Runtime/Source/Parsers/Fbx/Private/FbxConvert.cpp
@@ -13,7 +13,49 @@ namespace UE
 	{
 		namespace Private
 		{
-			void FFbxConvert::ConvertScene(FbxScene* SDKScene, const bool bConvertScene, const bool bForceFrontXAxis, const bool bConvertSceneUnit)
+			FString GetFileAxisDirection(FbxAxisSystem FileAxisSystem)
+			{
+				FString AxisDirection;
+				int32 Sign = 1;
+				switch (FileAxisSystem.GetUpVector(Sign))
+				{
+					case FbxAxisSystem::eXAxis:
+						{
+							AxisDirection += TEXT("X");
+						}
+						break;
+					case FbxAxisSystem::eYAxis:
+						{
+							AxisDirection += TEXT("Y");
+						}
+						break;
+					case FbxAxisSystem::eZAxis:
+						{
+							AxisDirection += TEXT("Z");
+						}
+						break;
+				}
+
+				//Negative sign mean down instead of up
+				AxisDirection += Sign == 1 ? TEXT("-UP") : TEXT("-DOWN");
+
+				switch (FileAxisSystem.GetCoorSystem())
+				{
+					case FbxAxisSystem::eLeftHanded:
+						{
+							AxisDirection += TEXT(" (LH)");
+						}
+						break;
+					case FbxAxisSystem::eRightHanded:
+						{
+							AxisDirection += TEXT(" (RH)");
+						}
+						break;
+				}
+				return AxisDirection;
+			}
+
+			void FFbxConvert::ConvertScene(FbxScene* SDKScene, const bool bConvertScene, const bool bForceFrontXAxis, const bool bConvertSceneUnit, FString& FileSystemDirection, FString& FileUnitSystem)
 			{
 				if (!ensure(SDKScene))
 				{
@@ -42,13 +84,11 @@ namespace UE
 					}
 				}
 
+				FbxAxisSystem FileAxisSystem = SDKScene->GetGlobalSettings().GetAxisSystem();
+				FileSystemDirection = GetFileAxisDirection(FileAxisSystem);
 
 				if (bConvertScene)
 				{
-					//Set the original file information
-					FbxAxisSystem FileAxisSystem = SDKScene->GetGlobalSettings().GetAxisSystem();
-
-
 					//UE is: z up, front x, left handed
 					FbxAxisSystem::EUpVector UpVector = FbxAxisSystem::EUpVector::eZAxis;
 					FbxAxisSystem::EFrontVector FrontVector = (FbxAxisSystem::EFrontVector)(bForceFrontXAxis ? FbxAxisSystem::eParityEven : -FbxAxisSystem::eParityOdd);
@@ -62,10 +102,12 @@ namespace UE
 					}
 				}
 
+				FbxSystemUnit OriginalFileUnitSystem = SDKScene->GetGlobalSettings().GetSystemUnit();
+				FileUnitSystem = FString(UTF8_TO_TCHAR(OriginalFileUnitSystem.GetScaleFactorAsString(false).Buffer()));
+
 				if (bConvertSceneUnit)
 				{
-					FbxSystemUnit FileUnitSystem = SDKScene->GetGlobalSettings().GetSystemUnit();
-					if (FileUnitSystem != FbxSystemUnit::cm)
+					if (OriginalFileUnitSystem != FbxSystemUnit::cm)
 					{
 						FbxSystemUnit::cm.ConvertScene(SDKScene);
 					}
diff --git a/Engine/Plugins/Interchange/Runtime/Source/Parsers/Fbx/Private/FbxConvert.h b/Engine/Plugins/Interchange/Runtime/Source/Parsers/Fbx/Private/FbxConvert.h
--- a/Engine/Plugins/Interchange/Runtime/Source/Parsers/Fbx/Private/FbxConvert.h
+++ b/Engine/Plugins/Interchange/Runtime/Source/Parsers/Fbx/Private/FbxConvert.h
@@ -156,7 +156,7 @@ namespace UE
 				/**
 				 * Convert a fbx scene 
 				 */
-				static void ConvertScene(FbxScene* SDKScene, const bool bConvertScene, const bool bForceFrontXAxis, const bool bConvertSceneUnit);
+				static void ConvertScene(FbxScene* SDKScene, const bool bConvertScene, const bool bForceFrontXAxis, const bool bConvertSceneUnit, FString& FileSystemDirection, FString& FileUnitSystem);
 				
 				/** Scene Conversion API End */
 				//////////////////////////////////////////////////////////////////////////
diff --git a/Engine/Plugins/Interchange/Runtime/Source/Pipelines/Private/InterchangeGenericAnimationPipeline.cpp b/Engine/Plugins/Interchange/Runtime/Source/Pipelines/Private/InterchangeGenericAnimationPipeline.cpp
--- a/Engine/Plugins/Interchange/Runtime/Source/Pipelines/Private/InterchangeGenericAnimationPipeline.cpp
+++ b/Engine/Plugins/Interchange/Runtime/Source/Pipelines/Private/InterchangeGenericAnimationPipeline.cpp
@@ -460,12 +460,13 @@ void UInterchangeGenericAnimationPipeline::CreateAnimSequenceFactoryNode(UInterc
 
 	if (bImportBoneTracks)
 	{
-		if (const UInterchangeSourceNode* SourceNode = UInterchangeSourceNode::GetUniqueInstance(BaseNodeContainer))
+		int32 Numerator, Denominator;
+		const UInterchangeSourceNode* SourceNode = UInterchangeSourceNode::GetUniqueInstance(BaseNodeContainer);
+		if (SourceNode && SourceNode->GetCustomSourceFrameRateNumerator(Numerator) && SourceNode->GetCustomSourceFrameRateDenominator(Denominator))
 		{
-			int32 Numerator, Denominator;
-			if (!bUse30HzToBakeBoneAnimation && CustomBoneAnimationSampleRate == 0 && SourceNode->GetCustomSourceFrameRateNumerator(Numerator))
+			if (!bUse30HzToBakeBoneAnimation && CustomBoneAnimationSampleRate == 0)
 			{
-				if (SourceNode->GetCustomSourceFrameRateDenominator(Denominator) && Denominator > 0 && Numerator > 0)
+				if (Denominator > 0 && Numerator > 0)
 				{
 					SampleRate = static_cast<double>(Numerator) / static_cast<double>(Denominator);
 				}
diff --git a/Engine/Source/Runtime/Interchange/Core/Private/Nodes/InterchangeSourceNode.cpp b/Engine/Source/Runtime/Interchange/Core/Private/Nodes/InterchangeSourceNode.cpp
--- a/Engine/Source/Runtime/Interchange/Core/Private/Nodes/InterchangeSourceNode.cpp
+++ b/Engine/Source/Runtime/Interchange/Core/Private/Nodes/InterchangeSourceNode.cpp
@@ -19,10 +19,17 @@ namespace UE::Interchange::SourceNode
 		static FString StaticUid = TEXT("__SourceNode__");
 		return StaticUid;
 	}
+
+	const FString& GetExtraInformationKey()
+	{
+		static FString ExtraInformationKey(TEXT("__ExtraInformation__Key"));
+		return ExtraInformationKey;
+	}
 }
 
 UInterchangeSourceNode::UInterchangeSourceNode()
 {
+	ExtraInformation.Initialize(Attributes.ToSharedRef(), UE::Interchange::SourceNode::GetExtraInformationKey());
 }
 
 void UInterchangeSourceNode::InitializeSourceNode(const FString& UniqueID, const FString& DisplayLabel)
@@ -127,3 +134,17 @@ bool UInterchangeSourceNode::SetCustomImportUnusedMaterial(const bool& Attribute
 	IMPLEMENT_NODE_ATTRIBUTE_SETTER_NODELEGATE(ImportUnusedMaterial, bool);
 }
 
+bool UInterchangeSourceNode::SetExtraInformation(const FString& Name, const FString& Value)
+{
+	return ExtraInformation.SetKeyValue(Name, Value);
+}
+
+bool UInterchangeSourceNode::RemoveExtraInformation(const FString& Name)
+{
+	return ExtraInformation.RemoveKey(Name);
+}
+
+void UInterchangeSourceNode::GetExtraInformation(TMap<FString, FString>& OutExtraInformation) const
+{
+	OutExtraInformation = ExtraInformation.ToMap();
+}
\ No newline at end of file
diff --git a/Engine/Source/Runtime/Interchange/Core/Public/Nodes/InterchangeSourceNode.h b/Engine/Source/Runtime/Interchange/Core/Public/Nodes/InterchangeSourceNode.h
--- a/Engine/Source/Runtime/Interchange/Core/Public/Nodes/InterchangeSourceNode.h
+++ b/Engine/Source/Runtime/Interchange/Core/Public/Nodes/InterchangeSourceNode.h
@@ -106,6 +106,19 @@ public:
 	UFUNCTION(BlueprintCallable, Category = "Interchange | Node | Source")
 	INTERCHANGECORE_API bool SetCustomImportUnusedMaterial(const bool& AttributeValue);
 
+
+	/** Set Extra Information that we want to show in the Config Panel (such as File Information). */
+	UFUNCTION(BlueprintCallable, Category = "Interchange | Node | Source")
+	INTERCHANGECORE_API bool SetExtraInformation(const FString& Name, const FString& Value);
+
+	/** Remove Extra Information that we dont want to show in the Config Panel. */
+	UFUNCTION(BlueprintCallable, Category = "Interchange | Node | Source")
+	INTERCHANGECORE_API bool RemoveExtraInformation(const FString& Name);
+
+	/** Get Extra Information that we want to show in the Config Panel (such as File Information). */
+	UFUNCTION(BlueprintCallable, Category = "Interchange | Node | Source")
+	INTERCHANGECORE_API void GetExtraInformation(TMap<FString, FString>& OutExtraInformation) const;
+
 private:
 	const UE::Interchange::FAttributeKey Macro_CustomSourceFrameRateNumeratorKey = UE::Interchange::FAttributeKey(TEXT("SourceFrameRateNumerator"));
 	const UE::Interchange::FAttributeKey Macro_CustomSourceFrameRateDenominatorKey = UE::Interchange::FAttributeKey(TEXT("SourceFrameRateDenominator"));
@@ -114,5 +127,8 @@ private:
 	const UE::Interchange::FAttributeKey Macro_CustomAnimatedTimeStartKey = UE::Interchange::FAttributeKey(TEXT("AnimatedTimeStart"));
 	const UE::Interchange::FAttributeKey Macro_CustomAnimatedTimeEndKey = UE::Interchange::FAttributeKey(TEXT("AnimatedTimeEnd"));
 	const UE::Interchange::FAttributeKey Macro_CustomImportUnusedMaterialKey = UE::Interchange::FAttributeKey(TEXT("ImportUnusedMaterial"));
+
+	// Extra InformationTo show in the Config Panel.
+	UE::Interchange::TMapAttributeHelper<FString, FString> ExtraInformation; 
 };
 
