[Interchange] Implement support for meta-data on actors

#rnx

diff --git a/Engine/Plugins/Interchange/Runtime/Source/Import/Public/InterchangeAssetUserData.h b/Engine/Plugins/Interchange/Runtime/Source/Import/Public/InterchangeAssetUserData.h
new file mode 100644
--- /dev/null
+++ b/Engine/Plugins/Interchange/Runtime/Source/Import/Public/InterchangeAssetUserData.h
@@ -0,0 +1,19 @@
+// Copyright Epic Games, Inc. All Rights Reserved.
+
+#pragma once
+
+#include "CoreMinimal.h"
+#include "Engine/AssetUserData.h"
+
+#include "InterchangeAssetUserData.generated.h"
+
+/** Asset user data that can be used with Interchange on Actors and other objects  */
+UCLASS(BlueprintType, meta = (ScriptName = "InterchangeUserData", DisplayName = "Interchange User Data"))
+class INTERCHANGEIMPORT_API UInterchangeAssetUserData : public UAssetUserData
+{
+	GENERATED_BODY()
+
+public:
+	UPROPERTY(VisibleAnywhere, BlueprintReadWrite, AdvancedDisplay, Category = "Interchange User Data", meta = (ScriptName = "Metadata", DisplayName = "Metadata"))
+	TMap<FString, FString> MetaData;
+};
\ No newline at end of file
diff --git a/Engine/Plugins/Interchange/Runtime/Source/Pipelines/Private/InterchangeGenericAssetsPipeline.cpp b/Engine/Plugins/Interchange/Runtime/Source/Pipelines/Private/InterchangeGenericAssetsPipeline.cpp
--- a/Engine/Plugins/Interchange/Runtime/Source/Pipelines/Private/InterchangeGenericAssetsPipeline.cpp
+++ b/Engine/Plugins/Interchange/Runtime/Source/Pipelines/Private/InterchangeGenericAssetsPipeline.cpp
@@ -29,6 +29,7 @@
 #include "InterchangeSourceData.h"
 #include "InterchangeStaticMeshFactoryNode.h"
 #include "InterchangeStaticMeshLodDataNode.h"
+#include "InterchangeAssetUserData.h"
 #include "Misc/PackageName.h"
 #include "Misc/Paths.h"
 #include "Nodes/InterchangeBaseNode.h"
@@ -47,6 +48,7 @@
 #include "Widgets/Layout/SSeparator.h"
 #include "Widgets/Text/STextBlock.h"
 #include "Widgets/Views/STableRow.h"
+#include "GameFramework/Actor.h"
 
 #include UE_INLINE_GENERATED_CPP_BY_NAME(InterchangeGenericAssetsPipeline)
 
@@ -1177,13 +1179,10 @@ void UInterchangeGenericAssetsPipeline::ExecutePostImportPipeline(const UInterch
 		AnimationPipeline->ScriptedExecutePostImportPipeline(InBaseNodeContainer, NodeKey, CreatedAsset, bIsAReimport);
 	}
 
-#if WITH_EDITORONLY_DATA
-	AddPackageMetaData(CreatedAsset, InBaseNodeContainer->GetNode(NodeKey));
-#endif
+	AddMetaData(CreatedAsset, InBaseNodeContainer->GetNode(NodeKey));
 }
 
-#if WITH_EDITORONLY_DATA
-void UInterchangeGenericAssetsPipeline::AddPackageMetaData(UObject* CreatedAsset, const UInterchangeBaseNode* Node)
+void UInterchangeGenericAssetsPipeline::AddMetaData(UObject* CreatedAsset, const UInterchangeBaseNode* Node)
 {
 	if (!CreatedAsset || !Node)
 	{
@@ -1192,27 +1191,20 @@ void UInterchangeGenericAssetsPipeline::AddPackageMetaData(UObject* CreatedAsset
 
 	const FString InterchangeMetaDataPrefix = TEXT("INTERCHANGE.");
 
-	//Add UObject package meta data
-	if (UMetaData* MetaData = CreatedAsset->GetOutermost()->GetMetaData())
+	bool bAssetIsActor = CreatedAsset->IsA<AActor>();
+#if WITH_EDITORONLY_DATA
+	UMetaData* MetaData = CreatedAsset->GetOutermost()->GetMetaData();
+#endif
+
+	bool bProcessMetaData = 
+#if WITH_EDITORONLY_DATA
+		MetaData ||
+#endif
+		bAssetIsActor;
+
+	TMap<FString, FString> MetaDataMap;
+	if (bProcessMetaData)
 	{
-		//Cleanup existing INTERCHANGE_ prefix metadata name for this object (in case we re-import)
-		{
-			TArray<FName> InterchangeMetaDataKeys;
-			if(TMap<FName, FString>* MetaDataMapPtr = MetaData->GetMapForObject(CreatedAsset))
-			{
-				for (const TPair<FName, FString>& ObjectMetadata : *MetaDataMapPtr)
-				{
-					if (ObjectMetadata.Key.ToString().StartsWith(InterchangeMetaDataPrefix))
-					{
-						InterchangeMetaDataKeys.Add(ObjectMetadata.Key);
-					}
-				}
-				for (const FName& MetaDataKey : InterchangeMetaDataKeys)
-				{
-					MetaData->RemoveValue(CreatedAsset, MetaDataKey);
-				}
-			}
-		}
 		TArray<FInterchangeUserDefinedAttributeInfo> UserAttributeInfos;
 		UInterchangeUserDefinedAttributesAPI::GetUserDefinedAttributeInfos(Node, UserAttributeInfos);
 		//We must convert all different type to String since meta data only support string
@@ -1230,7 +1222,7 @@ void UInterchangeGenericAssetsPipeline::AddPackageMetaData(UObject* CreatedAsset
 				case UE::Interchange::EAttributeTypes::Bool:
 				{
 					bool Value = false;
-					if(UInterchangeUserDefinedAttributesAPI::GetUserDefinedAttribute(Node, UserAttributeInfo.Name, Value, PayloadKey))
+					if (UInterchangeUserDefinedAttributesAPI::GetUserDefinedAttribute(Node, UserAttributeInfo.Name, Value, PayloadKey))
 					{
 						MetaDataValue = UE::Interchange::AttributeValueToString(Value);
 					}
@@ -1401,27 +1393,89 @@ void UInterchangeGenericAssetsPipeline::AddPackageMetaData(UObject* CreatedAsset
 			}
 			if (MetaDataValue.IsSet())
 			{
-				const FString& MetaDataStringValue = MetaDataValue.GetValue();
-				FString MetaDataKeyString = InterchangeMetaDataPrefix + UserAttributeInfo.Name;
-				if (MetaDataKeyString.Len() < NAME_SIZE)
+				FString MetaDataKeyString = UserAttributeInfo.Name;
+				MetaDataMap.Add(MetaDataKeyString, MetaDataValue.GetValue());
+			}
+		}
+	}
+
+	if (bAssetIsActor)
+	{
+		if (AActor* Actor = Cast<AActor>(CreatedAsset))
+		{
+			if (IInterface_AssetUserData* AssetUserData = Cast<IInterface_AssetUserData>(Actor->GetRootComponent()))
+			{
+				UInterchangeAssetUserData* InterchangeUserData = AssetUserData->GetAssetUserData< UInterchangeAssetUserData >();
+
+				if (!InterchangeUserData)
 				{
-					const FName& MetaDataKey = FName(MetaDataKeyString);
-					//SetValue either add the key or set the new value
-					MetaData->SetValue(CreatedAsset, MetaDataKey, *MetaDataStringValue);
+					if (MetaDataMap.Num() > 0)
+					{
+						InterchangeUserData = NewObject<UInterchangeAssetUserData>(Actor->GetRootComponent(), NAME_None, RF_Public | RF_Transactional);
+						AssetUserData->AddAssetUserData(InterchangeUserData);
+					}
 				}
-				else if(!bHasNotify_MetaDataAttributeKeyNameTooLong)
+				else
+				{
+					InterchangeUserData->MetaData.Reset();
+				}
+
+				for (const TPair<FString, FString>& MetaDataPair : MetaDataMap)
+				{
+					InterchangeUserData->MetaData.Add(MetaDataPair.Key, MetaDataPair.Value);
+				}
+
+				if (InterchangeUserData)
+				{
+					InterchangeUserData->MetaData.KeySort(TLess<FString>());
+				}
+			}
+		}
+	}
+
+#if WITH_EDITORONLY_DATA
+	//Add UObject package meta data
+	if (MetaData)
+	{
+		//Cleanup existing INTERCHANGE_ prefix metadata name for this object (in case we re-import)
+		{
+			TArray<FName> InterchangeMetaDataKeys;
+			if(TMap<FName, FString>* MetaDataMapPtr = MetaData->GetMapForObject(CreatedAsset))
+			{
+				for (const TPair<FName, FString>& ObjectMetadata : *MetaDataMapPtr)
 				{
-					bHasNotify_MetaDataAttributeKeyNameTooLong = true;
-					//We cannot add this meta data, notify the user the meta attribute key name is too long
-					UInterchangeResultWarning_Generic* Message = AddMessage<UInterchangeResultWarning_Generic>();
-					Message->Text = FText::Format(NSLOCTEXT("UInterchangeGenericAssetsPipeline", "MetadataKeyNameTooLong", "One or more metadata key(s) cannot be added because the name exceeds the maximum length ({0}) allowed by the engine. The metadata is provided by the source file node's custom attributes."),
-						FText::AsNumber(NAME_SIZE));
+					if (ObjectMetadata.Key.ToString().StartsWith(InterchangeMetaDataPrefix))
+					{
+						InterchangeMetaDataKeys.Add(ObjectMetadata.Key);
+					}
+				}
+				for (const FName& MetaDataKey : InterchangeMetaDataKeys)
+				{
+					MetaData->RemoveValue(CreatedAsset, MetaDataKey);
 				}
 			}
 		}
+		for (const TPair<FString, FString>& MetaDataPair : MetaDataMap)
+		{
+			FString MetaDataKeyStr = InterchangeMetaDataPrefix + MetaDataPair.Key;
+			if (MetaDataKeyStr.Len() < NAME_SIZE)
+			{
+				const FName& MetaDataKey = FName(MetaDataKeyStr);
+				//SetValue either add the key or set the new value
+				MetaData->SetValue(CreatedAsset, MetaDataKey, *MetaDataPair.Value);
+			}
+			else if (!bHasNotify_MetaDataAttributeKeyNameTooLong)
+			{
+				bHasNotify_MetaDataAttributeKeyNameTooLong = true;
+				//We cannot add this meta data, notify the user the meta attribute key name is too long
+				UInterchangeResultWarning_Generic* Message = AddMessage<UInterchangeResultWarning_Generic>();
+				Message->Text = FText::Format(NSLOCTEXT("UInterchangeGenericAssetsPipeline", "MetadataKeyNameTooLong", "One or more metadata key(s) cannot be added because the name exceeds the maximum length ({0}) allowed by the engine. The metadata is provided by the source file node's custom attributes."),
+					FText::AsNumber(NAME_SIZE));
+			}
+		}
 	}
+#endif
 }
-#endif // WITH_EDITORONLY_DATA
 
 void UInterchangeGenericAssetsPipeline::SetReimportSourceIndex(UClass* ReimportObjectClass, const int32 SourceFileIndex)
 {
diff --git a/Engine/Plugins/Interchange/Runtime/Source/Pipelines/Private/InterchangeGenericScenesPipeline.cpp b/Engine/Plugins/Interchange/Runtime/Source/Pipelines/Private/InterchangeGenericScenesPipeline.cpp
--- a/Engine/Plugins/Interchange/Runtime/Source/Pipelines/Private/InterchangeGenericScenesPipeline.cpp
+++ b/Engine/Plugins/Interchange/Runtime/Source/Pipelines/Private/InterchangeGenericScenesPipeline.cpp
@@ -443,6 +443,8 @@ void UInterchangeGenericLevelPipeline::ExecuteSceneNodePreImport(const FTransfor
 		return;
 	}
 
+	UInterchangeUserDefinedAttributesAPI::DuplicateAllUserDefinedAttribute(SceneNode, ActorFactoryNode, false);
+
 	FString NodeUid = SceneNode->GetUniqueID() + (bRootJointNode ? TEXT("_SkeletonNode") : TEXT(""));
 	FString FactoryNodeUid = UInterchangeFactoryBaseNode::BuildFactoryNodeUid(NodeUid);
 
diff --git a/Engine/Plugins/Interchange/Runtime/Source/Pipelines/Public/InterchangeGenericAssetsPipeline.h b/Engine/Plugins/Interchange/Runtime/Source/Pipelines/Public/InterchangeGenericAssetsPipeline.h
--- a/Engine/Plugins/Interchange/Runtime/Source/Pipelines/Public/InterchangeGenericAssetsPipeline.h
+++ b/Engine/Plugins/Interchange/Runtime/Source/Pipelines/Public/InterchangeGenericAssetsPipeline.h
@@ -162,12 +162,10 @@ private:
 	 */
 	void ImplementUseSourceNameForAssetOption(UInterchangeBaseNodeContainer* InBaseNodeContainer, const TArray<UInterchangeSourceData*>& InSourceDatas);
 
-#if WITH_EDITORONLY_DATA
 	/**
-	 * Adds the user defined attributes (UInterchangeUserDefinedAttributesAPI) to the package meta data (UMetaData)
+	 * Adds the user defined attributes (UInterchangeUserDefinedAttributesAPI) to the package meta data (UMetaData) for WITH_EDITORONLY_DATA, and add UAssetUserData for AActors.
 	 */
-	void AddPackageMetaData(UObject* CreatedAsset, const UInterchangeBaseNode* Node);
-#endif // WITH_EDITORONLY_DATA
+	void AddMetaData(UObject* CreatedAsset, const UInterchangeBaseNode* Node);
 
 	struct FMaterialConflictData
 	{
