[Interchange glTF] Add support for 'extras'.

#rnx

diff --git a/Engine/Plugins/Interchange/Runtime/Source/Import/Private/Gltf/InterchangeGltfAnimation.cpp b/Engine/Plugins/Interchange/Runtime/Source/Import/Private/Gltf/InterchangeGltfAnimation.cpp
--- a/Engine/Plugins/Interchange/Runtime/Source/Import/Private/Gltf/InterchangeGltfAnimation.cpp
+++ b/Engine/Plugins/Interchange/Runtime/Source/Import/Private/Gltf/InterchangeGltfAnimation.cpp
@@ -15,6 +15,7 @@
 #include "InterchangeSceneNode.h"
 #include "InterchangeMeshNode.h"
 #include "InterchangeAnimationTrackSetNode.h"
+#include "Nodes/InterchangeUserDefinedAttribute.h"
 
 namespace UE::Interchange::Gltf::Private
 {
@@ -744,6 +745,11 @@ namespace UE::Interchange::Gltf::Private
 
 					NodeContainer.AddNode(TrackNode);
 
+					for (const TPair<FString, FString>& Extra : GLTFAnimation.Extras)
+					{
+						UInterchangeUserDefinedAttributesAPI::CreateUserDefinedAttribute(TrackNode, Extra.Key, Extra.Value, TOptional<FString>());
+					}
+
 					RootJointIndexToTrackNodeMap.Add(SkeletonNodeUid, TrackNode);
 				}
 
diff --git a/Engine/Plugins/Interchange/Runtime/Source/Import/Private/Gltf/InterchangeGltfTranslator.cpp b/Engine/Plugins/Interchange/Runtime/Source/Import/Private/Gltf/InterchangeGltfTranslator.cpp
--- a/Engine/Plugins/Interchange/Runtime/Source/Import/Private/Gltf/InterchangeGltfTranslator.cpp
+++ b/Engine/Plugins/Interchange/Runtime/Source/Import/Private/Gltf/InterchangeGltfTranslator.cpp
@@ -24,6 +24,7 @@
 #include "InterchangeTranslatorHelper.h"
 #include "InterchangeVariantSetNode.h"
 #include "Nodes/InterchangeSourceNode.h"
+#include "Nodes/InterchangeUserDefinedAttribute.h"
 
 #include "Sections/MovieScene3DTransformSection.h"
 #include "Texture/InterchangeImageWrapperTranslator.h"
@@ -226,6 +227,11 @@ void UInterchangeGLTFTranslator::HandleGltfNode( UInterchangeBaseNodeContainer&
 	InterchangeSceneNode->SetAssetName(GltfNode.UniqueId);
 	NodeContainer.AddNode( InterchangeSceneNode );
 
+	for (const TPair<FString, FString>& Extra : GltfNode.Extras)
+	{
+		UInterchangeUserDefinedAttributesAPI::CreateUserDefinedAttribute(InterchangeSceneNode, Extra.Key, Extra.Value, TOptional<FString>());
+	}
+
 	NodeUidMap.Add( &GltfNode, NodeUid );
 
 	FTransform Transform = GltfNode.Transform;
@@ -330,7 +336,13 @@ void UInterchangeGLTFTranslator::HandleGltfNode( UInterchangeBaseNodeContainer&
 			if ( GltfAsset.Cameras.IsValidIndex( GltfNode.CameraIndex ) )
 			{
 				const FString CameraNodeUid = TEXT("\\Camera\\") + GltfAsset.Cameras[ GltfNode.CameraIndex ].UniqueId;
+				const FString Prefix = GltfAsset.Cameras[GltfNode.CameraIndex].Name + TEXT("_");
 				InterchangeSceneNode->SetCustomAssetInstanceUid( CameraNodeUid );
+
+				for (const TPair<FString, FString>& Extra : GltfAsset.Cameras[GltfNode.CameraIndex].Extras)
+				{
+					UInterchangeUserDefinedAttributesAPI::CreateUserDefinedAttribute(InterchangeSceneNode, Prefix + Extra.Key, Extra.Value, TOptional<FString>());
+				}
 			}
 			break;
 		}
@@ -546,6 +558,15 @@ bool UInterchangeGLTFTranslator::Translate( UInterchangeBaseNodeContainer& NodeC
 
 			TextureNode->SetCustomWrapU( UE::Interchange::Gltf::Private::ConvertWrap( GltfTexture.Sampler.WrapS ) );
 			TextureNode->SetCustomWrapV( UE::Interchange::Gltf::Private::ConvertWrap( GltfTexture.Sampler.WrapT ) );
+
+			for (const TPair<FString, FString>& Extra : GltfTexture.Source.Extras)
+			{
+				UInterchangeUserDefinedAttributesAPI::CreateUserDefinedAttribute(TextureNode, Extra.Key, Extra.Value, TOptional<FString>());
+			}
+			for (const TPair<FString, FString>& Extra : GltfTexture.Extras)
+			{
+				UInterchangeUserDefinedAttributesAPI::CreateUserDefinedAttribute(TextureNode, Extra.Key, Extra.Value, TOptional<FString>());
+			}
 		}
 	}
 
@@ -595,6 +616,11 @@ bool UInterchangeGLTFTranslator::Translate( UInterchangeBaseNodeContainer& NodeC
 			UInterchangeShaderGraphNode* ShaderGraphNode = UInterchangeShaderGraphNode::Create(&NodeContainer, GltfMaterial.UniqueId);
 			ShaderGraphNode->SetDisplayLabel(GltfMaterial.Name);
 
+			for (const TPair<FString, FString>& Extra : GltfMaterial.Extras)
+			{
+				UInterchangeUserDefinedAttributesAPI::CreateUserDefinedAttribute(ShaderGraphNode, Extra.Key, Extra.Value, TOptional<FString>());
+			}
+
 			UE::Interchange::GLTFMaterials::HandleGltfMaterial(NodeContainer, GltfMaterial, GltfAsset.Textures, ShaderGraphNode);
 			
 			++MaterialIndex;
@@ -734,6 +760,11 @@ bool UInterchangeGLTFTranslator::Translate( UInterchangeBaseNodeContainer& NodeC
 			SceneNode->InitializeNode( SceneNodeUid, SceneName, EInterchangeNodeContainerType::TranslatedScene );
 			NodeContainer.AddNode( SceneNode );
 
+			for (const TPair<FString, FString>& Extra : GltfScene.Extras)
+			{
+				UInterchangeUserDefinedAttributesAPI::CreateUserDefinedAttribute(SceneNode, Extra.Key, Extra.Value, TOptional<FString>());
+			}
+
 			//All scene node should have a valid local transform
 			SceneNode->SetCustomLocalTransform(&NodeContainer, FTransform::Identity);
 
@@ -1245,6 +1276,7 @@ void UInterchangeGLTFTranslator::HandleGltfSkeletons(UInterchangeBaseNodeContain
 				SkeletalMeshNode->SetSkeletonDependencyUid(*SkeletonNodeUid);
 			}
 
+			TSet<int32> SkinIndices;
 			//generate payload key:
 			//of template:
 			//"LexToString(SkinnedMeshNode.MeshIndex | (SkinnedMeshNode.Skindex << 16))":"LexToString(SkinnedMeshNode.MeshIndex | (SkinnedMeshNode.Skindex << 16))".....
@@ -1258,9 +1290,22 @@ void UInterchangeGLTFTranslator::HandleGltfSkeletons(UInterchangeBaseNodeContain
 				}
 				
 				Payload += LexToString(SkinnedMeshNode.MeshIndex | (SkinnedMeshNode.Skindex << 16));
+				SkinIndices.Add(SkinnedMeshNode.Skindex);
 			}
 			SkeletalMeshNode->SetPayLoadKey(Payload, EInterchangeMeshPayLoadType::SKELETAL);
 
+			for (const int32& SkinIndex : SkinIndices)
+			{
+				if (GltfAsset.Skins.IsValidIndex(SkinIndex))
+				{
+					const FString Prefix = GltfAsset.Skins[SkinIndex].Name + TEXT("_");
+					for (const TPair<FString, FString>& Extra : GltfAsset.Skins[SkinIndex].Extras)
+					{
+						UInterchangeUserDefinedAttributesAPI::CreateUserDefinedAttribute(SkeletalMeshNode, Prefix + Extra.Key, Extra.Value, TOptional<FString>());
+					}
+				}
+			}
+
 			//set the mesh actor node's custom asset instance uid to the new duplicated mesh
 			//if there are more than one skins, then choose the topmost (root node of the collection, top most in a hierarchical tree term) occurance of SkinnedMeshIndex
 			int32 MeshActorNodeIndex = UE::Interchange::Gltf::Private::GetRootNodeIndex(GltfAsset, RootJointToSkinnedMeshNodes.Value);
@@ -1306,6 +1351,22 @@ UInterchangeMeshNode* UInterchangeGLTFTranslator::HandleGltfMesh(UInterchangeBas
 	UInterchangeMeshNode* MeshNode = NewObject< UInterchangeMeshNode >(&NodeContainer);
 	MeshNode->InitializeNode(MeshNodeUid, MeshName, EInterchangeNodeContainerType::TranslatedAsset);
 
+	int32 PrimitiveIndex = 0;
+	for (const GLTF::FPrimitive& Primitive : GltfMesh.Primitives)
+	{
+		FString Prefix = TEXT("Primitive[") + FString::FromInt(PrimitiveIndex) + TEXT("]_");
+		for (const TPair<FString, FString>& Extra : Primitive.Extras)
+		{
+			UInterchangeUserDefinedAttributesAPI::CreateUserDefinedAttribute(MeshNode, Prefix + Extra.Key, Extra.Value, TOptional<FString>());
+		}
+
+		PrimitiveIndex++;
+	}
+	for (const TPair<FString, FString>& Extra : GltfMesh.Extras)
+	{
+		UInterchangeUserDefinedAttributesAPI::CreateUserDefinedAttribute(MeshNode, Extra.Key, Extra.Value, TOptional<FString>());
+	}
+
 	//Generate Mesh Payload:
 	FString PayloadKey = LexToString(MeshIndex);
 	MeshNode->SetPayLoadKey(PayloadKey, EInterchangeMeshPayLoadType::STATIC);
diff --git a/Engine/Plugins/Interchange/Runtime/Source/Parsers/GLTFCore/Private/GLTF/GLTFReader.cpp b/Engine/Plugins/Interchange/Runtime/Source/Parsers/GLTFCore/Private/GLTF/GLTFReader.cpp
--- a/Engine/Plugins/Interchange/Runtime/Source/Parsers/GLTFCore/Private/GLTF/GLTFReader.cpp
+++ b/Engine/Plugins/Interchange/Runtime/Source/Parsers/GLTFCore/Private/GLTF/GLTFReader.cpp
@@ -411,6 +411,15 @@ namespace GLTF
 			}
 		}
 
+		if (Object.HasField(TEXT("extras")))
+		{
+			const TSharedPtr<FJsonObject>& Extras = Object.GetObjectField(TEXT("extras"));
+			for (TPair<FString, TSharedPtr<FJsonValue>>& Pair : Extras->Values)
+			{
+				Mesh.Primitives.Last().Extras.Add(Pair.Key, Pair.Value->AsString());
+			}
+		}
+
 		ExtensionsHandler->SetupPrimitiveExtensions(Object, Mesh.Primitives.Last(), Mesh.Primitives.Num()-1, Mesh.UniqueId);
 	}
 
@@ -470,6 +479,14 @@ namespace GLTF
 					Mesh.MorphTargetNames.Add(Value->AsString());
 				}
 			}
+
+			for (TPair<FString, TSharedPtr<FJsonValue>>& Pair : Extras->Values)
+			{
+				if (Pair.Key != TEXT("targetNames"))
+				{
+					Mesh.Extras.Add(Pair.Key, Pair.Value->AsString());
+				}
+			}
 		}
 
 		Mesh.GenerateIsValidCache();
@@ -495,6 +512,15 @@ namespace GLTF
 			}
 		}
 
+		if (Object.HasField(TEXT("extras")))
+		{
+			const TSharedPtr<FJsonObject>& Extras = Object.GetObjectField(TEXT("extras"));
+			for (TPair<FString, TSharedPtr<FJsonValue>>& Pair : Extras->Values)
+			{
+				Scene.Extras.Add(Pair.Key, Pair.Value->AsString());
+			}
+		}
+
 		ExtensionsHandler->SetupSceneExtensions(Object, Scene);
 	}
 
@@ -548,6 +574,15 @@ namespace GLTF
 			}
 		}
 
+		if (Object.HasField(TEXT("extras")))
+		{
+			const TSharedPtr<FJsonObject>& Extras = Object.GetObjectField(TEXT("extras"));
+			for (TPair<FString, TSharedPtr<FJsonValue>>& Pair : Extras->Values)
+			{
+				Node.Extras.Add(Pair.Key, Pair.Value->AsString());
+			}
+		}
+
 		ExtensionsHandler->SetupNodeExtensions(Object, Node);
 	}
 
@@ -587,7 +622,18 @@ namespace GLTF
 			Camera.bIsPerspective              = false;
 		}
 		else
+		{
 			Messages.Emplace(EMessageSeverity::Error, TEXT("Invalid camera type: ") + Type);
+		}
+
+		if (Object.HasField(TEXT("extras")))
+		{
+			const TSharedPtr<FJsonObject>& Extras = Object.GetObjectField(TEXT("extras"));
+			for (TPair<FString, TSharedPtr<FJsonValue>>& Pair : Extras->Values)
+			{
+				Camera.Extras.Add(Pair.Key, Pair.Value->AsString());
+			}
+		}
 
 		ExtensionsHandler->SetupCameraExtensions(Object, Camera);
 	}
@@ -655,6 +701,15 @@ namespace GLTF
 			}
 		}
 
+		if (Object.HasField(TEXT("extras")))
+		{
+			const TSharedPtr<FJsonObject>& Extras = Object.GetObjectField(TEXT("extras"));
+			for (TPair<FString, TSharedPtr<FJsonValue>>& Pair : Extras->Values)
+			{
+				Animation.Extras.Add(Pair.Key, Pair.Value->AsString());
+			}
+		}
+
 		ExtensionsHandler->SetupAnimationExtensions(Object, Animation);
 	}
 
@@ -675,6 +730,15 @@ namespace GLTF
 
 		Skin.Skeleton = GetIndex(Object, TEXT("skeleton"));
 
+		if (Object.HasField(TEXT("extras")))
+		{
+			const TSharedPtr<FJsonObject>& Extras = Object.GetObjectField(TEXT("extras"));
+			for (TPair<FString, TSharedPtr<FJsonValue>>& Pair : Extras->Values)
+			{
+				Skin.Extras.Add(Pair.Key, Pair.Value->AsString());
+			}
+		}
+
 		ExtensionsHandler->SetupSkinExtensions(Object, Skin);
 	}
 
@@ -753,6 +817,15 @@ namespace GLTF
 			CurrentBufferOffset += Image.DataByteLength;
 		}
 
+		if (Object.HasField(TEXT("extras")))
+		{
+			const TSharedPtr<FJsonObject>& Extras = Object.GetObjectField(TEXT("extras"));
+			for (TPair<FString, TSharedPtr<FJsonValue>>& Pair : Extras->Values)
+			{
+				Image.Extras.Add(Pair.Key, Pair.Value->AsString());
+			}
+		}
+
 		ExtensionsHandler->SetupImageExtensions(Object, Image);
 	}
 
@@ -788,6 +861,16 @@ namespace GLTF
 			const FSampler& Sampler = HasSampler ? Asset->Samplers[SamplerIndex] : FSampler::DefaultSampler;
 
 			Asset->Textures.Emplace(TexName, Source, Sampler);
+
+			if (Object.HasField(TEXT("extras")))
+			{
+				const TSharedPtr<FJsonObject>& Extras = Object.GetObjectField(TEXT("extras"));
+				for (TPair<FString, TSharedPtr<FJsonValue>>& Pair : Extras->Values)
+				{
+					Asset->Textures.Last().Extras.Add(Pair.Key, Pair.Value->AsString());
+				}
+			}
+
 			ExtensionsHandler->SetupTextureExtensions(Object, Asset->Textures.Last());
 		}
 		else
@@ -830,6 +913,15 @@ namespace GLTF
 
 		Material.bIsDoubleSided = GetBool(Object, TEXT("doubleSided"));
 
+		if (Object.HasField(TEXT("extras")))
+		{
+			const TSharedPtr<FJsonObject>& Extras = Object.GetObjectField(TEXT("extras"));
+			for (TPair<FString, TSharedPtr<FJsonValue>>& Pair : Extras->Values)
+			{
+				Material.Extras.Add(Pair.Key, Pair.Value->AsString());
+			}
+		}
+
 		ExtensionsHandler->SetupMaterialExtensions(Object, Material);
 	}
 
diff --git a/Engine/Plugins/Interchange/Runtime/Source/Parsers/GLTFCore/Public/GLTFAnimation.h b/Engine/Plugins/Interchange/Runtime/Source/Parsers/GLTFCore/Public/GLTFAnimation.h
--- a/Engine/Plugins/Interchange/Runtime/Source/Parsers/GLTFCore/Public/GLTFAnimation.h
+++ b/Engine/Plugins/Interchange/Runtime/Source/Parsers/GLTFCore/Public/GLTFAnimation.h
@@ -72,6 +72,8 @@ namespace GLTF
 		TArray<FSampler> Samplers;
 		TArray<FChannel> Channels;
 
+		TMap<FString, FString> Extras;
+
 		FString          UniqueId; //will be generated in FAsset::GenerateNames
 	};
 
diff --git a/Engine/Plugins/Interchange/Runtime/Source/Parsers/GLTFCore/Public/GLTFAsset.h b/Engine/Plugins/Interchange/Runtime/Source/Parsers/GLTFCore/Public/GLTFAsset.h
--- a/Engine/Plugins/Interchange/Runtime/Source/Parsers/GLTFCore/Public/GLTFAsset.h
+++ b/Engine/Plugins/Interchange/Runtime/Source/Parsers/GLTFCore/Public/GLTFAsset.h
@@ -54,6 +54,7 @@ namespace GLTF
 	{
 		FString       Name;
 		TArray<int32> Nodes;
+		TMap<FString, FString> Extras;
 		FString	      UniqueId; //will be generated in FAsset::GenerateNames
 	};
 
diff --git a/Engine/Plugins/Interchange/Runtime/Source/Parsers/GLTFCore/Public/GLTFMaterial.h b/Engine/Plugins/Interchange/Runtime/Source/Parsers/GLTFCore/Public/GLTFMaterial.h
--- a/Engine/Plugins/Interchange/Runtime/Source/Parsers/GLTFCore/Public/GLTFMaterial.h
+++ b/Engine/Plugins/Interchange/Runtime/Source/Parsers/GLTFCore/Public/GLTFMaterial.h
@@ -185,6 +185,7 @@ namespace GLTF
 		bool            bHasEmissiveStrength;
 		FIridescence    Iridescence;
 
+		TMap<FString, FString> Extras;
 
 		FString UniqueId; //will be generated in FAsset::GenerateNames
 
diff --git a/Engine/Plugins/Interchange/Runtime/Source/Parsers/GLTFCore/Public/GLTFMesh.h b/Engine/Plugins/Interchange/Runtime/Source/Parsers/GLTFCore/Public/GLTFMesh.h
--- a/Engine/Plugins/Interchange/Runtime/Source/Parsers/GLTFCore/Public/GLTFMesh.h
+++ b/Engine/Plugins/Interchange/Runtime/Source/Parsers/GLTFCore/Public/GLTFMesh.h
@@ -121,6 +121,8 @@ namespace GLTF
 		const int32 MaterialIndex;
 		TArray<FVariantMapping>	VariantMappings;
 
+		TMap<FString, FString> Extras;
+
 		FPrimitive(EMode InMode, int32 InMaterial, const FAccessor& InIndices, const FAccessor& InPosition, const FAccessor& InNormal,
 		           const FAccessor& InTangent, const FAccessor& InTexCoord0, const FAccessor& InTexCoord1, const FAccessor& InColor0,
 		           const FAccessor& InJoints0, const FAccessor& InWeights0);
@@ -175,6 +177,8 @@ namespace GLTF
 		TArray<float>		MorphTargetWeights;
 		TArray<FString>		MorphTargetNames;
 
+		TMap<FString, FString> Extras;
+
 		FString				UniqueId; //will be generated in FAsset::GenerateNames
 	
 		bool HasNormals() const;
diff --git a/Engine/Plugins/Interchange/Runtime/Source/Parsers/GLTFCore/Public/GLTFNode.h b/Engine/Plugins/Interchange/Runtime/Source/Parsers/GLTFCore/Public/GLTFNode.h
--- a/Engine/Plugins/Interchange/Runtime/Source/Parsers/GLTFCore/Public/GLTFNode.h
+++ b/Engine/Plugins/Interchange/Runtime/Source/Parsers/GLTFCore/Public/GLTFNode.h
@@ -43,8 +43,6 @@ namespace GLTF
 
 		TArray<float> MorphTargetWeights; //for the instantiated mesh with morph targets.
 
-		FString UniqueId; //will be generated in FAsset::GenerateNames
-
 		TMap<int, FTransform> SkinIndexToGlobalInverseBindTransform;
 		TMap<int, FTransform> SkinIndexToLocalBindPose; //bind pose would be CurrentNode.GlobalInverseBindTransform.Inverse() * ParentNode.GlobalInverseBindTransform
 		bool bHasLocalBindPose;
@@ -52,6 +50,10 @@ namespace GLTF
 									//	Edge case Scenario which is currently not supported:
 									//	Where multiple skins use the same Joint. Currently expected bad outcome if the different skins have different inversebindmatrices on the joint.
 
+		TMap<FString, FString> Extras;
+
+		FString UniqueId;
+
 		FNode()
 		    : Type(EType::None)
 		    , MeshIndex(INDEX_NONE)
@@ -92,6 +94,9 @@ namespace GLTF
 		float ZNear;
 		float ZFar;
 		bool  bIsPerspective;
+
+		TMap<FString, FString> Extras;
+
 		FString UniqueId; //will be generated in FAsset::GenerateNames
 
 		FCamera(const FNode& Node)
@@ -152,6 +157,8 @@ namespace GLTF
 		TArray<int32>    Joints;    // each is an index into FAsset::Nodes
 		int32            Skeleton;  // root node, index into FAsset::Nodes
 
+		TMap<FString, FString> Extras;
+
 		FString          UniqueId; //will be generated in FAsset::GenerateNames
 
 		FSkinInfo(const FAccessor& InverseBindMatrices)
diff --git a/Engine/Plugins/Interchange/Runtime/Source/Parsers/GLTFCore/Public/GLTFTexture.h b/Engine/Plugins/Interchange/Runtime/Source/Parsers/GLTFCore/Public/GLTFTexture.h
--- a/Engine/Plugins/Interchange/Runtime/Source/Parsers/GLTFCore/Public/GLTFTexture.h
+++ b/Engine/Plugins/Interchange/Runtime/Source/Parsers/GLTFCore/Public/GLTFTexture.h
@@ -26,6 +26,8 @@ namespace GLTF
 		uint32       DataByteLength;
 		const uint8* Data;
 
+		TMap<FString, FString> Extras;
+
 		FString      UniqueId; //will be generated in FAsset::GenerateNames
 
 		FImage()
@@ -80,6 +82,8 @@ namespace GLTF
 		const FSampler& Sampler;
 		FString         Name;
 
+		TMap<FString, FString> Extras;
+
 		FString			UniqueId; //will be generated in FAsset::GenerateNames
 
 		FTexture(const FString& InName, const FImage& InSource, const FSampler& InSampler)
