Interchange: Force re-import of alternate skinning when re-importing base mesh or custom LOD

#rnx

diff --git a/Engine/Plugins/Interchange/Runtime/Source/Import/Private/Mesh/InterchangeSkeletalMeshFactory.cpp b/Engine/Plugins/Interchange/Runtime/Source/Import/Private/Mesh/InterchangeSkeletalMeshFactory.cpp
--- a/Engine/Plugins/Interchange/Runtime/Source/Import/Private/Mesh/InterchangeSkeletalMeshFactory.cpp
+++ b/Engine/Plugins/Interchange/Runtime/Source/Import/Private/Mesh/InterchangeSkeletalMeshFactory.cpp
@@ -19,6 +19,7 @@
 #include "InterchangeManager.h"
 #include "InterchangeMaterialFactoryNode.h"
 #include "InterchangeMeshNode.h"
+#include "InterchangeMeshUtilities.h"
 #include "InterchangeSceneNode.h"
 #include "InterchangeSkeletalMeshLodDataNode.h"
 #include "InterchangeSkeletalMeshFactoryNode.h"
@@ -934,69 +935,6 @@ namespace UE
 
 #endif //#if WITH_EDITOR
 
-void FInterchangeSkeletalMeshPostImportTask::Execute()
-{
-#if WITH_EDITOR
-	if (!SkeletalMesh)
-	{
-		return;
-	}
-
-	//This code works only on the game thread and is not asynchronous
-	check(IsInGameThread());
-
-	if (bReImportAlternateSkinWeights)
-	{
-		//We can't show user if unattended just return;
-		if (GIsAutomationTesting || GIsRunningUnattendedScript || FApp::IsUnattended())
-		{
-			return;
-		}
-
-		FText MessageText = FText::Format(NSLOCTEXT("UInterchangeSkeletalMeshPostImportTask", "ShouldReimportAlternateSkinning", "When reimporting asset \"{0}\" the system cannot re-apply the alternate skinning data, do you want to re-import all alternate skinning?")
-			, FText::FromString(SkeletalMesh->GetName()));
-
-		EAppReturnType::Type ReimportAlternateSkinningChoice = FMessageDialog::Open(EAppMsgType::YesNo
-			, EAppReturnType::No
-			, MessageText);
-		if (ReimportAlternateSkinningChoice == EAppReturnType::No)
-		{
-			return;
-		}
-
-		//User say yes so re-import the alternate skinning
-		const int32 LodCount = SkeletalMesh->GetLODNum();
-		float ProgressCount = LodCount + 0.1f;
-
-		FScopedSlowTask Progress(ProgressCount, NSLOCTEXT("UInterchangeSkeletalMeshPostImportTask", "SkeletalMeshPostImportTaskGameThread", "Executing Skeletal Mesh Post Import Tasks..."));
-		Progress.MakeDialog();
-		{
-			//Make sure we rebuild the skeletal mesh after re-importing all skin weight
-			FScopedSkeletalMeshPostEditChange ScopePostEditChange(SkeletalMesh);
-
-			//Wait until the asset is finish building then lock the skeletal mesh properties to prevent the UI to update during the alternate skinning reimport
-			FEvent* LockEvent = SkeletalMesh->LockPropertiesUntil();
-			FSkinnedAssetAsyncBuildScope AsyncBuildScope(SkeletalMesh);
-
-			//We have a 0.1 progress for the lock
-			Progress.EnterProgressFrame(0.1f);
-
-			//Reimport all the alternate skinning
-			for (int32 LodIndex = 0; LodIndex < LodCount; ++LodIndex)
-			{
-				FSkinWeightsUtilities::ReimportAlternateSkinWeight(SkeletalMesh, LodIndex);
-				Progress.EnterProgressFrame(1.0f);
-			}
-
-			//Release the skeletal mesh async properties
-			LockEvent->Trigger();
-
-			//Skeletal mesh will rebuild when going out of scope
-		}
-	}
-#endif //WITH_EDITOR
-}
-
 UClass* UInterchangeSkeletalMeshFactory::GetFactoryClass() const
 {
 	return USkeletalMesh::StaticClass();
@@ -2078,19 +2016,27 @@ void UInterchangeSkeletalMeshFactory::SetupObject_GameThread(const FSetupObjectP
 		//Re-apply the alternate skinning data
 		if (ImportAssetObjectData.ExistingSkinWeightProfileInfos.Num() > 0)
 		{
-			TSharedPtr<FInterchangeSkeletalMeshPostImportTask> SkeletalMeshPostImportTask = nullptr;
+			//Reimport alternate skin of the reimport lod only by looking at the factory node lod data count.
+			int32 LodCount = SkeletalMesh->GetLODNum();
+			if (const UInterchangeSkeletalMeshFactoryNode* SkeletalMeshFactoryNode = Cast<const UInterchangeSkeletalMeshFactoryNode>(Arguments.NodeContainer->GetFactoryNode(Arguments.NodeUniqueID)))
+			{
+				int32 ImportedLodCount = SkeletalMeshFactoryNode->GetLodDataCount();
+				if (ImportedLodCount >  0 && ImportedLodCount <= SkeletalMesh->GetLODNum())
+				{
+					LodCount = ImportedLodCount;
+				}
+			}
+
+			TSharedPtr<FInterchangeSkeletalMeshAlternateSkinWeightPostImportTask> SkeletalMeshPostImportTask = nullptr;
 			TArray<FSkinWeightProfileInfo>& SkinProfiles = SkeletalMesh->GetSkinWeightProfiles();
 			SkinProfiles = ImportAssetObjectData.ExistingSkinWeightProfileInfos;
 			for (const FSkinWeightProfileInfo& ProfileInfo : SkinProfiles)
 			{
-				const int32 LodCount = SkeletalMesh->GetLODNum();
 				for (int32 LodIndex = 0; LodIndex < LodCount; ++LodIndex)
 				{
-					if (!ImportAssetObjectData.ExistingAlternateImportDataPerLOD.IsValidIndex(LodIndex))
-					{
-						continue;
-					}
-					if (!SkeletalMesh->HasMeshDescription(LodIndex))
+					if (!SkeletalMesh->IsValidLODIndex(LodIndex)
+						|| !ImportAssetObjectData.ExistingAlternateImportDataPerLOD.IsValidIndex(LodIndex)
+						|| !SkeletalMesh->HasMeshDescription(LodIndex))
 					{
 						continue;
 					}
@@ -2108,50 +2054,50 @@ void UInterchangeSkeletalMeshFactory::SetupObject_GameThread(const FSetupObjectP
 					PRAGMA_DISABLE_DEPRECATION_WARNINGS
 					SkeletalMesh->LoadLODImportedData(LodIndex, ImportDataDest);
 					PRAGMA_ENABLE_DEPRECATION_WARNINGS
-					
+
 					int32 PointNumberDest = ImportDataDest.Points.Num();
 					int32 VertexNumberDest = ImportDataDest.Points.Num();
 
-					if (ExistingImportDataSrc.Points.Num() != PointNumberDest)
+					if (ExistingImportDataSrc.Points.Num() == PointNumberDest)
 					{
-						//Warn the user but still apply the data
-						UInterchangeResultWarning_Generic* Message = AddMessage<UInterchangeResultWarning_Generic>();
-						Message->Text = FText::Format(NSLOCTEXT("InterchangeSkeletalMeshFactory", "SetupObject_GameThread_AlternateSkinning missmatch", "Alternate skinning mesh vertex number is different from the mesh LOD, we cannot apply the existing alternate skinning [{0}] when re-importing skeletal mesh LOD{1} [{2}]")
-							, FText::FromString(ProfileNameStr)
-							, FText::AsNumber(LodIndex)
-							, FText::FromString(SkeletalMesh->GetName()));
-
-						//We must enqueue a post import task that will re-import all skin weight profile.
-						if (!SkeletalMeshPostImportTask.IsValid())
+						//Replace the data into the destination bulk data and save it, so the build dont complain until we reimport the new alternate skinning
+						int32 ProfileIndex = 0;
+						if (ImportDataDest.AlternateInfluenceProfileNames.Find(ProfileNameStr, ProfileIndex))
 						{
-							SkeletalMeshPostImportTask = MakeShared<FInterchangeSkeletalMeshPostImportTask>();
-							SkeletalMeshPostImportTask->SkeletalMesh = SkeletalMesh;
-							SkeletalMeshPostImportTask->bReImportAlternateSkinWeights = true;
-							UInterchangeManager::GetInterchangeManager().EnqueuePostImportTask(SkeletalMeshPostImportTask);
+							ImportDataDest.AlternateInfluenceProfileNames.RemoveAt(ProfileIndex);
+							ImportDataDest.AlternateInfluences.RemoveAt(ProfileIndex);
+						}
+						int32 SrcProfileIndex = 0;
+						if (ExistingImportDataSrc.AlternateInfluenceProfileNames.Find(ProfileNameStr, SrcProfileIndex))
+						{
+							ImportDataDest.AlternateInfluenceProfileNames.Add(ProfileNameStr);
+							ImportDataDest.AlternateInfluences.Add(ExistingImportDataSrc.AlternateInfluences[SrcProfileIndex]);
 						}
-						continue;
-					}
 
-					//Replace the data into the destination bulk data and save it
-					int32 ProfileIndex = 0;
-					if (ImportDataDest.AlternateInfluenceProfileNames.Find(ProfileNameStr, ProfileIndex))
-					{
-						ImportDataDest.AlternateInfluenceProfileNames.RemoveAt(ProfileIndex);
-						ImportDataDest.AlternateInfluences.RemoveAt(ProfileIndex);
+						//Resave the bulk data with the new or refreshed data
+						PRAGMA_DISABLE_DEPRECATION_WARNINGS
+						SkeletalMesh->SaveLODImportedData(LodIndex, ImportDataDest);
+						PRAGMA_ENABLE_DEPRECATION_WARNINGS
 					}
-					int32 SrcProfileIndex = 0;
-					if (ExistingImportDataSrc.AlternateInfluenceProfileNames.Find(ProfileNameStr, SrcProfileIndex))
+
+					//We must enqueue a post import task that will re-import all skin weight profile.
+					if (!SkeletalMeshPostImportTask.IsValid())
 					{
-						ImportDataDest.AlternateInfluenceProfileNames.Add(ProfileNameStr);
-						ImportDataDest.AlternateInfluences.Add(ExistingImportDataSrc.AlternateInfluences[SrcProfileIndex]);
+						SkeletalMeshPostImportTask = MakeShared<FInterchangeSkeletalMeshAlternateSkinWeightPostImportTask>(SkeletalMesh);
+						SkeletalMeshPostImportTask->ReimportAlternateSkinWeightDelegate.BindLambda([](USkeletalMesh* SkeletalMesh, int32 LodIndex)
+							{
+								return FSkinWeightsUtilities::ReimportAlternateSkinWeight(SkeletalMesh, LodIndex);
+							});
 					}
-
-					//Resave the bulk data with the new or refreshed data
-					PRAGMA_DISABLE_DEPRECATION_WARNINGS
-					SkeletalMesh->SaveLODImportedData(LodIndex, ImportDataDest);
-					PRAGMA_ENABLE_DEPRECATION_WARNINGS
+					SkeletalMeshPostImportTask->AddLodToReimportAlternate(LodIndex);
 				}
 			}
+			
+			//Enqueue the task if needed
+			if (SkeletalMeshPostImportTask.IsValid())
+			{
+				UInterchangeManager::GetInterchangeManager().EnqueuePostImportTask(SkeletalMeshPostImportTask);
+			}
 		}
 #endif //WITH_EDITOR
 	}
diff --git a/Engine/Plugins/Interchange/Runtime/Source/Import/Public/Mesh/InterchangeSkeletalMeshFactory.h b/Engine/Plugins/Interchange/Runtime/Source/Import/Public/Mesh/InterchangeSkeletalMeshFactory.h
--- a/Engine/Plugins/Interchange/Runtime/Source/Import/Public/Mesh/InterchangeSkeletalMeshFactory.h
+++ b/Engine/Plugins/Interchange/Runtime/Source/Import/Public/Mesh/InterchangeSkeletalMeshFactory.h
@@ -16,15 +16,6 @@ class UInterchangeSkeletalMeshFactoryNode;
 class USkeletalMesh;
 class USkeleton;
 
-class INTERCHANGEIMPORT_API FInterchangeSkeletalMeshPostImportTask : public FInterchangePostImportTask
-{
-public:
-	virtual void Execute() override;
-
-	TObjectPtr<USkeletalMesh> SkeletalMesh;
-	bool bReImportAlternateSkinWeights = false;
-};
-
 UCLASS(BlueprintType)
 class INTERCHANGEIMPORT_API UInterchangeSkeletalMeshFactory : public UInterchangeFactoryBase
 {
diff --git a/Engine/Source/Editor/UnrealEd/Private/FbxMeshUtils.cpp b/Engine/Source/Editor/UnrealEd/Private/FbxMeshUtils.cpp
--- a/Engine/Source/Editor/UnrealEd/Private/FbxMeshUtils.cpp
+++ b/Engine/Source/Editor/UnrealEd/Private/FbxMeshUtils.cpp
@@ -507,6 +507,20 @@ namespace FbxMeshUtils
 						{
 							if (bResult)
 							{
+								//If we use alternate skinweight, we must re-import all profile for this LOD
+								if(!SelectedSkelMesh->GetSkinWeightProfiles().IsEmpty())
+								{
+									//Enqueue the re-import alternate skinning
+									TSharedPtr<FInterchangeSkeletalMeshAlternateSkinWeightPostImportTask> SkeletalMeshPostImportTask = MakeShared<FInterchangeSkeletalMeshAlternateSkinWeightPostImportTask>(SelectedSkelMesh);
+									SkeletalMeshPostImportTask->ReimportAlternateSkinWeightDelegate.BindLambda([](USkeletalMesh* SkeletalMesh, int32 LodIndex)
+										{
+											return FSkinWeightsUtilities::ReimportAlternateSkinWeight(SkeletalMesh, LodIndex);
+										});
+									SkeletalMeshPostImportTask->AddLodToReimportAlternate(LODLevel);
+									UInterchangeManager::GetInterchangeManager().EnqueuePostImportTask(SkeletalMeshPostImportTask);
+								}
+								
+
 								// Notification of success
 								FNotificationInfo NotificationInfo(FText::GetEmpty());
 								NotificationInfo.Text = FText::Format(NSLOCTEXT("UnrealEd", "LODImportSuccessful", "Mesh for LOD {0} imported successfully!"), FText::AsNumber(LODLevel));
@@ -981,6 +995,18 @@ namespace FbxMeshUtils
 					bool bResult = FutureResult.Get();
 					if (bResult)
 					{
+						//If we use alternate skinweight, we must re-import all profile for this LOD
+						if (!SkeletalMesh->GetSkinWeightProfiles().IsEmpty())
+						{
+							//Enqueue the re-import alternate skinning
+							TSharedPtr<FInterchangeSkeletalMeshAlternateSkinWeightPostImportTask> SkeletalMeshPostImportTask = MakeShared<FInterchangeSkeletalMeshAlternateSkinWeightPostImportTask>(SkeletalMesh);
+							SkeletalMeshPostImportTask->ReimportAlternateSkinWeightDelegate.BindLambda([](USkeletalMesh* SkeletalMesh, int32 LodIndex)
+								{
+									return FSkinWeightsUtilities::ReimportAlternateSkinWeight(SkeletalMesh, LodIndex);
+								});
+							SkeletalMeshPostImportTask->AddLodToReimportAlternate(LODLevel);
+							UInterchangeManager::GetInterchangeManager().EnqueuePostImportTask(SkeletalMeshPostImportTask);
+						}
 						if (bNotifyCB)
 						{
 							if (SkeletalMesh)
diff --git a/Engine/Source/Editor/UnrealEd/Private/SkinWeightsUtilities.cpp b/Engine/Source/Editor/UnrealEd/Private/SkinWeightsUtilities.cpp
--- a/Engine/Source/Editor/UnrealEd/Private/SkinWeightsUtilities.cpp
+++ b/Engine/Source/Editor/UnrealEd/Private/SkinWeightsUtilities.cpp
@@ -68,6 +68,10 @@ bool FSkinWeightsUtilities::ImportAlternateSkinWeight(USkeletalMesh* SkeletalMes
 		UE_LOG(LogSkinWeightsUtilities, Error, TEXT("Path containing Skin Weight Profile data does not exist (%s)."), *Path);
 		return false;
 	}
+
+	FString Action = bIsReimport ? TEXT("Reimport") : TEXT("Import");
+	UE_ASSET_LOG(LogSkinWeightsUtilities, Display, SkeletalMesh, TEXT("%s Alternate skin weight Begin [LodIndex: %d] [Profile: %s] [file: %s]."), *Action, TargetLODIndex, *ProfileName.ToString(), *Path);
+
 	FScopedSuspendAlternateSkinWeightPreview ScopedSuspendAlternateSkinnWeightPreview(SkeletalMesh);
 	FScopedSkeletalMeshPostEditChange ScopePostEditChange(SkeletalMesh);
 
diff --git a/Engine/Source/Runtime/Interchange/Engine/Private/InterchangeMeshUtilities.cpp b/Engine/Source/Runtime/Interchange/Engine/Private/InterchangeMeshUtilities.cpp
--- a/Engine/Source/Runtime/Interchange/Engine/Private/InterchangeMeshUtilities.cpp
+++ b/Engine/Source/Runtime/Interchange/Engine/Private/InterchangeMeshUtilities.cpp
@@ -14,6 +14,7 @@
 #include "InterchangePythonPipelineBase.h"
 #include "InterchangeSourceData.h"
 #include "Logging/LogMacros.h"
+#include "Misc/ScopedSlowTask.h"
 #include "UObject/Class.h"
 #include "UObject/Object.h"
 #include "UObject/ObjectMacros.h"
@@ -24,6 +25,78 @@
 
 #include UE_INLINE_GENERATED_CPP_BY_NAME(InterchangeMeshUtilities)
 
+FInterchangeSkeletalMeshAlternateSkinWeightPostImportTask::FInterchangeSkeletalMeshAlternateSkinWeightPostImportTask(USkeletalMesh* InSkeletalMesh)
+	:SkeletalMesh(InSkeletalMesh)
+{
+
+}
+
+void FInterchangeSkeletalMeshAlternateSkinWeightPostImportTask::Execute()
+{
+#if WITH_EDITOR
+	//This code works only on the game thread and is not asynchronous
+	if (!ensure(IsInGameThread()))
+	{
+		return;
+	}
+
+	//The delegate must be bound if we want to reimport the alternate skinning.
+	if (!SkeletalMesh
+		|| !ReimportAlternateSkinWeightDelegate.IsBound()
+		|| ReImportAlternateSkinWeightsLods.IsEmpty())
+	{
+		return;
+	}
+
+	//User say yes so re-import the alternate skinning
+	const int32 LodCount = SkeletalMesh->GetLODNum();
+	float ProgressCount = LodCount + 0.1f;
+
+	FScopedSlowTask Progress(ProgressCount, NSLOCTEXT("UInterchangeSkeletalMeshPostImportTask", "SkeletalMeshPostImportTaskGameThread", "Executing Skeletal Mesh Post Import Tasks..."));
+	Progress.MakeDialog();
+	{
+		//Make sure we rebuild the skeletal mesh after re-importing all skin weight
+		FScopedSkeletalMeshPostEditChange ScopePostEditChange(SkeletalMesh);
+
+		//Wait until the asset is finish building then lock the skeletal mesh properties to prevent the UI to update during the alternate skinning reimport
+		FEvent* LockEvent = SkeletalMesh->LockPropertiesUntil();
+		FSkinnedAssetAsyncBuildScope AsyncBuildScope(SkeletalMesh);
+
+		//We have a 0.1 progress for the lock
+		Progress.EnterProgressFrame(0.1f);
+
+		//Reimport all the alternate skinning
+		for (int32 LodIndex = 0; LodIndex < LodCount; ++LodIndex)
+		{
+			if (ReImportAlternateSkinWeightsLods.Contains(LodIndex))
+			{
+				// This delegate should execute the following editor function
+				// FSkinWeightsUtilities::ReimportAlternateSkinWeight(SkeletalMesh, LodIndex);
+				ReimportAlternateSkinWeightDelegate.Execute(SkeletalMesh, LodIndex);
+			}
+			Progress.EnterProgressFrame(1.0f);
+		}
+
+		//Release the skeletal mesh async properties
+		LockEvent->Trigger();
+
+		//Skeletal mesh will rebuild when going out of scope
+	}
+#endif //WITH_EDITOR
+}
+
+//DECLARE_DELEGATE_RetVal_TwoParams(bool, FInterchangeReimportAlternateSkinWeight, USkeletalMesh*, int32 LodIndex);
+
+bool FInterchangeSkeletalMeshAlternateSkinWeightPostImportTask::AddLodToReimportAlternate(int32 LodToAdd)
+{
+	if (!SkeletalMesh || !SkeletalMesh->IsValidLODIndex(LodToAdd))
+	{
+		return false;
+	}
+	ReImportAlternateSkinWeightsLods.AddUnique(LodToAdd);
+	return true;
+}
+
 TFuture<bool> UInterchangeMeshUtilities::ImportCustomLod(UObject* MeshObject, const int32 LodIndex)
 {
 	TSharedPtr<TPromise<bool>> Promise = MakeShared<TPromise<bool>>();
diff --git a/Engine/Source/Runtime/Interchange/Engine/Public/InterchangeMeshUtilities.h b/Engine/Source/Runtime/Interchange/Engine/Public/InterchangeMeshUtilities.h
--- a/Engine/Source/Runtime/Interchange/Engine/Public/InterchangeMeshUtilities.h
+++ b/Engine/Source/Runtime/Interchange/Engine/Public/InterchangeMeshUtilities.h
@@ -3,12 +3,37 @@
 #pragma once
 
 #include "Async/Future.h"
+#include "InterchangeFactoryBase.h"
 #include "UObject/Object.h"
 #include "UObject/ObjectMacros.h"
 
 #include "InterchangeMeshUtilities.generated.h"
 
 class UInterchangeSourceData;
+class USkeletalMesh;
+
+class INTERCHANGEENGINE_API FInterchangeSkeletalMeshAlternateSkinWeightPostImportTask : public FInterchangePostImportTask
+{
+public:
+	FInterchangeSkeletalMeshAlternateSkinWeightPostImportTask(USkeletalMesh* InSkeletalMesh);
+
+	virtual ~FInterchangeSkeletalMeshAlternateSkinWeightPostImportTask() {};
+
+	virtual void Execute() override;
+
+	// This delegate should execute the following editor function (since we are a runtime module we need this delegate to call the editor function)
+	// FSkinWeightsUtilities::ReimportAlternateSkinWeight(SkeletalMesh, LodIndex);
+	DECLARE_DELEGATE_RetVal_TwoParams(bool, FInterchangeReimportAlternateSkinWeight, USkeletalMesh*, int32 LodIndex);
+	FInterchangeReimportAlternateSkinWeight ReimportAlternateSkinWeightDelegate;
+
+	bool AddLodToReimportAlternate(int32 LodToAdd);
+	
+private:
+	TObjectPtr<USkeletalMesh> SkeletalMesh;
+
+	//Alternate skin weights reimport
+	TArray<int32> ReImportAlternateSkinWeightsLods;
+};
 
 UCLASS(Experimental, MinimalAPI)
 class UInterchangeMeshUtilities : public UObject
@@ -40,6 +65,7 @@ public:
 	 */
 	static INTERCHANGEENGINE_API TFuture<bool> ImportCustomLod(UObject* MeshObject, const int32 LodIndex, const UInterchangeSourceData* SourceData);
 
+
 private:
 
 	static TFuture<bool> InternalImportCustomLod(TSharedPtr<TPromise<bool>> Promise, UObject* MeshObject, const int32 LodIndex, const UInterchangeSourceData* SourceData);
