[Interchange] Import dialog not showing Skeletal Conflicts when adding joints to imported skelMesh

#rnx

diff --git a/Engine/Plugins/Interchange/Runtime/Source/Import/Private/Mesh/InterchangeSkeletonHelper.cpp b/Engine/Plugins/Interchange/Runtime/Source/Import/Private/Mesh/InterchangeSkeletonHelper.cpp
--- a/Engine/Plugins/Interchange/Runtime/Source/Import/Private/Mesh/InterchangeSkeletonHelper.cpp
+++ b/Engine/Plugins/Interchange/Runtime/Source/Import/Private/Mesh/InterchangeSkeletonHelper.cpp
@@ -73,7 +73,7 @@ namespace UE::Interchange::Private
 		return true;
 	}
 
-	bool FSkeletonHelper::IsCompatibleSkeleton(const USkeleton* Skeleton, const FString RootJoinUid, const UInterchangeBaseNodeContainer* BaseNodeContainer, bool bConvertStaticToSkeletalActive)
+	bool FSkeletonHelper::IsCompatibleSkeleton(const USkeleton* Skeleton, const FString RootJoinUid, const UInterchangeBaseNodeContainer* BaseNodeContainer, bool bConvertStaticToSkeletalActive, bool bCheckForIdenticalSkeleton)
 	{
 		if (!Skeleton)
 		{
@@ -112,6 +112,10 @@ namespace UE::Interchange::Private
 			}
 			else
 			{
+				if (bCheckForIdenticalSkeleton)
+				{
+					return false;
+				}
 				int32 CurrentBoneId = MeshBoneIndex;
 				// if not look for parents that matches
 				while (SkeletonBoneIndex == INDEX_NONE && CurrentBoneId != INDEX_NONE)
diff --git a/Engine/Plugins/Interchange/Runtime/Source/Import/Public/InterchangeSkeletonHelper.h b/Engine/Plugins/Interchange/Runtime/Source/Import/Public/InterchangeSkeletonHelper.h
--- a/Engine/Plugins/Interchange/Runtime/Source/Import/Public/InterchangeSkeletonHelper.h
+++ b/Engine/Plugins/Interchange/Runtime/Source/Import/Public/InterchangeSkeletonHelper.h
@@ -37,7 +37,7 @@ namespace UE::Interchange::Private
 			, const bool bUseTimeZeroAsBindPose
 			, bool& bOutDiffPose);
 
-		static bool IsCompatibleSkeleton(const USkeleton* Skeleton, const FString RootJoinUid, const UInterchangeBaseNodeContainer* BaseNodeContainer, bool bConvertStaticToSkeletalActive);
+		static bool IsCompatibleSkeleton(const USkeleton* Skeleton, const FString RootJoinUid, const UInterchangeBaseNodeContainer* BaseNodeContainer, bool bConvertStaticToSkeletalActive, bool bCheckForIdenticalSkeleton = false);
 		static void RecursiveAddSkeletonMetaDataValues(UInterchangeBaseNodeContainer* NodeContainer, UInterchangeBaseNode* DestinationNode, const FString& JointUid);
 	private:
 		static void RecursiveAddBones(const UInterchangeBaseNodeContainer* NodeContainer
diff --git a/Engine/Plugins/Interchange/Runtime/Source/Pipelines/Private/InterchangeGenericAssetsPipeline.cpp b/Engine/Plugins/Interchange/Runtime/Source/Pipelines/Private/InterchangeGenericAssetsPipeline.cpp
--- a/Engine/Plugins/Interchange/Runtime/Source/Pipelines/Private/InterchangeGenericAssetsPipeline.cpp
+++ b/Engine/Plugins/Interchange/Runtime/Source/Pipelines/Private/InterchangeGenericAssetsPipeline.cpp
@@ -828,8 +828,9 @@ void UInterchangeGenericAssetsPipeline::CreateSkeletonConflict(USkeleton* Specif
 		}
 
 		const bool bConvertStaticToSkeletalActive = CommonSkeletalMeshesAndAnimationsProperties->bConvertStaticsWithMorphTargetsToSkeletals || CommonMeshesProperties->ForceAllMeshAsType == EInterchangeForceMeshType::IFMT_SkeletalMesh;
-		//If we have a compatible skeleton we do not need to create a conflict
-		if (UE::Interchange::Private::FSkeletonHelper::IsCompatibleSkeleton(SpecifiedSkeleton ? SpecifiedSkeleton : SkeletalMesh->GetSkeleton(), RootJointNodeId, TransientBaseNodeContainer, bConvertStaticToSkeletalActive))
+		//If we have a compatible skeleton we do not need to create a conflict, for stricter comparisons we check for identical skeletons to report any deviations.
+		constexpr bool bCheckForIdenticalSkeleton = true;
+		if (UE::Interchange::Private::FSkeletonHelper::IsCompatibleSkeleton(SpecifiedSkeleton ? SpecifiedSkeleton : SkeletalMesh->GetSkeleton(), RootJointNodeId, TransientBaseNodeContainer, bConvertStaticToSkeletalActive, bCheckForIdenticalSkeleton))
 		{
 			return;
 		}
