[Interchange - FBX] Combine skeletal meshes option not importing SKM when disabled.
Deprecating Combine Skeletal Mesh option.

#rnx

diff --git a/Engine/Plugins/Interchange/Runtime/Source/Pipelines/Private/InterchangeGenericSkeletalMeshPipeline.cpp b/Engine/Plugins/Interchange/Runtime/Source/Pipelines/Private/InterchangeGenericSkeletalMeshPipeline.cpp
--- a/Engine/Plugins/Interchange/Runtime/Source/Pipelines/Private/InterchangeGenericSkeletalMeshPipeline.cpp
+++ b/Engine/Plugins/Interchange/Runtime/Source/Pipelines/Private/InterchangeGenericSkeletalMeshPipeline.cpp
@@ -120,6 +120,9 @@ void UInterchangeGenericMeshPipeline::ExecutePreImportPipelineSkeletalMesh()
 		}
 	};
 
+	PRAGMA_DISABLE_DEPRECATION_WARNINGS
+	const bool bCombineSkeletalMeshes = bCombineSkeletalMeshes_DEPRECATED;
+	PRAGMA_ENABLE_DEPRECATION_WARNINGS
 	if (bCombineSkeletalMeshes)
 	{
 		//////////////////////////////////////////////////////////////////////////
diff --git a/Engine/Plugins/Interchange/Runtime/Source/Pipelines/Public/InterchangeGenericMeshPipeline.h b/Engine/Plugins/Interchange/Runtime/Source/Pipelines/Public/InterchangeGenericMeshPipeline.h
--- a/Engine/Plugins/Interchange/Runtime/Source/Pipelines/Public/InterchangeGenericMeshPipeline.h
+++ b/Engine/Plugins/Interchange/Runtime/Source/Pipelines/Public/InterchangeGenericMeshPipeline.h
@@ -152,9 +152,24 @@ public:
 	UPROPERTY()
 	EInterchangeSkeletalMeshContentType LastSkeletalMeshImportContentType;
 
+	UE_DEPRECATED(5.5, "bCombineSkeletalMeshes is no longer used")
+	UFUNCTION(BlueprintGetter, meta = (DeprecatedFunction, DeprecationMessage = "bCombineSkeletalMeshes is no longer used"))
+	bool GetCombineSkeletalMeshes() const
+	{
+		PRAGMA_DISABLE_DEPRECATION_WARNINGS
+		return bCombineSkeletalMeshes_DEPRECATED;
+		PRAGMA_ENABLE_DEPRECATION_WARNINGS
+	}
+
+	UE_DEPRECATED(5.5, "bCombineSkeletalMeshes is no longer used")
+	UFUNCTION(BlueprintSetter, meta = (DeprecatedFunction, DeprecationMessage = "bCombineSkeletalMeshes is no longer used"))
+	void SetCombineSkeletalMeshes(bool InbCombineSkeletalMeshes) {}
+
 	/** If enabled, all skinned mesh nodes that belong to the same skeleton root joint are combined into a single skeletal mesh. */
-	UPROPERTY(EditAnywhere, BlueprintReadWrite, Category = "Skeletal Meshes")
-	bool bCombineSkeletalMeshes = true;
+	UE_DEPRECATED(5.5, "Please do not access this member. It will be remove in the next version.")
+	UPROPERTY(BlueprintReadWrite, BlueprintGetter = GetCombineSkeletalMeshes, BlueprintSetter = SetCombineSkeletalMeshes, Category = "Skeletal Meshes")
+	bool bCombineSkeletalMeshes_DEPRECATED = true;
+
 
 	/** If enabled, imports all morph target shapes found in the source. */
 	UPROPERTY(EditAnywhere, BlueprintReadWrite, Category = "Skeletal Meshes")
