Fixing FBX material instances not correctly overriding input parameters.

diff --git a/Engine/Plugins/Interchange/Runtime/Source/Parsers/Fbx/Private/FbxMaterial.cpp b/Engine/Plugins/Interchange/Runtime/Source/Parsers/Fbx/Private/FbxMaterial.cpp
--- a/Engine/Plugins/Interchange/Runtime/Source/Parsers/Fbx/Private/FbxMaterial.cpp
+++ b/Engine/Plugins/Interchange/Runtime/Source/Parsers/Fbx/Private/FbxMaterial.cpp
@@ -178,26 +178,25 @@ namespace UE
 
 					if (DefaultValue.IsType<float>())
 					{
-						LerpNode->AddFloatAttribute(UInterchangeShaderPortsAPI::MakeInputValueKey(Lerp::Inputs::B.ToString()), DefaultValue.Get<float>());
+						LerpNode->AddFloatAttribute(UInterchangeShaderPortsAPI::MakeInputValueKey(Lerp::Inputs::A.ToString()), DefaultValue.Get<float>());
 					}
 					else if (DefaultValue.IsType<FLinearColor>())
 					{
-						LerpNode->AddLinearColorAttribute(UInterchangeShaderPortsAPI::MakeInputValueKey(Lerp::Inputs::B.ToString()), DefaultValue.Get<FLinearColor>());
+						LerpNode->AddLinearColorAttribute(UInterchangeShaderPortsAPI::MakeInputValueKey(Lerp::Inputs::A.ToString()), DefaultValue.Get<FLinearColor>());
 					}
 					
 					const FString WeightNodeName = InputName.ToString() + TEXT("MapWeight");
 					UInterchangeShaderNode* WeightNode = UInterchangeShaderNode::Create(&NodeContainer, WeightNodeName, LerpNode->GetUniqueID());
 					WeightNode->SetCustomShaderType(ScalarParameter::Name.ToString());
-
-					const float InverseFactor = 1.f - Factor; // We lerp from A to B and prefer to put the strongest input in A so we need to flip the lerp factor
-					WeightNode->AddFloatAttribute(UInterchangeShaderPortsAPI::MakeInputParameterKey(ScalarParameter::Attributes::DefaultValue.ToString()), InverseFactor);
+										
+					WeightNode->AddFloatAttribute(UInterchangeShaderPortsAPI::MakeInputParameterKey(ScalarParameter::Attributes::DefaultValue.ToString()), Factor);
 
 					UInterchangeShaderPortsAPI::ConnectDefaultOuputToInput(LerpNode, Lerp::Inputs::Factor.ToString() , WeightNode->GetUniqueID());
 
 					UInterchangeShaderPortsAPI::ConnectDefaultOuputToInput(NodeToConnectTo, InputToConnectTo, LerpNode->GetUniqueID());
 
 					NodeToConnectTo = LerpNode;
-					InputToConnectTo = Lerp::Inputs::A.ToString();
+					InputToConnectTo = Lerp::Inputs::B.ToString();
 				}
 
 				// Handles max one texture per property.
