Interchange Level Sequence:
- Adding animation tracks for Cameras: AspectRatio, AutoCalculateOrthoPlanes, FilmBackSensorAspectRatio, FilmBackSensorHeight, FilmBackSensorWidth, CameraOrthoNearClipPlane

Interchange Animation Track Set Node:
- deprecated function to Get/Set an animation property track based on an FName, using an ENum instead.

Interchange FBX Settings:
- Removed the enum InterchangePayload Type variable, letting the Translator guess the type of curves based on the incoming type (Integer/Bool/Enum types should always be treated as Step curves)

FBX export:
- Fixed the export of the FieldOfView being exported as a FocusLength
- Export light Intensity as a node attribute instead of a custom attribute

diff --git a/Engine/Plugins/Experimental/Enterprise/DatasmithInterchange/Source/DatasmithInterchange/Private/InterchangeDatasmithAnimationUtils.cpp b/Engine/Plugins/Experimental/Enterprise/DatasmithInterchange/Source/DatasmithInterchange/Private/InterchangeDatasmithAnimationUtils.cpp
--- a/Engine/Plugins/Experimental/Enterprise/DatasmithInterchange/Source/DatasmithInterchange/Private/InterchangeDatasmithAnimationUtils.cpp
+++ b/Engine/Plugins/Experimental/Enterprise/DatasmithInterchange/Source/DatasmithInterchange/Private/InterchangeDatasmithAnimationUtils.cpp
@@ -138,7 +138,7 @@ namespace UE::DatasmithInterchange::AnimUtils
 		const FString ActorNodeUid = NodeUtils::GetActorUid(VisibilityAnimation.GetName());
 		TrackNode->SetCustomActorDependencyUid(ActorNodeUid);
 
-		TrackNode->SetCustomPropertyTrack(UE::Interchange::Animation::PropertyTracks::Visibility);
+		TrackNode->SetCustomPropertyTrack(EInterchangePropertyTracks::Visibility);
 		TrackNode->SetCustomAnimationPayloadKey(TrackNode->GetUniqueID(), EInterchangeAnimationPayLoadType::STEPCURVE);
 
 		LevelSequenceNode->AddCustomAnimationTrackUid(TrackNode->GetUniqueID());
diff --git a/Engine/Plugins/Interchange/Runtime/Source/Import/Private/Animation/InterchangeLevelSequenceFactory.cpp b/Engine/Plugins/Interchange/Runtime/Source/Import/Private/Animation/InterchangeLevelSequenceFactory.cpp
--- a/Engine/Plugins/Interchange/Runtime/Source/Import/Private/Animation/InterchangeLevelSequenceFactory.cpp
+++ b/Engine/Plugins/Interchange/Runtime/Source/Import/Private/Animation/InterchangeLevelSequenceFactory.cpp
@@ -430,7 +430,7 @@ namespace UE::Interchange::Private
 
 	void FLevelSequenceHelper::PopulateAnimationTrack(const UInterchangeAnimationTrackNode& AnimationTrackNode)
 	{
-		FName PropertyTrack;
+		EInterchangePropertyTracks PropertyTrack;
 		if(!AnimationTrackNode.GetCustomPropertyTrack(PropertyTrack))
 		{
 			return;
@@ -461,7 +461,7 @@ namespace UE::Interchange::Private
 		}
 
 		FGuid ObjectBinding;
-		if(PropertyTrack == UE::Interchange::Animation::PropertyTracks::Visibility)
+		if(PropertyTrack == EInterchangePropertyTracks::Visibility)
 		{
 			ObjectBinding = BindActorToLevelSequence(Actor);
 		}
diff --git a/Engine/Plugins/Interchange/Runtime/Source/Import/Private/Animation/InterchangeLevelSequenceHelper.cpp b/Engine/Plugins/Interchange/Runtime/Source/Import/Private/Animation/InterchangeLevelSequenceHelper.cpp
--- a/Engine/Plugins/Interchange/Runtime/Source/Import/Private/Animation/InterchangeLevelSequenceHelper.cpp
+++ b/Engine/Plugins/Interchange/Runtime/Source/Import/Private/Animation/InterchangeLevelSequenceHelper.cpp
@@ -19,26 +19,36 @@ namespace UE::Interchange::Private
 {
 	FInterchangePropertyTracksHelper::FInterchangePropertyTracksHelper()
 		: PropertyTracks{
+			//Common
+			{EInterchangePropertyTracks::Visibility, {UMovieSceneVisibilityTrack::StaticClass()->GetName(), AActor::GetHiddenPropertyName().ToString(), TEXT("Visibility")}},
+			{EInterchangePropertyTracks::AutoActivate, {UMovieSceneBoolTrack::StaticClass()->GetName(), GET_MEMBER_NAME_STRING_CHECKED(UActorComponent, bAutoActivate), TEXT("Auto Activate")}},
+
 			// Light
-			{UE::Interchange::Animation::PropertyTracks::Light::Color, {UMovieSceneColorTrack::StaticClass()->GetName(), GET_MEMBER_NAME_STRING_CHECKED(ULocalLightComponent, LightColor), TEXT("Light Color")}},
-			{UE::Interchange::Animation::PropertyTracks::Light::Intensity, {UMovieSceneFloatTrack::StaticClass()->GetName(), GET_MEMBER_NAME_STRING_CHECKED(ULightComponent, Intensity), TEXT("Intensity")}},
-			{UE::Interchange::Animation::PropertyTracks::Light::IntensityUnits, {UMovieSceneByteTrack::StaticClass()->GetName(), GET_MEMBER_NAME_STRING_CHECKED(ULocalLightComponent, IntensityUnits), TEXT("Intensity Units"), StaticEnum<ELightUnits>()}},
-			{UE::Interchange::Animation::PropertyTracks::Light::Temperature, {UMovieSceneFloatTrack::StaticClass()->GetName(), GET_MEMBER_NAME_STRING_CHECKED(ULocalLightComponent, Temperature), TEXT("Temperature")}},
-			{UE::Interchange::Animation::PropertyTracks::Light::UseTemperature, {UMovieSceneBoolTrack::StaticClass()->GetName(), GET_MEMBER_NAME_STRING_CHECKED(ULocalLightComponent, bUseTemperature), TEXT("Use Temperature")}},
-			{UE::Interchange::Animation::PropertyTracks::Visibility, {UMovieSceneVisibilityTrack::StaticClass()->GetName(), AActor::GetHiddenPropertyName().ToString(), TEXT("Visibility")}},
+			{EInterchangePropertyTracks::LightColor, {UMovieSceneColorTrack::StaticClass()->GetName(), GET_MEMBER_NAME_STRING_CHECKED(ULocalLightComponent, LightColor), TEXT("Light Color")}},
+			{EInterchangePropertyTracks::LightIntensity, {UMovieSceneFloatTrack::StaticClass()->GetName(), GET_MEMBER_NAME_STRING_CHECKED(ULightComponent, Intensity), TEXT("Intensity")}},
+			{EInterchangePropertyTracks::LightIntensityUnits, {UMovieSceneByteTrack::StaticClass()->GetName(), GET_MEMBER_NAME_STRING_CHECKED(ULocalLightComponent, IntensityUnits), TEXT("Intensity Units"), StaticEnum<ELightUnits>()}},
+			{EInterchangePropertyTracks::LightTemperature, {UMovieSceneFloatTrack::StaticClass()->GetName(), GET_MEMBER_NAME_STRING_CHECKED(ULocalLightComponent, Temperature), TEXT("Temperature")}},
+			{EInterchangePropertyTracks::LightUseTemperature, {UMovieSceneBoolTrack::StaticClass()->GetName(), GET_MEMBER_NAME_STRING_CHECKED(ULocalLightComponent, bUseTemperature), TEXT("Use Temperature")}},
 
 			// Camera
-			{UE::Interchange::Animation::PropertyTracks::Camera::AspectRatioAxisConstraint, {UMovieSceneByteTrack::StaticClass()->GetName(), GET_MEMBER_NAME_STRING_CHECKED(UCameraComponent, AspectRatioAxisConstraint), TEXT("Aspect Ratio Axis Constraint"), StaticEnum<EAspectRatioAxisConstraint>()}},
-			{UE::Interchange::Animation::PropertyTracks::Camera::ConstrainAspectRatio, {UMovieSceneBoolTrack::StaticClass()->GetName(), GET_MEMBER_NAME_STRING_CHECKED(UCameraComponent, bConstrainAspectRatio), TEXT("Constrain Aspect Ratio")}},
-			{UE::Interchange::Animation::PropertyTracks::Camera::CurrentAperture, {UMovieSceneFloatTrack::StaticClass()->GetName(), GET_MEMBER_NAME_STRING_CHECKED(UCineCameraComponent, CurrentAperture), TEXT("Current Aperture")}},
-			{UE::Interchange::Animation::PropertyTracks::Camera::CurrentFocalLength, {UMovieSceneFloatTrack::StaticClass()->GetName(), GET_MEMBER_NAME_STRING_CHECKED(UCineCameraComponent, CurrentFocalLength), TEXT("Current Focal Length")}},
-			{UE::Interchange::Animation::PropertyTracks::Camera::CustomNearClippingPlane, {UMovieSceneFloatTrack::StaticClass()->GetName(), GET_MEMBER_NAME_STRING_CHECKED(UCineCameraComponent, CustomNearClippingPlane), TEXT("Custom Near Clipping Plane")}},
-			{UE::Interchange::Animation::PropertyTracks::Camera::FieldOfView, {UMovieSceneFloatTrack::StaticClass()->GetName(), GET_MEMBER_NAME_STRING_CHECKED(UCameraComponent, FieldOfView), TEXT("Field of View")}},
-			{UE::Interchange::Animation::PropertyTracks::Camera::OrthoFarClipPlane, {UMovieSceneFloatTrack::StaticClass()->GetName(), GET_MEMBER_NAME_STRING_CHECKED(UCameraComponent, OrthoFarClipPlane), TEXT("Ortho Far Clip Plane")}},
-			{UE::Interchange::Animation::PropertyTracks::Camera::OrthoWidth, {UMovieSceneFloatTrack::StaticClass()->GetName(), GET_MEMBER_NAME_STRING_CHECKED(UCameraComponent, OrthoWidth), TEXT("Ortho Width")}},
-			{UE::Interchange::Animation::PropertyTracks::Camera::PostProcessBlendWeight, {UMovieSceneFloatTrack::StaticClass()->GetName(), GET_MEMBER_NAME_STRING_CHECKED(UCameraComponent, PostProcessBlendWeight), TEXT("PostProcess Blend Weight")}},
-			{UE::Interchange::Animation::PropertyTracks::Camera::ProjectionMode, {UMovieSceneByteTrack::StaticClass()->GetName(), GET_MEMBER_NAME_STRING_CHECKED(UCameraComponent, ProjectionMode), TEXT("Projection Mode"), StaticEnum<ECameraProjectionMode::Type>()}},
-			{UE::Interchange::Animation::PropertyTracks::Camera::UseFieldOfViewForLOD, {UMovieSceneBoolTrack::StaticClass()->GetName(), GET_MEMBER_NAME_STRING_CHECKED(UCameraComponent, bUseFieldOfViewForLOD), TEXT("Use Field Of View For LOD")}},
+			{EInterchangePropertyTracks::CameraAutoCalculateOrthoPlanes, {UMovieSceneBoolTrack::StaticClass()->GetName(), GET_MEMBER_NAME_STRING_CHECKED(UCameraComponent, bAutoCalculateOrthoPlanes), TEXT("Auto Calculate Ortho Planes")}},
+			{EInterchangePropertyTracks::CameraAspectRatio, {UMovieSceneFloatTrack::StaticClass()->GetName(), GET_MEMBER_NAME_STRING_CHECKED(UCameraComponent, AspectRatio), TEXT("Aspect Ratio")}},
+			{EInterchangePropertyTracks::CameraAspectRatioAxisConstraint, {UMovieSceneByteTrack::StaticClass()->GetName(), GET_MEMBER_NAME_STRING_CHECKED(UCameraComponent, AspectRatioAxisConstraint), TEXT("Aspect Ratio Axis Constraint"), StaticEnum<EAspectRatioAxisConstraint>()}},
+			{EInterchangePropertyTracks::CameraConstrainAspectRatio, {UMovieSceneBoolTrack::StaticClass()->GetName(), GET_MEMBER_NAME_STRING_CHECKED(UCameraComponent, bConstrainAspectRatio), TEXT("Constrain Aspect Ratio")}},
+			{EInterchangePropertyTracks::CameraCurrentAperture, {UMovieSceneFloatTrack::StaticClass()->GetName(), GET_MEMBER_NAME_STRING_CHECKED(UCineCameraComponent, CurrentAperture), TEXT("Current Aperture")}},
+			{EInterchangePropertyTracks::CameraCurrentFocalLength, {UMovieSceneFloatTrack::StaticClass()->GetName(), GET_MEMBER_NAME_STRING_CHECKED(UCineCameraComponent, CurrentFocalLength), TEXT("Current Focal Length")}},
+			{EInterchangePropertyTracks::CameraCustomNearClippingPlane, {UMovieSceneFloatTrack::StaticClass()->GetName(), GET_MEMBER_NAME_STRING_CHECKED(UCineCameraComponent, CustomNearClippingPlane), TEXT("Custom Near Clipping Plane")}},
+			{EInterchangePropertyTracks::CameraFieldOfView, {UMovieSceneFloatTrack::StaticClass()->GetName(), GET_MEMBER_NAME_STRING_CHECKED(UCameraComponent, FieldOfView), TEXT("Field of View")}},
+			{EInterchangePropertyTracks::CameraFilmbackSensorAspectRatio, {UMovieSceneFloatTrack::StaticClass()->GetName(), GET_MEMBER_NAME_STRING_CHECKED(UCineCameraComponent, Filmback) + FString{TEXT(".")} + GET_MEMBER_NAME_STRING_CHECKED(FCameraFilmbackSettings, SensorAspectRatio) , TEXT("Sensor Aspect Ratio (Filmback)")}},
+			{EInterchangePropertyTracks::CameraFilmbackSensorHeight, {UMovieSceneFloatTrack::StaticClass()->GetName(), GET_MEMBER_NAME_STRING_CHECKED(UCineCameraComponent, Filmback) + FString{TEXT(".")} + GET_MEMBER_NAME_STRING_CHECKED(FCameraFilmbackSettings, SensorHeight) , TEXT("Sensor Height (Filmback)")}},
+			{EInterchangePropertyTracks::CameraFilmbackSensorWidth, {UMovieSceneFloatTrack::StaticClass()->GetName(), GET_MEMBER_NAME_STRING_CHECKED(UCineCameraComponent, Filmback) + FString{TEXT(".")} + GET_MEMBER_NAME_STRING_CHECKED(FCameraFilmbackSettings, SensorWidth) , TEXT("Sensor Width (Filmback)")}},
+			{EInterchangePropertyTracks::CameraFocusSettingsManualFocusDistance, {UMovieSceneFloatTrack::StaticClass()->GetName(), GET_MEMBER_NAME_STRING_CHECKED(UCineCameraComponent, FocusSettings) + FString{TEXT(".")} + GET_MEMBER_NAME_STRING_CHECKED(FCameraFocusSettings, ManualFocusDistance), TEXT("Manuel Focus Distance (Focus Settings)")}},
+			{EInterchangePropertyTracks::CameraOrthoFarClipPlane, {UMovieSceneFloatTrack::StaticClass()->GetName(), GET_MEMBER_NAME_STRING_CHECKED(UCameraComponent, OrthoFarClipPlane), TEXT("Ortho Far Clip Plane")}},
+			{EInterchangePropertyTracks::CameraOrthoNearClipPlane, {UMovieSceneFloatTrack::StaticClass()->GetName(), GET_MEMBER_NAME_STRING_CHECKED(UCameraComponent, OrthoNearClipPlane), TEXT("Ortho Near Clip Plane")}},
+			{EInterchangePropertyTracks::CameraOrthoWidth, {UMovieSceneFloatTrack::StaticClass()->GetName(), GET_MEMBER_NAME_STRING_CHECKED(UCameraComponent, OrthoWidth), TEXT("Ortho Width")}},
+			{EInterchangePropertyTracks::CameraPostProcessBlendWeight, {UMovieSceneFloatTrack::StaticClass()->GetName(), GET_MEMBER_NAME_STRING_CHECKED(UCameraComponent, PostProcessBlendWeight), TEXT("PostProcess Blend Weight")}},
+			{EInterchangePropertyTracks::CameraProjectionMode, {UMovieSceneByteTrack::StaticClass()->GetName(), GET_MEMBER_NAME_STRING_CHECKED(UCameraComponent, ProjectionMode), TEXT("Projection Mode"), StaticEnum<ECameraProjectionMode::Type>()}},
+			{EInterchangePropertyTracks::CameraUseFieldOfViewForLOD, {UMovieSceneBoolTrack::StaticClass()->GetName(), GET_MEMBER_NAME_STRING_CHECKED(UCameraComponent, bUseFieldOfViewForLOD), TEXT("Use Field Of View For LOD")}},
 		}
 	{}
 
@@ -48,7 +58,7 @@ namespace UE::Interchange::Private
 		return Instance;
 	}
 
-	UMovieSceneSection* FInterchangePropertyTracksHelper::GetSection(UMovieScene* MovieScene, const UInterchangeAnimationTrackNode& AnimationTrackNode, const FGuid& ObjectBinding, const FName& Property) const
+	UMovieSceneSection* FInterchangePropertyTracksHelper::GetSection(UMovieScene* MovieScene, const UInterchangeAnimationTrackNode& AnimationTrackNode, const FGuid& ObjectBinding, EInterchangePropertyTracks Property) const
 	{
 		const FInterchangeProperty* InterchangePropertyTrack = PropertyTracks.Find(Property);
 
diff --git a/Engine/Plugins/Interchange/Runtime/Source/Import/Private/Animation/InterchangeLevelSequenceHelper.h b/Engine/Plugins/Interchange/Runtime/Source/Import/Private/Animation/InterchangeLevelSequenceHelper.h
--- a/Engine/Plugins/Interchange/Runtime/Source/Import/Private/Animation/InterchangeLevelSequenceHelper.h
+++ b/Engine/Plugins/Interchange/Runtime/Source/Import/Private/Animation/InterchangeLevelSequenceHelper.h
@@ -50,7 +50,7 @@ namespace UE::Interchange::Private
 	{
 		static FInterchangePropertyTracksHelper& GetInstance();
 
-		UMovieSceneSection* GetSection(UMovieScene* MovieScene, const UInterchangeAnimationTrackNode& AnimationTrackNode, const FGuid& ObjectBinding, const FName& Property) const;
+		UMovieSceneSection* GetSection(UMovieScene* MovieScene, const UInterchangeAnimationTrackNode& AnimationTrackNode, const FGuid& ObjectBinding, EInterchangePropertyTracks Property) const;
 
 	private:
 		FInterchangePropertyTracksHelper();
@@ -63,6 +63,6 @@ namespace UE::Interchange::Private
 			UEnum* EnumClass = nullptr; // Only used for Enum property tracks
 		};
 
-		TMap<FName, FInterchangeProperty> PropertyTracks;
+		TMap<EInterchangePropertyTracks, FInterchangeProperty> PropertyTracks;
 	};
 }
\ No newline at end of file
diff --git a/Engine/Plugins/Interchange/Runtime/Source/Nodes/Private/InterchangeAnimationTrackSetNode.cpp b/Engine/Plugins/Interchange/Runtime/Source/Nodes/Private/InterchangeAnimationTrackSetNode.cpp
--- a/Engine/Plugins/Interchange/Runtime/Source/Nodes/Private/InterchangeAnimationTrackSetNode.cpp
+++ b/Engine/Plugins/Interchange/Runtime/Source/Nodes/Private/InterchangeAnimationTrackSetNode.cpp
@@ -180,24 +180,24 @@ bool UInterchangeAnimationTrackNode::SetCustomFrameCount(const int32& AttributeV
 	IMPLEMENT_NODE_ATTRIBUTE_SETTER_NODELEGATE(FrameCount, int32);
 }
 
-bool UInterchangeAnimationTrackNode::GetCustomTargetedProperty(int32& AttributeValue) const
+bool UInterchangeAnimationTrackNode::GetCustomPropertyTrack(FName& AttributeValue) const
 {
-	IMPLEMENT_NODE_ATTRIBUTE_GETTER(TargetedProperty, int32);
+	return false;
 }
 
-bool UInterchangeAnimationTrackNode::GetCustomPropertyTrack(FName& AttributeValue) const
+bool UInterchangeAnimationTrackNode::SetCustomPropertyTrack(const FName& AttributeValue)
 {
-	IMPLEMENT_NODE_ATTRIBUTE_GETTER(PropertyTrack, FName);
+	return false;
 }
 
-bool UInterchangeAnimationTrackNode::SetCustomTargetedProperty(const int32& AttributeValue)
+bool UInterchangeAnimationTrackNode::SetCustomPropertyTrack(EInterchangePropertyTracks AttributeValue)
 {
-	IMPLEMENT_NODE_ATTRIBUTE_SETTER_NODELEGATE(TargetedProperty, int32);
+	IMPLEMENT_NODE_ATTRIBUTE_SETTER_NODELEGATE(PropertyTrack, EInterchangePropertyTracks);
 }
 
-bool UInterchangeAnimationTrackNode::SetCustomPropertyTrack(const FName& AttributeValue)
+bool UInterchangeAnimationTrackNode::GetCustomPropertyTrack(EInterchangePropertyTracks& AttributeValue) const
 {
-	IMPLEMENT_NODE_ATTRIBUTE_SETTER_NODELEGATE(PropertyTrack, FName);
+	IMPLEMENT_NODE_ATTRIBUTE_GETTER(PropertyTrack, EInterchangePropertyTracks);
 }
 
 // UInterchangeAnimationTrackSetInstanceNode
diff --git a/Engine/Plugins/Interchange/Runtime/Source/Nodes/Public/InterchangeAnimationDefinitions.h b/Engine/Plugins/Interchange/Runtime/Source/Nodes/Public/InterchangeAnimationDefinitions.h
--- a/Engine/Plugins/Interchange/Runtime/Source/Nodes/Public/InterchangeAnimationDefinitions.h
+++ b/Engine/Plugins/Interchange/Runtime/Source/Nodes/Public/InterchangeAnimationDefinitions.h
@@ -3,6 +3,7 @@
 #pragma once
 
 #include "UObject/NameTypes.h"
+#include "Types/AttributeStorage.h"
 
 #include "InterchangeAnimationDefinitions.generated.h"
 
@@ -11,6 +12,7 @@ enum class EInterchangePropertyTracks : int32
 {
 	/** Common Properties*/
 	Visibility,
+	AutoActivate,
 
 	/** Light Properties */
 	LightColor,
@@ -20,15 +22,21 @@ enum class EInterchangePropertyTracks : int32
 	LightUseTemperature,
 
 	/** Camera Properties*/
-	CameraAutoActivate,
+	CameraAspectRatio,
+	CameraAutoCalculateOrthoPlanes,
 	CameraAspectRatioAxisConstraint,
 	CameraConstrainAspectRatio,
 	CameraCurrentAperture,
 	CameraCurrentFocalLength,
 	CameraCustomNearClippingPlane,
 	CameraFieldOfView,
+	CameraFilmbackSensorAspectRatio,
+	CameraFilmbackSensorHeight,
+	CameraFilmbackSensorWidth,
+	CameraFocusSettingsManualFocusDistance,
 	CameraMobility,
 	CameraOrthoFarClipPlane,
+	CameraOrthoNearClipPlane,
 	CameraOrthoWidth,
 	CameraPostProcessBlendWeight,
 	CameraProjectionMode,
@@ -42,39 +50,19 @@ namespace UE
 {
 	namespace Interchange
 	{
-		namespace Animation
+		template<> struct TAttributeTypeTraits<EInterchangePropertyTracks>
 		{
-			namespace PropertyTracks
+			static constexpr EAttributeTypes GetType()
 			{
-				const FName Visibility = TEXT("Actor.bHidden");
-
-				namespace Light
-				{
-					const FName Color = TEXT("Light.LightColor");
-					const FName Intensity = TEXT("Light.Intensity");
-					const FName IntensityUnits = TEXT("Light.IntensityUnits");
-					const FName Temperature = TEXT("Light.Temperature");
-					const FName UseTemperature = TEXT("Light.bUseTemperature");
-				}
-
-				namespace Camera
-				{
-					const FName AutoActivate = TEXT("Camera.bAutoActivate");
-					const FName AspectRatioAxisConstraint = TEXT("Camera.AspectRatioAxisConstraint");
-					const FName ConstrainAspectRatio = TEXT("Camera.bConstrainAspectRatio");
-					const FName CurrentAperture = TEXT("Camera.CurrentAperture");
-					const FName CurrentFocalLength = TEXT("Camera.CurrentFocalLength");
-					const FName CustomNearClippingPlane = TEXT("Camera.CustomNearClippingPlane");
-					const FName FieldOfView = TEXT("Camera.FieldOfView");
-					const FName Mobility = TEXT("Camera.Mobility");
-					const FName OrthoFarClipPlane = TEXT("Camera.OrthoFarClipPlane");
-					const FName OrthoWidth = TEXT("Camera.OrthoWidth");
-					const FName PostProcessBlendWeight = TEXT("Camera.PostProcessBlendWeight");
-					const FName ProjectionMode = TEXT("Camera.ProjectionMode");
-					const FName ShouldUpdatePhysicsVolume = TEXT("Camera.bShouldUpdatePhysicsVolume");
-					const FName UseFieldOfViewForLOD = TEXT("Camera.bUseFieldOfViewForLOD");
-				}
+				return EAttributeTypes::Int32;
+			}
+			static FString ToString(const uint16& Value)
+			{
+				int32 ValueConv = Value;
+				FStringFormatOrderedArguments OrderedArguments;
+				OrderedArguments.Add(FStringFormatArg(ValueConv));
+				return FString::Format(TEXT("{0}"), OrderedArguments);
 			}
-		}
+		};
 	}
 }
\ No newline at end of file
diff --git a/Engine/Plugins/Interchange/Runtime/Source/Nodes/Public/InterchangeAnimationTrackSetNode.h b/Engine/Plugins/Interchange/Runtime/Source/Nodes/Public/InterchangeAnimationTrackSetNode.h
--- a/Engine/Plugins/Interchange/Runtime/Source/Nodes/Public/InterchangeAnimationTrackSetNode.h
+++ b/Engine/Plugins/Interchange/Runtime/Source/Nodes/Public/InterchangeAnimationTrackSetNode.h
@@ -4,6 +4,7 @@
 
 #include "CoreMinimal.h"
 #include "Nodes/InterchangeBaseNode.h"
+#include "InterchangeAnimationDefinitions.h"
 
 #include "InterchangeAnimationTrackSetNode.generated.h"
 
@@ -310,29 +311,23 @@ public:
 	UFUNCTION(BlueprintCallable, Category = "Interchange | Node | AnimationTrack")
 	bool GetCustomActorDependencyUid(FString& DependencyUid) const;
 
-	/**
-	 * Set the property animated by this track.
-	 */
-	UE_DEPRECATED(5.4, "SetCustomTargetedProperty has been deprecated, please use SetCustomPropertyTrack instead.")
-	UFUNCTION(BlueprintCallable, Category = "Interchange | Node | AnimationTrack")
-	bool SetCustomTargetedProperty(const int32& TargetedProperty);
+	UE_DEPRECATED(5.5, "SetCustomPropertyTrack using an FName has been deprecated, please use the function with the enum instead.")
+	bool SetCustomPropertyTrack(const FName& PropertyTrack);
+
+	UE_DEPRECATED(5.5, "SetCustomPropertyTrack using an FName has been deprecated, please use the function with the enum instead.")
+	bool GetCustomPropertyTrack(FName& PropertyTrack) const;
 
 	/**
-	 * Set the property animated by this track. Usually the name of a UMovieSceneTrack, e.g for UMovieSceneColorTrack -> Color
+	 * Set the property animated by this track.
 	 */
 	UFUNCTION(BlueprintCallable, Category = "Interchange | Node | AnimationTrack")
-	bool SetCustomPropertyTrack(const FName& PropertyTrack);
+	bool SetCustomPropertyTrack(EInterchangePropertyTracks PropertyTrack);
 
 	/**
 	 * Get the property animated by this track.
 	 */
-	UE_DEPRECATED(5.4, "SetCustomTargetedProperty has been deprecated, please use SetCustomPropertyTrack instead.")
-	UFUNCTION(BlueprintCallable, Category = "Interchange | Node | AnimationTrack")
-	bool GetCustomTargetedProperty(int32& TargetedProperty) const;
-
-
 	UFUNCTION(BlueprintCallable, Category = "Interchange | Node | AnimationTrack")
-	bool GetCustomPropertyTrack(FName& PropertyTrack) const;
+	bool GetCustomPropertyTrack(EInterchangePropertyTracks& PropertyTrack) const;
 
 	/**
 	 * Set the payload key needed to retrieve the animation for this track.
diff --git a/Engine/Plugins/Interchange/Runtime/Source/Parsers/Fbx/Private/FbxScene.cpp b/Engine/Plugins/Interchange/Runtime/Source/Parsers/Fbx/Private/FbxScene.cpp
--- a/Engine/Plugins/Interchange/Runtime/Source/Parsers/Fbx/Private/FbxScene.cpp
+++ b/Engine/Plugins/Interchange/Runtime/Source/Parsers/Fbx/Private/FbxScene.cpp
@@ -370,7 +370,7 @@ namespace UE
 					}
 				}
 				
-				auto AddAnimationTrackNode = [&](const FName& PropertyTrack, const FString& CurveNodeName, const FString& PayloadKey, EInterchangeAnimationPayLoadType PayloadType)
+				auto AddAnimationTrackNode = [&](EInterchangePropertyTracks PropertyTrack, const FString& CurveNodeName, const FString& PayloadKey, EInterchangeAnimationPayLoadType PayloadType)
 				{
 					UInterchangeAnimationTrackNode* AnimTrackNode = NewObject< UInterchangeAnimationTrackNode >(&NodeContainer);
 					const FString AnimTrackNodeName = FString::Printf(TEXT("%s"), *UnrealNode->GetDisplayLabel()) + CurveNodeName;
@@ -383,6 +383,30 @@ namespace UE
 					NodeContainer.AddNode(AnimTrackNode);
 				};
 
+				//Bool/Enum/Integer curves are treated as StepCurves, otherwise it's a floating point curve
+				auto GetPropertyPayloadType = [](EFbxType FbxType)
+				{
+					switch(FbxType)
+					{
+					case EFbxType::eFbxBool:
+					case EFbxType::eFbxChar:
+					case EFbxType::eFbxUChar:
+					case EFbxType::eFbxShort:
+					case EFbxType::eFbxUShort:
+					case EFbxType::eFbxInt:
+					case EFbxType::eFbxUInt:
+					case EFbxType::eFbxLongLong:
+					case EFbxType::eFbxULongLong:
+					case EFbxType::eFbxEnum:
+					case EFbxType::eFbxEnumM:
+					case EFbxType::eFbxString:
+						return EInterchangeAnimationPayLoadType::STEPCURVE;
+
+					default:
+						return EInterchangeAnimationPayLoadType::CURVE;
+					}
+				};
+
 				const UInterchangeFbxSettings* InterchangeFbxSettings = GetDefault<UInterchangeFbxSettings>();
 
 				//Add all Node Attributes for the node
@@ -404,9 +428,9 @@ namespace UE
 							if(PayloadKey.IsSet())
 							{
 								const char* CurveNodeName = CurveNode->GetName();								
-								if(FInterchangePropertyTracksSettings PropertyTrack = InterchangeFbxSettings->GetPropertyTrack(CurveNodeName); PropertyTrack.Property != EInterchangePropertyTracks::None)
+								if(EInterchangePropertyTracks PropertyTrack = InterchangeFbxSettings->GetPropertyTrack(CurveNodeName); PropertyTrack != EInterchangePropertyTracks::None)
 								{
-									AddAnimationTrackNode(PropertyTrack.GetPropertyName(), CurveNodeName, *PayloadKey, PropertyTrack.Type);
+									AddAnimationTrackNode(PropertyTrack, CurveNodeName, *PayloadKey, GetPropertyPayloadType(PropertyType));
 								}
 							}
 						}
@@ -433,9 +457,9 @@ namespace UE
 							if(PayloadKey.IsSet())
 							{
 								const char* CurveNodeName = CurveNode->GetName();
-								if(FInterchangePropertyTracksSettings PropertyTrack = InterchangeFbxSettings->GetPropertyTrack(CurveNodeName); PropertyTrack.Property != EInterchangePropertyTracks::None)
+								if(EInterchangePropertyTracks PropertyTrack = InterchangeFbxSettings->GetPropertyTrack(CurveNodeName); PropertyTrack != EInterchangePropertyTracks::None)
 								{
-									AddAnimationTrackNode(PropertyTrack.GetPropertyName(), CurveNodeName, *PayloadKey, PropertyTrack.Type);
+									AddAnimationTrackNode(PropertyTrack, CurveNodeName, *PayloadKey, GetPropertyPayloadType(PropertyType));
 								}
 							}
 						}
@@ -999,4 +1023,4 @@ namespace UE
 	} //ns Interchange
 }//ns UE
 
-#undef LOCTEXT_NAMESPACE
\ No newline at end of file
+#undef LOCTEXT_NAMESPACE
diff --git a/Engine/Plugins/Interchange/Runtime/Source/Parsers/Fbx/Private/InterchangeFbxSettings.cpp b/Engine/Plugins/Interchange/Runtime/Source/Parsers/Fbx/Private/InterchangeFbxSettings.cpp
--- a/Engine/Plugins/Interchange/Runtime/Source/Parsers/Fbx/Private/InterchangeFbxSettings.cpp
+++ b/Engine/Plugins/Interchange/Runtime/Source/Parsers/Fbx/Private/InterchangeFbxSettings.cpp
@@ -8,22 +8,34 @@ UInterchangeFbxSettings::UInterchangeFbxSettings()
 	{
 		PredefinedPropertyTracks =
 		{
-			{ TEXT("Intensity"), { EInterchangePropertyTracks::LightIntensity, EInterchangeAnimationPayLoadType::CURVE }	},
-			{ TEXT("Color"), {EInterchangePropertyTracks::LightColor, EInterchangeAnimationPayLoadType::CURVE} },
-			{ TEXT("bUseTemperature"), {EInterchangePropertyTracks::LightUseTemperature, EInterchangeAnimationPayLoadType::STEPCURVE} },
-			{ TEXT("IntensityUnits"), {EInterchangePropertyTracks::LightIntensityUnits, EInterchangeAnimationPayLoadType::STEPCURVE} },
-			{ TEXT("bHidden"), {EInterchangePropertyTracks::Visibility, EInterchangeAnimationPayLoadType::STEPCURVE} },
-			{ TEXT("CurrentFocalLength"), {EInterchangePropertyTracks::CameraCurrentFocalLength, EInterchangeAnimationPayLoadType::CURVE} },
-			{ TEXT("CurrentAperture"), {EInterchangePropertyTracks::CameraCurrentAperture, EInterchangeAnimationPayLoadType::CURVE} },
-			{ TEXT("AspectRatioAxisConstraint"), {EInterchangePropertyTracks::CameraAspectRatioAxisConstraint, EInterchangeAnimationPayLoadType::CURVE} },
+			{ TEXT("bHidden"), {EInterchangePropertyTracks::Visibility} },
+			{ TEXT("bAutoActivate"), {EInterchangePropertyTracks::AutoActivate} },
+
+			{ TEXT("Intensity"), { EInterchangePropertyTracks::LightIntensity}	},
+			{ TEXT("Color"), {EInterchangePropertyTracks::LightColor} },
+			{ TEXT("bUseTemperature"), {EInterchangePropertyTracks::LightUseTemperature} },
+			{ TEXT("IntensityUnits"), {EInterchangePropertyTracks::LightIntensityUnits} },
+
+			{ TEXT("AspectRatio"), {EInterchangePropertyTracks::CameraAspectRatio} },
+			{ TEXT("AspectRatioAxisConstraint"), {EInterchangePropertyTracks::CameraAspectRatioAxisConstraint} },
+			{ TEXT("bAutoCalculateOrthoPlanes"), {EInterchangePropertyTracks::CameraAutoCalculateOrthoPlanes} },
+			{ TEXT("CurrentAperture"), {EInterchangePropertyTracks::CameraCurrentAperture} },
+			{ TEXT("FieldOfView"), {EInterchangePropertyTracks::CameraFieldOfView} },
+			{ TEXT("FilmAspectRatio"), {EInterchangePropertyTracks::CameraFilmbackSensorAspectRatio} },
+			{ TEXT("FilmHeight"), {EInterchangePropertyTracks::CameraFilmbackSensorHeight} },
+			{ TEXT("FilmWidth"), {EInterchangePropertyTracks::CameraFilmbackSensorWidth} },
+			{ TEXT("FocalLength"), {EInterchangePropertyTracks::CameraCurrentFocalLength} },
+			{ TEXT("OrthoFarClipPlane"), {EInterchangePropertyTracks::CameraOrthoFarClipPlane} },
+			{ TEXT("OrthoNearClipPlane"), {EInterchangePropertyTracks::CameraOrthoNearClipPlane} },
+			{ TEXT("FocusDistance"), {EInterchangePropertyTracks::CameraFocusSettingsManualFocusDistance} },
 		};
 	}
 }
 
-FInterchangePropertyTracksSettings UInterchangeFbxSettings::GetPropertyTrack(const FString& PropertyName) const
+EInterchangePropertyTracks UInterchangeFbxSettings::GetPropertyTrack(const FString& PropertyName) const
 {
-	FInterchangePropertyTracksSettings Result;
-	if(const FInterchangePropertyTracksSettings* Property = PredefinedPropertyTracks.Find(PropertyName))
+	EInterchangePropertyTracks Result = EInterchangePropertyTracks::None;
+	if(const EInterchangePropertyTracks* Property = PredefinedPropertyTracks.Find(PropertyName))
 	{
 		Result = *Property;
 	}
@@ -34,81 +46,3 @@ FInterchangePropertyTracksSettings UInterchangeFbxSettings::GetPropertyTrack(con
 
 	return Result;
 }
-
-FName FInterchangePropertyTracksSettings::GetPropertyName()
-{
-	FName Name;
-
-	switch(Property)
-	{
-	case EInterchangePropertyTracks::Visibility :
-		Name = UE::Interchange::Animation::PropertyTracks::Visibility;
-		break;
-
-	case EInterchangePropertyTracks::CameraAspectRatioAxisConstraint:
-		Name = UE::Interchange::Animation::PropertyTracks::Camera::AspectRatioAxisConstraint;
-		break;
-
-	case EInterchangePropertyTracks::CameraAutoActivate:
-		Name = UE::Interchange::Animation::PropertyTracks::Camera::AutoActivate;
-		break;
-
-	case EInterchangePropertyTracks::CameraConstrainAspectRatio:
-		Name = UE::Interchange::Animation::PropertyTracks::Camera::ConstrainAspectRatio;
-		break;
-
-	case EInterchangePropertyTracks::CameraCurrentAperture:
-		Name = UE::Interchange::Animation::PropertyTracks::Camera::CurrentAperture;
-		break;
-
-	case EInterchangePropertyTracks::CameraCurrentFocalLength:
-		Name = UE::Interchange::Animation::PropertyTracks::Camera::CurrentFocalLength;
-		break;
-
-	case EInterchangePropertyTracks::CameraCustomNearClippingPlane:
-		Name = UE::Interchange::Animation::PropertyTracks::Camera::CustomNearClippingPlane;
-		break;
-
-	case EInterchangePropertyTracks::CameraFieldOfView:
-		Name = UE::Interchange::Animation::PropertyTracks::Camera::FieldOfView;
-		break;
-
-	case EInterchangePropertyTracks::CameraMobility:
-		Name = UE::Interchange::Animation::PropertyTracks::Camera::Mobility;
-		break;
-
-	case EInterchangePropertyTracks::CameraOrthoFarClipPlane:
-		Name = UE::Interchange::Animation::PropertyTracks::Camera::OrthoFarClipPlane;
-		break;
-
-	case EInterchangePropertyTracks::CameraOrthoWidth:
-		Name = UE::Interchange::Animation::PropertyTracks::Camera::OrthoWidth;
-		break;
-
-	case EInterchangePropertyTracks::LightColor:
-		Name = UE::Interchange::Animation::PropertyTracks::Light::Color;
-		break;
-
-	case EInterchangePropertyTracks::LightIntensity:
-		Name = UE::Interchange::Animation::PropertyTracks::Light::Intensity;break;
-		break;
-
-	case EInterchangePropertyTracks::LightIntensityUnits:
-		Name = UE::Interchange::Animation::PropertyTracks::Light::IntensityUnits;
-		break;
-
-	case EInterchangePropertyTracks::LightTemperature:
-		Name = UE::Interchange::Animation::PropertyTracks::Light::Temperature;
-		break;
-
-	case EInterchangePropertyTracks::LightUseTemperature:
-		Name = UE::Interchange::Animation::PropertyTracks::Light::UseTemperature;
-		break;
-
-	case EInterchangePropertyTracks::None:
-	default:
-		checkf(false, TEXT("Unknown Interchange Property Track"));
-	}
-
-	return Name;
-}
diff --git a/Engine/Plugins/Interchange/Runtime/Source/Parsers/Fbx/Public/InterchangeFbxSettings.h b/Engine/Plugins/Interchange/Runtime/Source/Parsers/Fbx/Public/InterchangeFbxSettings.h
--- a/Engine/Plugins/Interchange/Runtime/Source/Parsers/Fbx/Public/InterchangeFbxSettings.h
+++ b/Engine/Plugins/Interchange/Runtime/Source/Parsers/Fbx/Public/InterchangeFbxSettings.h
@@ -12,21 +12,6 @@
 
 #include "InterchangeFbxSettings.generated.h"
 
-USTRUCT(BlueprintType)
-struct FInterchangePropertyTracksSettings
-{
-	GENERATED_BODY()
-
-	UPROPERTY(EditAnywhere, config, Category = "FBX | Property Tracks")
-	EInterchangePropertyTracks Property = EInterchangePropertyTracks::None;
-
-	UPROPERTY(EditAnywhere, config, Category = "FBX | Property Tracks")
-	EInterchangeAnimationPayLoadType Type = EInterchangeAnimationPayLoadType::NONE;
-
-	/** Return the property name of the enum as a string, see InterchangeAnimationDefinitions.h */
-	FName GetPropertyName();
-};
-
 UCLASS(config = Interchange, meta = (DisplayName = "Interchange FBX"))
 class INTERCHANGEFBXPARSER_API UInterchangeFbxSettings : public UDeveloperSettings
 {
@@ -37,11 +22,11 @@ public:
 	UInterchangeFbxSettings();
 
 	/** Search for a predefined property track, if the property has been found it returns it, otherwise we search for a custom property track*/
-	FInterchangePropertyTracksSettings GetPropertyTrack(const FString& PropertyName) const;
+	EInterchangePropertyTracks GetPropertyTrack(const FString& PropertyName) const;
 
 	UPROPERTY(EditAnywhere, config, Category = "FBX | Property Tracks")
-	TMap<FString, FInterchangePropertyTracksSettings> CustomPropertyTracks;
+	TMap<FString, EInterchangePropertyTracks > CustomPropertyTracks;
 
 private:
-	TMap<FString, FInterchangePropertyTracksSettings> PredefinedPropertyTracks;
+	TMap<FString, EInterchangePropertyTracks > PredefinedPropertyTracks;
 };
\ No newline at end of file
diff --git a/Engine/Source/Editor/UnrealEd/Private/Fbx/FbxMainExport.cpp b/Engine/Source/Editor/UnrealEd/Private/Fbx/FbxMainExport.cpp
--- a/Engine/Source/Editor/UnrealEd/Private/Fbx/FbxMainExport.cpp
+++ b/Engine/Source/Editor/UnrealEd/Private/Fbx/FbxMainExport.cpp
@@ -3409,6 +3409,7 @@ void FFbxExporter::ExportLevelSequenceTrackChannels( FbxNode* FbxNode, UMovieSce
 	}
 
 	FbxCamera* FbxCamera = FbxNode->GetCamera();
+	FbxLight* FbxLight = FbxNode->GetLight();
 	FFrameRate TickResolution = Track.GetTypedOuter<UMovieScene>()->GetTickResolution();
 
 	const FName BoolChannelTypeName = FMovieSceneBoolChannel::StaticStruct()->GetFName();
@@ -3457,9 +3458,9 @@ void FFbxExporter::ExportLevelSequenceTrackChannels( FbxNode* FbxNode, UMovieSce
 			FString PropertyName = MetaData.Name.IsNone() ? Track.GetTrackName().ToString() : MetaData.Name.ToString();
 			bool IsFoV = false;
 			// most properties are created as user property, only FOV of camera in FBX supports animation
-			if (PropertyName == "Intensity")
+			if (PropertyName == "Intensity" && FbxLight)
 			{
-				Property = FbxNode->FindProperty("UE_Intensity", false);
+				Property = FbxLight->Intensity;
 			}
 			else if (PropertyName == "FalloffExponent")
 			{
@@ -3471,7 +3472,7 @@ void FFbxExporter::ExportLevelSequenceTrackChannels( FbxNode* FbxNode, UMovieSce
 			}
 			else if (PropertyName == "FieldOfView" && FbxCamera)
 			{
-				Property = FbxCamera->FocalLength;
+				Property = FbxCamera->FieldOfView;
 				IsFoV = true;
 			}
 			else if (PropertyName == "FOVAngle" && FbxCamera)
@@ -3503,6 +3504,18 @@ void FFbxExporter::ExportLevelSequenceTrackChannels( FbxNode* FbxNode, UMovieSce
 			{
 				Property = FbxNode->FindProperty("UE_IntensityUnits", false);
 			}
+			else if(PropertyName == "Filmback.SensorAspectRatio" && FbxCamera)
+			{
+				Property = FbxCamera->FilmAspectRatio;
+			}
+			else if(PropertyName == "Filmback.SensorHeight" && FbxCamera)
+			{
+				Property = FbxCamera->FilmHeight;
+			}
+			else if(PropertyName == "Filmback.SensorWidth" && FbxCamera)
+			{
+				Property = FbxCamera->FilmWidth;
+			}
 
 			if (Property == 0)
 			{
