Interchange: Fbx translator need to remove the namespace when getting any fbx object name. We need to add to the translator setting the bKeepNamespace options.

#rnx

diff --git a/Engine/Plugins/Interchange/Editor/Source/InterchangeEditor/Private/InterchangeFbxAssetImportDataConverter.cpp b/Engine/Plugins/Interchange/Editor/Source/InterchangeEditor/Private/InterchangeFbxAssetImportDataConverter.cpp
--- a/Engine/Plugins/Interchange/Editor/Source/InterchangeEditor/Private/InterchangeFbxAssetImportDataConverter.cpp
+++ b/Engine/Plugins/Interchange/Editor/Source/InterchangeEditor/Private/InterchangeFbxAssetImportDataConverter.cpp
@@ -3,6 +3,7 @@
 
 #include "Animation/AnimSequence.h"
 #include "Animation/Skeleton.h"
+#include "Editor/EditorPerProjectUserSettings.h"
 #include "EditorFramework/AssetImportData.h"
 #include "Engine/SkeletalMesh.h"
 #include "Engine/SkinnedAssetCommon.h"
@@ -663,6 +664,7 @@ namespace UE::Interchange::Private
 		InterchangeFbxTranslatorSettings->bConvertScene = FbxAssetImportData->bConvertScene;
 		InterchangeFbxTranslatorSettings->bForceFrontXAxis = FbxAssetImportData->bForceFrontXAxis;
 		InterchangeFbxTranslatorSettings->bConvertSceneUnit = FbxAssetImportData->bConvertSceneUnit;
+		InterchangeFbxTranslatorSettings->bKeepFbxNamespace = GetDefault<UEditorPerProjectUserSettings>()->bKeepFbxNamespace;
 		DestinationData->SetTranslatorSettings(InterchangeFbxTranslatorSettings);
 
 		if (const UFbxStaticMeshImportData* LegacyStaticMeshImportData = Cast<UFbxStaticMeshImportData>(FbxAssetImportData))
@@ -733,6 +735,7 @@ namespace UE::Interchange::Private
 				InterchangeFbxTranslatorSettings->bConvertScene = FbxAssetImportData->bConvertScene;
 				InterchangeFbxTranslatorSettings->bForceFrontXAxis = FbxAssetImportData->bForceFrontXAxis;
 				InterchangeFbxTranslatorSettings->bConvertSceneUnit = FbxAssetImportData->bConvertSceneUnit;
+				InterchangeFbxTranslatorSettings->bKeepFbxNamespace = GetDefault<UEditorPerProjectUserSettings>()->bKeepFbxNamespace;
 				DestinationData->SetTranslatorSettings(InterchangeFbxTranslatorSettings);
 			};
 
diff --git a/Engine/Plugins/Interchange/Runtime/Config/BaseInterchange.ini b/Engine/Plugins/Interchange/Runtime/Config/BaseInterchange.ini
--- a/Engine/Plugins/Interchange/Runtime/Config/BaseInterchange.ini
+++ b/Engine/Plugins/Interchange/Runtime/Config/BaseInterchange.ini
@@ -66,4 +66,5 @@ PredefinedVDF=((Absorption, "/Interchange/Functions/MX_AbsorptionVDF.MX_Absorpti
 [/Script/InterchangeImport.InterchangeFbxTranslatorProjectSettings]
 bConvertScene=True
 bForceFrontXAxis=False
-bConvertSceneUnit=True
\ No newline at end of file
+bConvertSceneUnit=True
+bKeepFbxNamespace=False
\ No newline at end of file
diff --git a/Engine/Plugins/Interchange/Runtime/Source/Dispatcher/Private/InterchangeDispatcherTask.cpp b/Engine/Plugins/Interchange/Runtime/Source/Dispatcher/Private/InterchangeDispatcherTask.cpp
--- a/Engine/Plugins/Interchange/Runtime/Source/Dispatcher/Private/InterchangeDispatcherTask.cpp
+++ b/Engine/Plugins/Interchange/Runtime/Source/Dispatcher/Private/InterchangeDispatcherTask.cpp
@@ -27,6 +27,7 @@ namespace UE
 			ActionDataObject->SetBoolField(GetDoesConvertSceneJsonKey(), GetDoesConvertScene());
 			ActionDataObject->SetBoolField(GetDoesForceFrontXAxisJsonKey(), GetDoesForceFrontXAxis());
 			ActionDataObject->SetBoolField(GetDoesConvertSceneUnitJsonKey(), GetDoesConvertSceneUnit());
+			ActionDataObject->SetBoolField(GetDoesKeepFbxNamespaceJsonKey(), GetDoesKeepFbxNamespace());
 			CmdObject->SetObjectField(GetCommandDataJsonKey(), ActionDataObject);
 
 			FString LoadSourceCmd;
@@ -90,6 +91,10 @@ namespace UE
 			{
 				return false;
 			}
+			if (!((*ActionDataObject)->TryGetBoolField(GetDoesKeepFbxNamespaceJsonKey(), bKeepFbxNamespace)))
+			{
+				return false;
+			}
 
 			//Since we filled the data from the json file, set the data has been initialize.
 			bIsDataInitialize = true;
diff --git a/Engine/Plugins/Interchange/Runtime/Source/Dispatcher/Public/InterchangeDispatcherTask.h b/Engine/Plugins/Interchange/Runtime/Source/Dispatcher/Public/InterchangeDispatcherTask.h
--- a/Engine/Plugins/Interchange/Runtime/Source/Dispatcher/Public/InterchangeDispatcherTask.h
+++ b/Engine/Plugins/Interchange/Runtime/Source/Dispatcher/Public/InterchangeDispatcherTask.h
@@ -95,12 +95,14 @@ namespace UE
 				, const FString& InSourceFilename
 				, const bool InbConvertScene
 				, const bool InbForceFrontXAxis
-				, const bool InbConvertSceneUnit)
+				, const bool InbConvertSceneUnit
+				, const bool InbKeepFbxNamespace)
 				: TranslatorID(InTranslatorID)
 				, SourceFilename(InSourceFilename)
 				, bConvertScene(InbConvertScene)
 				, bForceFrontXAxis(InbForceFrontXAxis)
 				, bConvertSceneUnit(InbConvertSceneUnit)
+				, bKeepFbxNamespace(InbKeepFbxNamespace)
 			{
 				bIsDataInitialize = true;
 			}
@@ -173,6 +175,19 @@ namespace UE
 				return Key;
 			}
 
+			bool GetDoesKeepFbxNamespace() const
+			{
+				//Code should not do query data if the data was not set before
+				ensure(bIsDataInitialize);
+				return bKeepFbxNamespace;
+			}
+
+			static FString GetDoesKeepFbxNamespaceJsonKey()
+			{
+				static const FString Key = TEXT("KeepFbxNamespace");
+				return Key;
+			}
+
 			/**
 			 * Use this class helper to create the cmd result json string and to read it
 			 */
@@ -205,6 +220,7 @@ namespace UE
 			bool bConvertScene = true;
 			bool bForceFrontXAxis = false;
 			bool bConvertSceneUnit = true;
+			bool bKeepFbxNamespace = false;
 		};
 
 		class INTERCHANGEDISPATCHER_API FJsonFetchPayloadCmd : public IJsonCmdBase
diff --git a/Engine/Plugins/Interchange/Runtime/Source/Import/Private/Fbx/InterchangeFbxTranslator.cpp b/Engine/Plugins/Interchange/Runtime/Source/Import/Private/Fbx/InterchangeFbxTranslator.cpp
--- a/Engine/Plugins/Interchange/Runtime/Source/Import/Private/Fbx/InterchangeFbxTranslator.cpp
+++ b/Engine/Plugins/Interchange/Runtime/Source/Import/Private/Fbx/InterchangeFbxTranslator.cpp
@@ -192,6 +192,7 @@ bool UInterchangeFbxTranslator::Translate(UInterchangeBaseNodeContainer& BaseNod
 	const bool bConvertScene = CacheFbxTranslatorSettings ? CacheFbxTranslatorSettings->bConvertScene : true;
 	const bool bForceFrontXAxis = CacheFbxTranslatorSettings ? CacheFbxTranslatorSettings->bForceFrontXAxis : false;
 	const bool bConvertSceneUnit = CacheFbxTranslatorSettings ? CacheFbxTranslatorSettings->bConvertSceneUnit : true;
+	const bool bKeepFbxNamespace = CacheFbxTranslatorSettings ? CacheFbxTranslatorSettings->bKeepFbxNamespace : false;
 
 	if (bUseWorkerImport)
 	{
@@ -200,7 +201,7 @@ bool UInterchangeFbxTranslator::Translate(UInterchangeBaseNodeContainer& BaseNod
 			return false;
 		}
 		//Create a json command to read the fbx file
-		FString JsonCommand = CreateLoadFbxFileCommand(Filename, bConvertScene, bForceFrontXAxis, bConvertSceneUnit);
+		FString JsonCommand = CreateLoadFbxFileCommand(Filename, bConvertScene, bForceFrontXAxis, bConvertSceneUnit, bKeepFbxNamespace);
 		int32 TaskIndex = Dispatcher->AddTask(JsonCommand);
 
 		//Blocking call until all tasks are executed
@@ -240,7 +241,7 @@ bool UInterchangeFbxTranslator::Translate(UInterchangeBaseNodeContainer& BaseNod
 #if WITH_EDITOR
 		FbxParser.Reset();
 		FbxParser.SetResultContainer(Results);
-		FbxParser.SetConvertSettings(bConvertScene, bForceFrontXAxis, bConvertSceneUnit);
+		FbxParser.SetConvertSettings(bConvertScene, bForceFrontXAxis, bConvertSceneUnit, bKeepFbxNamespace);
 		FbxParser.LoadFbxFile(Filename, BaseNodeContainer);
 #endif
 	}
@@ -604,9 +605,9 @@ TFuture<TOptional<UE::Interchange::FAnimationPayloadData>> UInterchangeFbxTransl
 	return Promise->GetFuture();
 }
 
-FString UInterchangeFbxTranslator::CreateLoadFbxFileCommand(const FString& FbxFilePath, const bool bConvertScene, const bool bForceFrontXAxis, const bool bConvertSceneUnit) const
+FString UInterchangeFbxTranslator::CreateLoadFbxFileCommand(const FString& FbxFilePath, const bool bConvertScene, const bool bForceFrontXAxis, const bool bConvertSceneUnit, const bool bKeepFbxNamespace) const
 {
-	UE::Interchange::FJsonLoadSourceCmd LoadSourceCommand(TEXT("FBX"), FbxFilePath, bConvertScene, bForceFrontXAxis, bConvertSceneUnit);
+	UE::Interchange::FJsonLoadSourceCmd LoadSourceCommand(TEXT("FBX"), FbxFilePath, bConvertScene, bForceFrontXAxis, bConvertSceneUnit, bKeepFbxNamespace);
 	return LoadSourceCommand.ToJson();
 }
 
diff --git a/Engine/Plugins/Interchange/Runtime/Source/Import/Public/Fbx/InterchangeFbxTranslator.h b/Engine/Plugins/Interchange/Runtime/Source/Import/Public/Fbx/InterchangeFbxTranslator.h
--- a/Engine/Plugins/Interchange/Runtime/Source/Import/Public/Fbx/InterchangeFbxTranslator.h
+++ b/Engine/Plugins/Interchange/Runtime/Source/Import/Public/Fbx/InterchangeFbxTranslator.h
@@ -39,6 +39,10 @@ public:
 	/** Whether to convert the scene from FBX unit to UE unit (centimeter). */
 	UPROPERTY(EditAnywhere, Category = "Fbx Translator")
 	bool bConvertSceneUnit = true;
+
+	/** Whether to keep the name space from FBX name. */
+	UPROPERTY(EditAnywhere, Category = "Fbx Translator")
+	bool bKeepFbxNamespace = false;
 };
 
 UCLASS(BlueprintType)
@@ -102,7 +106,7 @@ public:
 	virtual TFuture<TOptional<UE::Interchange::FAnimationPayloadData>> GetAnimationPayloadData(const FInterchangeAnimationPayLoadKey& PayLoadKey, const double BakeFrequency = 0, const double RangeStartSecond = 0, const double RangeStopSecond = 0) const override;
 	/* IInterchangeAnimationPayloadInterface End */
 private:
-	FString CreateLoadFbxFileCommand(const FString& FbxFilePath, const bool bConvertScene, const bool bForceFrontXAxis, const bool bConvertSceneUnit) const;
+	FString CreateLoadFbxFileCommand(const FString& FbxFilePath, const bool bConvertScene, const bool bForceFrontXAxis, const bool bConvertSceneUnit, const bool bKeepFbxNamespace) const;
 
 	FString CreateFetchMeshPayloadFbxCommand(const FString& FbxPayloadKey, const FTransform& MeshGlobalTransform) const;
 
diff --git a/Engine/Plugins/Interchange/Runtime/Source/Parsers/Fbx/Private/FbxAPI.cpp b/Engine/Plugins/Interchange/Runtime/Source/Parsers/Fbx/Private/FbxAPI.cpp
--- a/Engine/Plugins/Interchange/Runtime/Source/Parsers/Fbx/Private/FbxAPI.cpp
+++ b/Engine/Plugins/Interchange/Runtime/Source/Parsers/Fbx/Private/FbxAPI.cpp
@@ -261,10 +261,31 @@ namespace UE
 				}
 			}
 
-			void FFbxParser::EnsureNodeNameAreValid(const FString& BaseFilename)
+			void ManageNamespace(const bool bKeepFbxNamespace, FString& ObjectName, FbxObject* Object)
 			{
-				const bool bKeepNamespace = false;//Todo use: GetDefault<UEditorPerProjectUserSettings>()->bKeepFbxNamespace;
+				if (bKeepFbxNamespace)
+				{
+					if (ObjectName.Contains(TEXT(":")))
+					{
+						ObjectName = ObjectName.Replace(TEXT(":"), TEXT("_"));
+						Object->SetName(TCHAR_TO_UTF8(*ObjectName));
+					}
+				}
+				else
+				{
+					// Remove namespaces
+					int32 LastNamespaceTokenIndex = INDEX_NONE;
+					if (ObjectName.FindLastChar(TEXT(':'), LastNamespaceTokenIndex))
+					{
+						//+1 to remove the ':' character we found
+						ObjectName.RightChopInline(LastNamespaceTokenIndex + 1, EAllowShrinking::Yes);
+						Object->SetName(TCHAR_TO_UTF8(*ObjectName));
+					}
+				}
+			}
 
+			void FFbxParser::EnsureNodeNameAreValid(const FString& BaseFilename)
+			{
 				TSet<FString> AllNodeName;
 				int32 CurrentNameIndex = 1;
 				for (int32 NodeIndex = 0; NodeIndex < SDKScene->GetNodeCount(); ++NodeIndex)
@@ -285,14 +306,8 @@ namespace UE
 							Message->Text = FText::Format(LOCTEXT("EnsureNodeNameAreValid_NoNodeName", "Interchange FBX file Loading: Found node with no name, new node name is '{0}'"), FText::FromString(NodeName));
 						}
 					}
-					if (bKeepNamespace)
-					{
-						if (NodeName.Contains(TEXT(":")))
-						{
-							NodeName = NodeName.Replace(TEXT(":"), TEXT("_"));
-							Node->SetName(TCHAR_TO_UTF8(*NodeName));
-						}
-					}
+					ManageNamespace(bKeepFbxNamespace, NodeName, Node);
+					
 					// Do not allow node to be named same as filename as this creates problems later on (reimport)
 					if (AllNodeName.Contains(NodeName))
 					{
@@ -331,10 +346,10 @@ namespace UE
 					UInterchangeResultWarning_Generic* Message = AddMessage<UInterchangeResultWarning_Generic>();
 					Message->Text = LOCTEXT("MissingBindPose", "Missing bind pose - the FBX SDK has created one.");
 				}
-
-				auto MakeFbxObjectNameUnique = [](FbxObject* Object, TMap<FString, int32>& Names)
+				auto MakeFbxObjectNameUnique = [bKeepFbxNamespaceClosure = bKeepFbxNamespace](FbxObject* Object, TMap<FString, int32>& Names)
 					{
 						FString ObjectName = UTF8_TO_TCHAR(Object->GetName());
+						ManageNamespace(bKeepFbxNamespaceClosure, ObjectName, Object);
 						if (int32* Count = Names.Find(ObjectName))
 						{
 							(*Count)++;
@@ -354,6 +369,11 @@ namespace UE
 				for (int32 NodeIndex = 0; NodeIndex < SDKScene->GetNodeCount(); ++NodeIndex)
 				{
 					FbxNode* Node = SDKScene->GetNode(NodeIndex);
+					FString NodeName = UTF8_TO_TCHAR(Node->GetName());
+					if (NodeName.IsEmpty())
+					{
+						Node->SetName(TCHAR_TO_UTF8(TEXT("Node")));
+					}
 					MakeFbxObjectNameUnique(Node, NodeNames);
 				}
 
@@ -373,8 +393,28 @@ namespace UE
 					{
 						continue;
 					}
+					FString MeshName = UTF8_TO_TCHAR(Mesh->GetName());
+					if (MeshName.IsEmpty())
+					{
+						Mesh->SetName(TCHAR_TO_UTF8(TEXT("Mesh")));
+					}
 					MakeFbxObjectNameUnique(Mesh, MeshNames);
 				}
+
+				/////////////////////////////////////////////////////////////////////////
+				// Ensure Material Name Validity (uniqueness)
+				// Name clash must be global because we will build Unique ID from the material name
+				TMap<FString, int32> MaterialNames;
+				for (int32 MaterialIndex = 0; MaterialIndex < SDKScene->GetMaterialCount(); ++MaterialIndex)
+				{
+					FbxSurfaceMaterial* Material = SDKScene->GetMaterial(MaterialIndex);
+					FString MaterialName = UTF8_TO_TCHAR(Material->GetName());
+					if (MaterialName.IsEmpty())
+					{
+						Material->SetName(TCHAR_TO_UTF8(TEXT("Material")));
+					}
+					MakeFbxObjectNameUnique(Material, MaterialNames);
+				}
 			}
 
 			void FFbxParser::ProcessExtraInformation(UInterchangeBaseNodeContainer& NodeContainer)
diff --git a/Engine/Plugins/Interchange/Runtime/Source/Parsers/Fbx/Private/FbxAPI.h b/Engine/Plugins/Interchange/Runtime/Source/Parsers/Fbx/Private/FbxAPI.h
--- a/Engine/Plugins/Interchange/Runtime/Source/Parsers/Fbx/Private/FbxAPI.h
+++ b/Engine/Plugins/Interchange/Runtime/Source/Parsers/Fbx/Private/FbxAPI.h
@@ -49,11 +49,12 @@ namespace UE
 					ResultsContainer = Result;
 				}
 
-				void SetConvertSettings(const bool InbConvertScene, const bool InbForceFrontXAxis, const bool InbConvertSceneUnit)
+				void SetConvertSettings(const bool InbConvertScene, const bool InbForceFrontXAxis, const bool InbConvertSceneUnit, const bool InbKeepFbxNamespace)
 				{
 					bConvertScene = InbConvertScene;
 					bForceFrontXAxis = InbForceFrontXAxis;
 					bConvertSceneUnit = InbConvertSceneUnit;
+					bKeepFbxNamespace = InbKeepFbxNamespace;
 				}
 
 				//return the fbx helper for this parser
@@ -129,6 +130,7 @@ namespace UE
 				bool bConvertScene = true;
 				bool bForceFrontXAxis = false;
 				bool bConvertSceneUnit = true;
+				bool bKeepFbxNamespace = false;
 
 				struct FileDetails
 				{
diff --git a/Engine/Plugins/Interchange/Runtime/Source/Parsers/Fbx/Private/FbxHelper.cpp b/Engine/Plugins/Interchange/Runtime/Source/Parsers/Fbx/Private/FbxHelper.cpp
--- a/Engine/Plugins/Interchange/Runtime/Source/Parsers/Fbx/Private/FbxHelper.cpp
+++ b/Engine/Plugins/Interchange/Runtime/Source/Parsers/Fbx/Private/FbxHelper.cpp
@@ -126,30 +126,6 @@ namespace UE
 					//Replace None by Null because None clash with NAME_None and the create asset will instead call the object ClassName_X
 					ObjName = TEXT("Null");
 				}
-				
-				//Material name clash have to be sorted here since the unique ID is only compose of the name.
-				//If we do not do it only one material node will be created.
-				if (Object->Is<FbxSurfaceMaterial>())
-				{
-					const FbxObject* SurfaceMaterialClash = MaterialNameClashMap.FindOrAdd(ObjName);
-					if (SurfaceMaterialClash != nullptr && SurfaceMaterialClash != Object)
-					{
-						int32 UniqueID = 1;
-						FString MaterialNameClash;
-						bool bBreak = false;
-						do
-						{
-							MaterialNameClash = ObjName + TEXT(NAMECLASH1_KEY) + FString::FromInt(UniqueID++);
-							SurfaceMaterialClash = MaterialNameClashMap.FindOrAdd(MaterialNameClash);
-							if (SurfaceMaterialClash == nullptr || SurfaceMaterialClash == Object)
-							{
-								bBreak = true;
-							}
-						} while (!bBreak);
-						ObjName = MaterialNameClash;
-					}
-					MaterialNameClashMap.FindChecked(ObjName) = Object;
-				}
 
 				return ObjName;
 			}
diff --git a/Engine/Plugins/Interchange/Runtime/Source/Parsers/Fbx/Private/InterchangeFbxParser.cpp b/Engine/Plugins/Interchange/Runtime/Source/Parsers/Fbx/Private/InterchangeFbxParser.cpp
--- a/Engine/Plugins/Interchange/Runtime/Source/Parsers/Fbx/Private/InterchangeFbxParser.cpp
+++ b/Engine/Plugins/Interchange/Runtime/Source/Parsers/Fbx/Private/InterchangeFbxParser.cpp
@@ -67,9 +67,9 @@ namespace UE::Interchange
 		FbxParserPrivate->SetResultContainer(Result);
 	}
 
-	void FInterchangeFbxParser::SetConvertSettings(const bool InbConvertScene, const bool InbForceFrontXAxis, const bool InbConvertSceneUnit)
+	void FInterchangeFbxParser::SetConvertSettings(const bool InbConvertScene, const bool InbForceFrontXAxis, const bool InbConvertSceneUnit, const bool InbKeepFbxNamespace)
 	{
-		FbxParserPrivate->SetConvertSettings(InbConvertScene, InbForceFrontXAxis, InbConvertSceneUnit);
+		FbxParserPrivate->SetConvertSettings(InbConvertScene, InbForceFrontXAxis, InbConvertSceneUnit, InbKeepFbxNamespace);
 	}
 
 	void FInterchangeFbxParser::LoadFbxFile(const FString& Filename, const FString& ResultFolder)
diff --git a/Engine/Plugins/Interchange/Runtime/Source/Parsers/Fbx/Public/InterchangeFbxParser.h b/Engine/Plugins/Interchange/Runtime/Source/Parsers/Fbx/Public/InterchangeFbxParser.h
--- a/Engine/Plugins/Interchange/Runtime/Source/Parsers/Fbx/Public/InterchangeFbxParser.h
+++ b/Engine/Plugins/Interchange/Runtime/Source/Parsers/Fbx/Public/InterchangeFbxParser.h
@@ -33,7 +33,7 @@ namespace UE
 
 			void SetResultContainer(UInterchangeResultsContainer* Result);
 
-			void SetConvertSettings(const bool InbConvertScene, const bool InbForceFrontXAxis, const bool InbConvertSceneUnit);
+			void SetConvertSettings(const bool InbConvertScene, const bool InbForceFrontXAxis, const bool InbConvertSceneUnit, const bool InbKeepFbxNamespace);
 			/**
 			 * Parse a file support by the fbx sdk. It just extract all the fbx node and create a FBaseNodeContainer and dump it in a json file inside the ResultFolder
 			 * @param - Filename is the file that the fbx sdk will read (.fbx or .obj)
diff --git a/Engine/Source/Programs/InterchangeWorker/Private/InterchangeWorkerImpl.cpp b/Engine/Source/Programs/InterchangeWorker/Private/InterchangeWorkerImpl.cpp
--- a/Engine/Source/Programs/InterchangeWorker/Private/InterchangeWorkerImpl.cpp
+++ b/Engine/Source/Programs/InterchangeWorker/Private/InterchangeWorkerImpl.cpp
@@ -291,7 +291,7 @@ ETaskState FInterchangeWorkerImpl::LoadFbxFile(const FJsonLoadSourceCmd& LoadSou
 	ETaskState ResultState = ETaskState::Unknown;
 	FString SourceFilename = LoadSourceCommand.GetSourceFilename();
 	FbxParser.Reset();
-	FbxParser.SetConvertSettings(LoadSourceCommand.GetDoesConvertScene(), LoadSourceCommand.GetDoesForceFrontXAxis(), LoadSourceCommand.GetDoesConvertSceneUnit());
+	FbxParser.SetConvertSettings(LoadSourceCommand.GetDoesConvertScene(), LoadSourceCommand.GetDoesForceFrontXAxis(), LoadSourceCommand.GetDoesConvertSceneUnit(), LoadSourceCommand.GetDoesKeepFbxNamespace());
 	FbxParser.LoadFbxFile(SourceFilename, ResultFolder);
 	FJsonLoadSourceCmd::JsonResultParser ResultParser;
 	ResultParser.SetResultFilename(FbxParser.GetResultFilepath());
