[Interchange - FBX] Smooth Group preservation bug fixing.

#rnx

diff --git a/Engine/Plugins/Interchange/Runtime/Source/Parsers/Fbx/Private/FbxMesh.cpp b/Engine/Plugins/Interchange/Runtime/Source/Parsers/Fbx/Private/FbxMesh.cpp
--- a/Engine/Plugins/Interchange/Runtime/Source/Parsers/Fbx/Private/FbxMesh.cpp
+++ b/Engine/Plugins/Interchange/Runtime/Source/Parsers/Fbx/Private/FbxMesh.cpp
@@ -638,7 +638,7 @@ bool FMeshDescriptionImporter::FillMeshDescriptionFromFbxMesh(FbxMesh* Mesh, con
 
 	for (int32 i = 0; i < LayerSmoothingCount; i++)
 	{
-		FbxLayerElementSmoothing const* SmoothingInfo = Mesh->GetLayer(0)->GetSmoothing();
+		FbxLayerElementSmoothing const* SmoothingInfo = Mesh->GetLayer(i)->GetSmoothing();
 		if (SmoothingInfo && SmoothingInfo->GetMappingMode() != FbxLayerElement::eByPolygon)
 		{
 			SDKGeometryConverter->ComputePolygonSmoothingFromEdgeSmoothing(Mesh, i);
