Interchange change the material reimport workflow to re-use unmatch material instead of assigning default surface material.

#rnx

diff --git a/Engine/Plugins/Interchange/Runtime/Source/Import/Private/Mesh/InterchangeMeshHelper.cpp b/Engine/Plugins/Interchange/Runtime/Source/Import/Private/Mesh/InterchangeMeshHelper.cpp
--- a/Engine/Plugins/Interchange/Runtime/Source/Import/Private/Mesh/InterchangeMeshHelper.cpp
+++ b/Engine/Plugins/Interchange/Runtime/Source/Import/Private/Mesh/InterchangeMeshHelper.cpp
@@ -1,11 +1,21 @@
 // Copyright Epic Games, Inc. All Rights Reserved. 
 #include "Mesh/InterchangeMeshHelper.h"
 
+#include "Engine/SkinnedAssetCommon.h"
+#include "Engine/StaticMesh.h"
+#include "GenericPlatform/GenericPlatformMisc.h"
+#include "InterchangeHelper.h"
+#include "InterchangeMaterialFactoryNode.h"
 #include "InterchangeSceneNode.h"
+#include "MaterialDomain.h"
+#include "Materials/Material.h"
+#include "Materials/MaterialInterface.h"
 #include "MeshDescription.h"
-#include "StaticMeshOperations.h"
-#include "GenericPlatform/GenericPlatformMisc.h"
+#include "MeshUVChannelInfo.h"
+#include "Nodes/InterchangeBaseNodeContainer.h"
 #include "StaticMeshAttributes.h"
+#include "StaticMeshOperations.h"
+
 
 namespace UE::Interchange::Private::MeshHelper
 {
@@ -67,4 +77,294 @@ namespace UE::Interchange::Private::MeshHelper
 		GlobalTransform = bBakeMeshes ? SceneNodeGeometricTransform * GlobalTransform : SceneNodeGeometricTransform;
 	}
 
+	template<typename MaterialType>
+	class INTERCHANGEIMPORT_API FMeshMaterialViewer
+	{
+	public:
+		FMeshMaterialViewer(TArray<MaterialType>& InMaterials, int32 InMaterialIndex)
+			: Materials(InMaterials)
+			, MaterialIndex(InMaterialIndex)
+		{ }
+
+		FName GetMaterialSlotName() const
+		{
+			if (Materials.IsValidIndex(MaterialIndex))
+			{
+				return Materials[MaterialIndex].MaterialSlotName;
+			}
+			return NAME_None;
+		}
+
+		FName GetImportedMaterialSlotName() const
+		{
+#if WITH_EDITOR
+			if (Materials.IsValidIndex(MaterialIndex))
+			{
+				return Materials[MaterialIndex].ImportedMaterialSlotName;
+			}
+			return NAME_None;
+#else
+			return GetMaterialSlotName();
+#endif
+		}
+
+		void SetMaterialSlotName(const FName MaterialSlotName)
+		{
+			if (Materials.IsValidIndex(MaterialIndex))
+			{
+				Materials[MaterialIndex].MaterialSlotName = MaterialSlotName;
+			}
+		}
+
+		void SetImportedMaterialSlotName(const FName ImportedMaterialSlotName)
+		{
+#if WITH_EDITOR
+			if (Materials.IsValidIndex(MaterialIndex))
+			{
+				Materials[MaterialIndex].ImportedMaterialSlotName = ImportedMaterialSlotName;
+			}
+#endif
+		}
+
+		UMaterialInterface* GetMaterialInterface() const
+		{
+			if (Materials.IsValidIndex(MaterialIndex))
+			{
+				return Materials[MaterialIndex].MaterialInterface;
+			}
+			return nullptr;
+		}
+
+		void SetMaterialInterface(UMaterialInterface* MaterialInterface)
+		{
+			if (Materials.IsValidIndex(MaterialIndex))
+			{
+				Materials[MaterialIndex].MaterialInterface = MaterialInterface;
+			}
+		}
+	private:
+		TArray<MaterialType>& Materials;
+		int32 MaterialIndex = INDEX_NONE;
+	};
+
+	template<typename MaterialType>
+	class INTERCHANGEIMPORT_API FMeshMaterialArrayViewer
+	{
+	public:
+		FMeshMaterialArrayViewer(TArray<MaterialType>& InMaterials, TFunction<void(MaterialType& Material)>& InEmplaceMaterialFunctor)
+			: Materials(InMaterials)
+			, EmplaceMaterialFunctor(InEmplaceMaterialFunctor)
+		{
+			RebuildViewer();
+		}
+
+		void RebuildViewer()
+		{
+			MeshMaterialArrayViewer.Reset(Materials.Num());
+			for (int32 MaterialIndex = 0; MaterialIndex < Materials.Num(); ++MaterialIndex)
+			{
+				MeshMaterialArrayViewer.Emplace(Materials, MaterialIndex);
+			}
+		}
+
+		int32 Num() const
+		{
+			return MeshMaterialArrayViewer.Num();
+		}
+
+		FMeshMaterialViewer<MaterialType>& operator[](int32 MaterialIndex)
+		{
+			check(MeshMaterialArrayViewer.IsValidIndex(MaterialIndex));
+			return MeshMaterialArrayViewer[MaterialIndex];
+		}
+
+		FMeshMaterialViewer<MaterialType>* FindByPredicate(TFunction<bool(const FMeshMaterialViewer<MaterialType>& MaterialViewer)> Predicate)
+		{
+			return MeshMaterialArrayViewer.FindByPredicate(Predicate);
+		}
+
+		void Emplace(UMaterialInterface* NewMaterial, FName MaterialSlotName, FName ImportedMaterialSlotName)
+		{
+			MaterialType& Material = Materials.AddDefaulted_GetRef();
+			Material.MaterialInterface = NewMaterial;
+			Material.MaterialSlotName = MaterialSlotName;
+#if WITH_EDITOR
+			Material.ImportedMaterialSlotName = ImportedMaterialSlotName;
+#endif
+
+			EmplaceMaterialFunctor(Material);
+
+			RebuildViewer();
+		}
+
+		void Reserve(int32 Count)
+		{
+			Materials.Reserve(Count);
+
+			RebuildViewer();
+		}
+
+	private:
+		TArray<MaterialType>& Materials;
+		TFunction<void(MaterialType& Material)>& EmplaceMaterialFunctor;
+		TArray<FMeshMaterialViewer<MaterialType>> MeshMaterialArrayViewer;
+	};
+
+	template<typename MaterialType>
+	void InternalMeshFactorySetupAssetMaterialArray(FMeshMaterialArrayViewer<MaterialType>& ExistMaterialArrayViewer
+		, TMap<FString, FString> ImportedSlotMaterialDependencies
+		, const UInterchangeBaseNodeContainer* NodeContainer
+		, const bool bIsReimport)
+	{
+		//Store the existing material index that match with the imported material
+		TArray<int32> ImportedMaterialMatchExistingMaterialIndex;
+		//Store the imported material index that match with the existing material
+		TArray<int32> ExistingMaterialMatchImportedMaterialIndex;
+
+		const int32 ImportedMaterialCount = ImportedSlotMaterialDependencies.Num();
+		const int32 ExistingMaterialCount = ExistMaterialArrayViewer.Num();
+
+		//Find which imported slot material match with existing slot material
+		{
+			ImportedMaterialMatchExistingMaterialIndex.SetNum(ImportedMaterialCount);
+			for (int32 ImportedMaterialIndex = 0; ImportedMaterialIndex < ImportedMaterialCount; ++ImportedMaterialIndex)
+			{
+				ImportedMaterialMatchExistingMaterialIndex[ImportedMaterialIndex] = INDEX_NONE;
+			}
+
+			ExistingMaterialMatchImportedMaterialIndex.SetNum(ExistingMaterialCount);
+			for (int32 ExistingMaterialIndex = 0; ExistingMaterialIndex < ExistingMaterialCount; ++ExistingMaterialIndex)
+			{
+				ExistingMaterialMatchImportedMaterialIndex[ExistingMaterialIndex] = INDEX_NONE;
+			}
+
+			int32 ImportedMaterialIndex = 0;
+			for (TPair<FString, FString>& SlotMaterialDependency : ImportedSlotMaterialDependencies)
+			{
+				FName MaterialSlotName = *SlotMaterialDependency.Key;
+				for (int32 ExistingMaterialIndex = 0; ExistingMaterialIndex < ExistingMaterialCount; ++ExistingMaterialIndex)
+				{
+					if (ExistingMaterialMatchImportedMaterialIndex[ExistingMaterialIndex] != INDEX_NONE)
+					{
+						continue;
+					}
+
+					const FMeshMaterialViewer<MaterialType>& Material = ExistMaterialArrayViewer[ExistingMaterialIndex];
+					if (Material.GetMaterialSlotName() == MaterialSlotName)
+					{
+						ExistingMaterialMatchImportedMaterialIndex[ExistingMaterialIndex] = ImportedMaterialIndex;
+						ImportedMaterialMatchExistingMaterialIndex[ImportedMaterialIndex] = ExistingMaterialIndex;
+						break;
+					}
+				}
+				ImportedMaterialIndex++;
+			}
+		}
+
+
+		auto UpdateOrAddMaterial = [&ExistMaterialArrayViewer
+			, bIsReimport
+			, &ImportedMaterialMatchExistingMaterialIndex
+			, &ExistingMaterialMatchImportedMaterialIndex
+			, &ExistingMaterialCount]
+			(const FName& MaterialSlotName, UMaterialInterface* MaterialInterface, const int32 ImportedMaterialIndex)
+			{
+				UMaterialInterface* NewMaterial = MaterialInterface ? MaterialInterface : UMaterial::GetDefaultMaterial(MD_Surface);
+
+				FMeshMaterialViewer<MaterialType>* MeshMaterialViewer = ExistMaterialArrayViewer.FindByPredicate([&MaterialSlotName](const FMeshMaterialViewer<MaterialType>& Material) { return Material.GetMaterialSlotName() == MaterialSlotName; });
+				if (MeshMaterialViewer)
+				{
+					//When we are not re-importing, we always force update the material, we should see this case when importing LODs is on since its an import.
+					//When we do a re-import we update the material interface only if the current asset matching material is null and is not the default material.
+					if (!bIsReimport || (MaterialInterface && (!MeshMaterialViewer->GetMaterialInterface() || MeshMaterialViewer->GetMaterialInterface() == UMaterial::GetDefaultMaterial(MD_Surface))))
+					{
+						MeshMaterialViewer->SetMaterialInterface(NewMaterial);
+					}
+				}
+				else
+				{
+					//See if we can pick and unmatched existing material slot before creating one
+					bool bCreateNewMaterialSlot = true;
+					for (int32 ExistingMaterialIndex = 0; ExistingMaterialIndex < ExistingMaterialCount; ++ExistingMaterialIndex)
+					{
+						//Find the next available unmatched existing material slot and pick it up instead of creating a new material
+						if (ExistingMaterialMatchImportedMaterialIndex[ExistingMaterialIndex] == INDEX_NONE)
+						{
+							bCreateNewMaterialSlot = false;
+							FMeshMaterialViewer<MaterialType>& ExistingSkeletalMaterial = ExistMaterialArrayViewer[ExistingMaterialIndex];
+							ExistingSkeletalMaterial.SetMaterialSlotName(MaterialSlotName);
+							ExistingSkeletalMaterial.SetImportedMaterialSlotName(MaterialSlotName);
+							ExistingMaterialMatchImportedMaterialIndex[ExistingMaterialIndex] = ImportedMaterialIndex;
+							ImportedMaterialMatchExistingMaterialIndex[ImportedMaterialIndex] = ExistingMaterialIndex;
+							break;
+						}
+					}
+					if (bCreateNewMaterialSlot)
+					{
+						ExistMaterialArrayViewer.Emplace(NewMaterial, MaterialSlotName, MaterialSlotName);
+					}
+				}
+			};
+
+		//Preallocate the extra memory if needed
+		if (ImportedMaterialCount > ExistingMaterialCount)
+		{
+			ExistMaterialArrayViewer.Reserve(ImportedMaterialCount);
+		}
+
+		int32 ImportedMaterialIndex = 0;
+		for (TPair<FString, FString>& SlotMaterialDependency : ImportedSlotMaterialDependencies)
+		{
+			UE::Interchange::FScopedLambda ScopedLambda([&ImportedMaterialIndex]()
+				{
+					++ImportedMaterialIndex;
+				});
+			FName MaterialSlotName = *SlotMaterialDependency.Key;
+
+			const UInterchangeBaseMaterialFactoryNode* MaterialFactoryNode = Cast<UInterchangeBaseMaterialFactoryNode>(NodeContainer->GetNode(SlotMaterialDependency.Value));
+			if (!MaterialFactoryNode)
+			{
+				UpdateOrAddMaterial(MaterialSlotName, nullptr, ImportedMaterialIndex);
+				continue;
+			}
+
+			FSoftObjectPath MaterialFactoryNodeReferenceObject;
+			MaterialFactoryNode->GetCustomReferenceObject(MaterialFactoryNodeReferenceObject);
+			if (!MaterialFactoryNodeReferenceObject.IsValid())
+			{
+				UpdateOrAddMaterial(MaterialSlotName, nullptr, ImportedMaterialIndex);
+				continue;
+			}
+
+			UMaterialInterface* MaterialInterface = Cast<UMaterialInterface>(MaterialFactoryNodeReferenceObject.ResolveObject());
+			UpdateOrAddMaterial(MaterialSlotName, MaterialInterface ? MaterialInterface : nullptr, ImportedMaterialIndex);
+		}
+	}
+
+	void SkeletalMeshFactorySetupAssetMaterialArray(TArray<FSkeletalMaterial>& ExistMaterials
+		, TMap<FString, FString> ImportedSlotMaterialDependencies
+		, const UInterchangeBaseNodeContainer* NodeContainer
+		, const bool bIsReimport)
+	{
+		TFunction<void(FSkeletalMaterial&)> EmplaceMaterialFunctor = [](FSkeletalMaterial& Material) { };
+		FMeshMaterialArrayViewer<FSkeletalMaterial> MeshMaterialArrayViewer(ExistMaterials, EmplaceMaterialFunctor);
+		InternalMeshFactorySetupAssetMaterialArray<FSkeletalMaterial>(MeshMaterialArrayViewer, ImportedSlotMaterialDependencies, NodeContainer, bIsReimport);
+	}
+
+	void StaticMeshFactorySetupAssetMaterialArray(TArray<FStaticMaterial>& ExistMaterials
+		, TMap<FString, FString> ImportedSlotMaterialDependencies
+		, const UInterchangeBaseNodeContainer* NodeContainer
+		, const bool bIsReimport)
+	{
+		TFunction<void(FStaticMaterial&)> EmplaceMaterialFunctor = [](FStaticMaterial& Material)
+			{
+#if !WITH_EDITOR
+				// UV density is not supported to be generated at runtime for now. We fake that it has been initialized so that we don't trigger ensures.
+				Material.UVChannelData = FMeshUVChannelInfo(1.f);
+#endif
+			};
+		FMeshMaterialArrayViewer<FStaticMaterial> MeshMaterialArrayViewer(ExistMaterials, EmplaceMaterialFunctor);
+		InternalMeshFactorySetupAssetMaterialArray<FStaticMaterial>(MeshMaterialArrayViewer, ImportedSlotMaterialDependencies, NodeContainer, bIsReimport);
+	}
+
 } //ns UE::Interchange::Private::MeshHelper
\ No newline at end of file
diff --git a/Engine/Plugins/Interchange/Runtime/Source/Import/Private/Mesh/InterchangeSkeletalMeshFactory.cpp b/Engine/Plugins/Interchange/Runtime/Source/Import/Private/Mesh/InterchangeSkeletalMeshFactory.cpp
--- a/Engine/Plugins/Interchange/Runtime/Source/Import/Private/Mesh/InterchangeSkeletalMeshFactory.cpp
+++ b/Engine/Plugins/Interchange/Runtime/Source/Import/Private/Mesh/InterchangeSkeletalMeshFactory.cpp
@@ -1277,55 +1277,13 @@ UInterchangeFactoryBase::FImportAssetResult UInterchangeSkeletalMeshFactory::Imp
 		SkeletalMeshFactoryNode->SetCustomVertexColorReplace(bFalseSetting);
 	}
 
-	// Update skeletal materials
-	TArray<FSkeletalMaterial>& Materials = SkeletalMesh->GetMaterials();
-	auto UpdateOrAddSkeletalMaterial = [&Materials, bIsReImport = ImportAssetObjectData.bIsReImport](const FName& MaterialSlotName, UMaterialInterface* MaterialInterface)
-	{
-		UMaterialInterface* NewMaterial = MaterialInterface ? MaterialInterface : UMaterial::GetDefaultMaterial(MD_Surface);
-
-		FSkeletalMaterial* SkeletalMaterial = Materials.FindByPredicate([&MaterialSlotName](const FSkeletalMaterial& Material) { return Material.MaterialSlotName == MaterialSlotName; });
-		if (SkeletalMaterial)
-		{
-			//When we are not re-importing, we always force update the material, we should see this case when importing LODs is on since its an import.
-			//When we do a re-import we update the material interface only if the current asset matching material is null and is not the default material.
-			if (!bIsReImport || (MaterialInterface && (!SkeletalMaterial->MaterialInterface || SkeletalMaterial->MaterialInterface == UMaterial::GetDefaultMaterial(MD_Surface))))
-			{
-				SkeletalMaterial->MaterialInterface = NewMaterial;
-			}
-		}
-		else
-		{
-			const bool bEnableShadowCasting = true;
-			const bool bInRecomputeTangent = false;
-			Materials.Emplace(NewMaterial, bEnableShadowCasting, bInRecomputeTangent, MaterialSlotName, MaterialSlotName);
-		}
-	};
-
-	TMap<FString, FString> SlotMaterialDependencies;
-	SkeletalMeshFactoryNode->GetSlotMaterialDependencies(SlotMaterialDependencies);
-	Materials.Reserve(SlotMaterialDependencies.Num());
 
-	for (TPair<FString, FString>& SlotMaterialDependency : SlotMaterialDependencies)
+	//Call the mesh helper to create the missing material and to use the unmatched existing slot with the unmatch import slot
 	{
-		FName MaterialSlotName = *SlotMaterialDependency.Key;
-
-		const UInterchangeBaseMaterialFactoryNode* MaterialFactoryNode = Cast<UInterchangeBaseMaterialFactoryNode>(Arguments.NodeContainer->GetNode(SlotMaterialDependency.Value));
-		if (!MaterialFactoryNode)
-		{
-			UpdateOrAddSkeletalMaterial(MaterialSlotName, nullptr);
-			continue;
-		}
-
-		FSoftObjectPath MaterialFactoryNodeReferenceObject;
-		MaterialFactoryNode->GetCustomReferenceObject(MaterialFactoryNodeReferenceObject);
-		if (!MaterialFactoryNodeReferenceObject.IsValid())
-		{
-			UpdateOrAddSkeletalMaterial(MaterialSlotName, nullptr);
-			continue;
-		}
-
-		UMaterialInterface* MaterialInterface = Cast<UMaterialInterface>(MaterialFactoryNodeReferenceObject.ResolveObject());
-		UpdateOrAddSkeletalMaterial(MaterialSlotName, MaterialInterface ? MaterialInterface : nullptr);
+		using namespace UE::Interchange::Private::MeshHelper;
+		TMap<FString, FString> SlotMaterialDependencies;
+		SkeletalMeshFactoryNode->GetSlotMaterialDependencies(SlotMaterialDependencies);
+		SkeletalMeshFactorySetupAssetMaterialArray(SkeletalMesh->GetMaterials(), SlotMaterialDependencies, Arguments.NodeContainer, ImportAssetObjectData.bIsReImport);
 	}
 
 	for (int32 LodIndex = 0; LodIndex < LodCount; ++LodIndex)
@@ -1620,6 +1578,7 @@ UInterchangeFactoryBase::FImportAssetResult UInterchangeSkeletalMeshFactory::Imp
 			}
 			else
 			{
+				TArray<FSkeletalMaterial>& Materials = SkeletalMesh->GetMaterials();
 				//LOD 0 import data is reorder to the material array before when building the LOD 0
 				for (int32 MaterialIndex = 0; MaterialIndex < Materials.Num(); ++MaterialIndex)
 				{
diff --git a/Engine/Plugins/Interchange/Runtime/Source/Import/Private/Mesh/InterchangeStaticMeshFactory.cpp b/Engine/Plugins/Interchange/Runtime/Source/Import/Private/Mesh/InterchangeStaticMeshFactory.cpp
--- a/Engine/Plugins/Interchange/Runtime/Source/Import/Private/Mesh/InterchangeStaticMeshFactory.cpp
+++ b/Engine/Plugins/Interchange/Runtime/Source/Import/Private/Mesh/InterchangeStaticMeshFactory.cpp
@@ -93,7 +93,6 @@ namespace UE::Interchange::Private::StaticMesh
 			for (FPolygonGroupID PolygonGroupID : LodMeshDescription->PolygonGroups().GetElementIDs())
 			{
 				FName ImportMaterialName = SlotNames[PolygonGroupID];
-				bool bFoundMatch = false;
 				for (int32 MaterialIndex = 0; MaterialIndex < Materials.Num(); ++MaterialIndex)
 				{
 					FName MaterialName = Materials[MaterialIndex].ImportedMaterialSlotName;
@@ -104,7 +103,6 @@ namespace UE::Interchange::Private::StaticMesh
 						//If the name match say we found the match)
 						if (MaterialName == ImportMaterialName)
 						{
-							bFoundMatch = true;
 							break;
 						}
 						continue;
@@ -114,12 +112,9 @@ namespace UE::Interchange::Private::StaticMesh
 					if (MaterialName == ImportMaterialName)
 					{
 						RemapIndex = ReorderMaterialArray.Add(Materials[MaterialIndex]);
-						bFoundMatch = true;
 						break;
 					}
 				}
-				//All mesh description polygon group should have a match
-				ensure(bFoundMatch);
 			}
 		}
 		//Custom LOD can add materials, so we add them at the end of the material slots
@@ -287,53 +282,13 @@ UInterchangeFactoryBase::FImportAssetResult UInterchangeStaticMeshFactory::Impor
 	bool bKeepSectionsSeparate = false;
 	StaticMeshFactoryNode->GetCustomKeepSectionsSeparate(bKeepSectionsSeparate);
 
-	// Set material slots from imported materials
-	auto UpdateOrAddStaticMaterial = [&StaticMesh, bReimport](const FName& MaterialSlotName, UMaterialInterface* MaterialInterface)
-	{
-		UMaterialInterface* NewMaterial = MaterialInterface ? MaterialInterface : UMaterial::GetDefaultMaterial(MD_Surface);
-
-		FStaticMaterial* StaticMaterial = StaticMesh->GetStaticMaterials().FindByPredicate([&MaterialSlotName](const FStaticMaterial& Material) { return Material.MaterialSlotName == MaterialSlotName; });
-		if (StaticMaterial)
-		{
-			//When we are not re-importing, we always force update the material, we should see this case when importing LODs is on since its an import.
-			//When we do a re-import we update the material interface only if the current asset matching material is null and is not the default material. (this avoid touching a slot that was change by the user)
-			if (!bReimport || (MaterialInterface && (!StaticMaterial->MaterialInterface || StaticMaterial->MaterialInterface == UMaterial::GetDefaultMaterial(MD_Surface))))
-			{
-				StaticMaterial->MaterialInterface = NewMaterial;
-			}
-		}
-		else
-		{
-			int32 MaterialSlotIndex = StaticMesh->GetStaticMaterials().Emplace(NewMaterial, MaterialSlotName);
-#if !WITH_EDITOR
-			// UV density is not supported to be generated at runtime for now. We fake that it has been initialized so that we don't trigger ensures.
-			StaticMesh->GetStaticMaterials()[MaterialSlotIndex].UVChannelData = FMeshUVChannelInfo(1.f);
-#endif
-		}
-	};
-
-	TMap<FString, FString> SlotMaterialDependencies;
-	StaticMeshFactoryNode->GetSlotMaterialDependencies(SlotMaterialDependencies);
-	for (TPair<FString, FString>& SlotMaterialDependency : SlotMaterialDependencies)
+	
+	//Call the mesh helper to create the missing material and to use the unmatched existing slot with the unmatch import slot
 	{
-		FName MaterialSlotName = *SlotMaterialDependency.Key;
-
-		const UInterchangeBaseMaterialFactoryNode* MaterialFactoryNode = Cast<UInterchangeBaseMaterialFactoryNode>(Arguments.NodeContainer->GetNode(SlotMaterialDependency.Value));
-		if (!MaterialFactoryNode)
-		{
-			UpdateOrAddStaticMaterial(MaterialSlotName, nullptr);
-			continue;
-		}
-		FSoftObjectPath MaterialFactoryNodeReferenceObject;
-		MaterialFactoryNode->GetCustomReferenceObject(MaterialFactoryNodeReferenceObject);
-		if (!MaterialFactoryNodeReferenceObject.IsValid())
-		{
-			UpdateOrAddStaticMaterial(MaterialSlotName, nullptr);
-			continue;
-		}
-
-		UMaterialInterface* MaterialInterface = Cast<UMaterialInterface>(MaterialFactoryNodeReferenceObject.TryLoad());
-		UpdateOrAddStaticMaterial(MaterialSlotName, MaterialInterface ? MaterialInterface : nullptr);
+		using namespace UE::Interchange::Private::MeshHelper;
+		TMap<FString, FString> SlotMaterialDependencies;
+		StaticMeshFactoryNode->GetSlotMaterialDependencies(SlotMaterialDependencies);
+		StaticMeshFactorySetupAssetMaterialArray(StaticMesh->GetStaticMaterials(), SlotMaterialDependencies, Arguments.NodeContainer, bReimport);
 	}
 
 	// Now import geometry for each LOD
diff --git a/Engine/Plugins/Interchange/Runtime/Source/Import/Public/Mesh/InterchangeMeshHelper.h b/Engine/Plugins/Interchange/Runtime/Source/Import/Public/Mesh/InterchangeMeshHelper.h
--- a/Engine/Plugins/Interchange/Runtime/Source/Import/Public/Mesh/InterchangeMeshHelper.h
+++ b/Engine/Plugins/Interchange/Runtime/Source/Import/Public/Mesh/InterchangeMeshHelper.h
@@ -2,6 +2,9 @@
 #include "MeshDescription.h"
 #include "StaticMeshOperations.h"
 
+struct FSkeletalMaterial;
+struct FStaticMaterial;
+class UInterchangeBaseNodeContainer;
 class UInterchangeSceneNode;
 
 namespace UE::Interchange::Private::MeshHelper
@@ -10,4 +13,15 @@ namespace UE::Interchange::Private::MeshHelper
 
 	/* return the result of the global transform with the geometric and pivot transform of the scene node. */
 	INTERCHANGEIMPORT_API void AddSceneNodeGeometricAndPivotToGlobalTransform(FTransform& GlobalTransform, const UInterchangeSceneNode* SceneNode, const bool bBakeMeshes, const bool bBakePivotMeshes);
+
+	INTERCHANGEIMPORT_API void SkeletalMeshFactorySetupAssetMaterialArray(TArray<FSkeletalMaterial>& ExistMaterials
+		, TMap<FString, FString> ImportedSlotMaterialDependencies
+		, const UInterchangeBaseNodeContainer* NodeContainer
+		, const bool bIsReimport);
+
+	INTERCHANGEIMPORT_API void StaticMeshFactorySetupAssetMaterialArray(TArray<FStaticMaterial>& ExistMaterials
+		, TMap<FString, FString> ImportedSlotMaterialDependencies
+		, const UInterchangeBaseNodeContainer* NodeContainer
+		, const bool bIsReimport);
+
 } //ns UE::Interchange::Private::MeshHelper
\ No newline at end of file
diff --git a/Engine/Source/Developer/SkeletalMeshUtilitiesCommon/Private/LODUtilities.cpp b/Engine/Source/Developer/SkeletalMeshUtilitiesCommon/Private/LODUtilities.cpp
--- a/Engine/Source/Developer/SkeletalMeshUtilitiesCommon/Private/LODUtilities.cpp
+++ b/Engine/Source/Developer/SkeletalMeshUtilitiesCommon/Private/LODUtilities.cpp
@@ -4336,25 +4336,26 @@ void FLODUtilities::ReorderMaterialSlotToBaseLod(USkeletalMesh* SkeletalMesh)
 		for (int32 ImportedMaterialIndex = 0; ImportedMaterialIndex < ImportedMaterials.Num(); ++ImportedMaterialIndex)
 		{
 			FName ImportedMaterialSlotName = FName(*ImportedMaterials[ImportedMaterialIndex].MaterialImportName);
-			bool bFoundMatch = false;
 			for (int32 MaterialIndex = 0; MaterialIndex < Materials.Num(); ++MaterialIndex)
 			{
+				FName MaterialSlotName = Materials[MaterialIndex].ImportedMaterialSlotName;
 				if (MaterialSlotRemap[MaterialIndex] != INDEX_NONE)
 				{
-					bFoundMatch = true;
+					//If the name match say we found the match)
+					if (MaterialSlotName == ImportedMaterialSlotName)
+					{
+						break;
+					}
 					continue;
 				}
 				int32& RemapIndex = MaterialSlotRemap[MaterialIndex];
-				FName MaterialSlotName = Materials[MaterialIndex].ImportedMaterialSlotName;
+
 				if (MaterialSlotName == ImportedMaterialSlotName)
 				{
 					RemapIndex = ReorderMaterialArray.Add(Materials[MaterialIndex]);
-					bFoundMatch = true;
 					break;
 				}
 			}
-			//All mesh description polygon group should have a match
-			ensure(bFoundMatch);
 		}
 	}
 	//Custom LOD can add materials, so we add them at the end of the material slots
