[Interchange] SkeletalMesh transform calculation fixes.

#rnx

diff --git a/Engine/Plugins/Interchange/Runtime/Source/Import/Private/Mesh/InterchangeSkeletalMeshFactory.cpp b/Engine/Plugins/Interchange/Runtime/Source/Import/Private/Mesh/InterchangeSkeletalMeshFactory.cpp
--- a/Engine/Plugins/Interchange/Runtime/Source/Import/Private/Mesh/InterchangeSkeletalMeshFactory.cpp
+++ b/Engine/Plugins/Interchange/Runtime/Source/Import/Private/Mesh/InterchangeSkeletalMeshFactory.cpp
@@ -1385,13 +1385,13 @@ UInterchangeFactoryBase::FImportAssetResult UInterchangeSkeletalMeshFactory::Imp
 					FString MeshDependencyUid;
 					MeshReference.SceneNode->GetCustomAssetInstanceUid(MeshDependencyUid);
 					MeshReference.MeshNode = Cast<UInterchangeMeshNode>(Arguments.NodeContainer->GetNode(MeshDependencyUid));
-					bool bRootAncestorOfMeshdependency = Arguments.NodeContainer->GetIsAncestor(MeshDependencyUid, RootJointNodeId);
+					bool bRootAncestorOfSceneNode = Arguments.NodeContainer->GetIsAncestor(MeshReference.SceneNode->GetUniqueID(), RootJointNode->GetParentUid());
 					//Cache the scene node global matrix, we will use this matrix to bake the vertices, add the node geometric mesh offset to this matrix to bake it properly
 					FTransform SceneNodeTransform;
 					if (!ImportAssetObjectLODData.bUseTimeZeroAsBindPose || !MeshReference.SceneNode->GetCustomTimeZeroGlobalTransform(Arguments.NodeContainer, GlobalOffsetTransform, SceneNodeTransform))
 					{
 						ensure(MeshReference.SceneNode->GetCustomBindPoseGlobalTransform(Arguments.NodeContainer, GlobalOffsetTransform, SceneNodeTransform));
-						if (bRootAncestorOfMeshdependency)
+						if (bRootAncestorOfSceneNode)
 						{
 							if (!bBakeMeshes)
 							{
