[Interchange] Reimport pointing to new files after cancel fix.

#rnx

diff --git a/Engine/Plugins/Interchange/Runtime/Source/Import/Private/Animation/InterchangeAnimSequenceFactory.cpp b/Engine/Plugins/Interchange/Runtime/Source/Import/Private/Animation/InterchangeAnimSequenceFactory.cpp
--- a/Engine/Plugins/Interchange/Runtime/Source/Import/Private/Animation/InterchangeAnimSequenceFactory.cpp
+++ b/Engine/Plugins/Interchange/Runtime/Source/Import/Private/Animation/InterchangeAnimSequenceFactory.cpp
@@ -1391,6 +1391,36 @@ bool UInterchangeAnimSequenceFactory::SetSourceFilename(const UObject* Object, c
 	return false;
 }
 
+void UInterchangeAnimSequenceFactory::BackupSourceData(const UObject* Object) const
+{
+#if WITH_EDITORONLY_DATA
+	if (const UAnimSequence* TempAnimSequence = Cast<UAnimSequence>(Object))
+	{
+		UE::Interchange::FFactoryCommon::BackupSourceData(TempAnimSequence->AssetImportData.Get());
+	}
+#endif
+}
+
+void UInterchangeAnimSequenceFactory::ReinstateSourceData(const UObject* Object) const
+{
+#if WITH_EDITORONLY_DATA
+	if (const UAnimSequence* TempAnimSequence = Cast<UAnimSequence>(Object))
+	{
+		UE::Interchange::FFactoryCommon::ReinstateSourceData(TempAnimSequence->AssetImportData.Get());
+	}
+#endif
+}
+
+void UInterchangeAnimSequenceFactory::ClearBackupSourceData(const UObject* Object) const
+{
+#if WITH_EDITORONLY_DATA
+	if (const UAnimSequence* TempAnimSequence = Cast<UAnimSequence>(Object))
+	{
+		UE::Interchange::FFactoryCommon::ClearBackupSourceData(TempAnimSequence->AssetImportData.Get());
+	}
+#endif
+}
+
 bool UInterchangeAnimSequenceFactory::IsBoneTrackAnimationValid(const UInterchangeAnimSequenceFactoryNode* AnimSequenceFactoryNode, const FImportAssetObjectParams& Arguments)
 {
 	bool bResult = true;
diff --git a/Engine/Plugins/Interchange/Runtime/Source/Import/Private/Material/InterchangeMaterialFactory.cpp b/Engine/Plugins/Interchange/Runtime/Source/Import/Private/Material/InterchangeMaterialFactory.cpp
--- a/Engine/Plugins/Interchange/Runtime/Source/Import/Private/Material/InterchangeMaterialFactory.cpp
+++ b/Engine/Plugins/Interchange/Runtime/Source/Import/Private/Material/InterchangeMaterialFactory.cpp
@@ -1174,6 +1174,35 @@ bool UInterchangeMaterialFactory::SetSourceFilename(const UObject* Object, const
 	return false;
 }
 
+void UInterchangeMaterialFactory::BackupSourceData(const UObject* Object) const
+{
+#if WITH_EDITORONLY_DATA
+	if (const UMaterialInterface* MaterialInterface = Cast<UMaterialInterface>(Object))
+	{
+		UE::Interchange::FFactoryCommon::BackupSourceData(MaterialInterface->AssetImportData.Get());
+	}
+#endif
+}
+
+void UInterchangeMaterialFactory::ReinstateSourceData(const UObject* Object) const
+{
+#if WITH_EDITORONLY_DATA
+	if (const UMaterialInterface* MaterialInterface = Cast<UMaterialInterface>(Object))
+	{
+		UE::Interchange::FFactoryCommon::ReinstateSourceData(MaterialInterface->AssetImportData.Get());
+	}
+#endif
+}
+void UInterchangeMaterialFactory::ClearBackupSourceData(const UObject* Object) const
+{
+#if WITH_EDITORONLY_DATA
+	if (const UMaterialInterface* MaterialInterface = Cast<UMaterialInterface>(Object))
+	{
+		UE::Interchange::FFactoryCommon::ClearBackupSourceData(MaterialInterface->AssetImportData.Get());
+	}
+#endif
+}
+
 #if WITH_EDITOR
 void UInterchangeMaterialFactory::SetupMaterial(UMaterial* Material, const FImportAssetObjectParams& Arguments, const UInterchangeBaseMaterialFactoryNode* BaseMaterialFactoryNode)
 {
diff --git a/Engine/Plugins/Interchange/Runtime/Source/Import/Private/Mesh/InterchangeSkeletalMeshFactory.cpp b/Engine/Plugins/Interchange/Runtime/Source/Import/Private/Mesh/InterchangeSkeletalMeshFactory.cpp
--- a/Engine/Plugins/Interchange/Runtime/Source/Import/Private/Mesh/InterchangeSkeletalMeshFactory.cpp
+++ b/Engine/Plugins/Interchange/Runtime/Source/Import/Private/Mesh/InterchangeSkeletalMeshFactory.cpp
@@ -2144,6 +2144,38 @@ bool UInterchangeSkeletalMeshFactory::SetSourceFilename(const UObject* Object, c
 	return false;
 }
 
+void UInterchangeSkeletalMeshFactory::BackupSourceData(const UObject* Object) const
+{
+	TRACE_CPUPROFILER_EVENT_SCOPE(UInterchangeSkeletalMeshFactory::BackupSourceData)
+#if WITH_EDITORONLY_DATA
+	if (const USkeletalMesh* SkeletalMesh = Cast<USkeletalMesh>(Object))
+	{
+		UE::Interchange::FFactoryCommon::BackupSourceData(SkeletalMesh->GetAssetImportData());
+	}
+#endif
+}
+
+void UInterchangeSkeletalMeshFactory::ReinstateSourceData(const UObject* Object) const
+{
+	TRACE_CPUPROFILER_EVENT_SCOPE(UInterchangeSkeletalMeshFactory::ReinstateSourceData)
+#if WITH_EDITORONLY_DATA
+		if (const USkeletalMesh* SkeletalMesh = Cast<USkeletalMesh>(Object))
+		{
+			UE::Interchange::FFactoryCommon::ReinstateSourceData(SkeletalMesh->GetAssetImportData());
+		}
+#endif
+}
+void UInterchangeSkeletalMeshFactory::ClearBackupSourceData(const UObject* Object) const
+{
+	TRACE_CPUPROFILER_EVENT_SCOPE(UInterchangeSkeletalMeshFactory::ClearBackupSourceData)
+#if WITH_EDITORONLY_DATA
+		if (const USkeletalMesh* SkeletalMesh = Cast<USkeletalMesh>(Object))
+		{
+			UE::Interchange::FFactoryCommon::ClearBackupSourceData(SkeletalMesh->GetAssetImportData());
+		}
+#endif
+}
+
 bool UInterchangeSkeletalMeshFactory::SetReimportSourceIndex(const UObject* Object, int32 SourceIndex) const
 {
 	TRACE_CPUPROFILER_EVENT_SCOPE(UInterchangeSkeletalMeshFactory::SetReimportSourceIndex)
diff --git a/Engine/Plugins/Interchange/Runtime/Source/Import/Private/Mesh/InterchangeStaticMeshFactory.cpp b/Engine/Plugins/Interchange/Runtime/Source/Import/Private/Mesh/InterchangeStaticMeshFactory.cpp
--- a/Engine/Plugins/Interchange/Runtime/Source/Import/Private/Mesh/InterchangeStaticMeshFactory.cpp
+++ b/Engine/Plugins/Interchange/Runtime/Source/Import/Private/Mesh/InterchangeStaticMeshFactory.cpp
@@ -1834,3 +1834,32 @@ bool UInterchangeStaticMeshFactory::SetSourceFilename(const UObject* Object, con
 	return false;
 }
 
+void UInterchangeStaticMeshFactory::BackupSourceData(const UObject* Object) const
+{
+#if WITH_EDITORONLY_DATA
+	if (const UStaticMesh* StaticMesh = Cast<UStaticMesh>(Object))
+	{
+		UE::Interchange::FFactoryCommon::BackupSourceData(StaticMesh->GetAssetImportData());
+	}
+#endif
+}
+
+void UInterchangeStaticMeshFactory::ReinstateSourceData(const UObject* Object) const
+{
+#if WITH_EDITORONLY_DATA
+	if (const UStaticMesh* StaticMesh = Cast<UStaticMesh>(Object))
+	{
+		UE::Interchange::FFactoryCommon::ReinstateSourceData(StaticMesh->GetAssetImportData());
+	}
+#endif
+}
+
+void UInterchangeStaticMeshFactory::ClearBackupSourceData(const UObject* Object) const
+{
+#if WITH_EDITORONLY_DATA
+	if (const UStaticMesh* StaticMesh = Cast<UStaticMesh>(Object))
+	{
+		UE::Interchange::FFactoryCommon::ClearBackupSourceData(StaticMesh->GetAssetImportData());
+	}
+#endif
+}
diff --git a/Engine/Plugins/Interchange/Runtime/Source/Import/Private/Scene/InterchangeSceneImportAssetFactory.cpp b/Engine/Plugins/Interchange/Runtime/Source/Import/Private/Scene/InterchangeSceneImportAssetFactory.cpp
--- a/Engine/Plugins/Interchange/Runtime/Source/Import/Private/Scene/InterchangeSceneImportAssetFactory.cpp
+++ b/Engine/Plugins/Interchange/Runtime/Source/Import/Private/Scene/InterchangeSceneImportAssetFactory.cpp
@@ -83,6 +83,36 @@ bool UInterchangeSceneImportAssetFactory::SetSourceFilename(const UObject* Objec
 	return false;
 }
 
+void UInterchangeSceneImportAssetFactory::BackupSourceData(const UObject* Object) const
+{
+#if WITH_EDITORONLY_DATA
+	if (const UInterchangeSceneImportAsset* SceneImportAsset = Cast<UInterchangeSceneImportAsset>(Object))
+	{
+		UE::Interchange::FFactoryCommon::BackupSourceData(SceneImportAsset->AssetImportData);
+	}
+#endif
+}
+
+void UInterchangeSceneImportAssetFactory::ReinstateSourceData(const UObject* Object) const
+{
+#if WITH_EDITORONLY_DATA
+	if (const UInterchangeSceneImportAsset* SceneImportAsset = Cast<UInterchangeSceneImportAsset>(Object))
+	{
+		UE::Interchange::FFactoryCommon::ReinstateSourceData(SceneImportAsset->AssetImportData);
+	}
+#endif
+}
+
+void UInterchangeSceneImportAssetFactory::ClearBackupSourceData(const UObject* Object) const
+{
+#if WITH_EDITORONLY_DATA
+	if (const UInterchangeSceneImportAsset* SceneImportAsset = Cast<UInterchangeSceneImportAsset>(Object))
+	{
+		UE::Interchange::FFactoryCommon::ClearBackupSourceData(SceneImportAsset->AssetImportData);
+	}
+#endif
+}
+
 UInterchangeFactoryBase::FImportAssetResult UInterchangeSceneImportAssetFactory::BeginImportAsset_GameThread(const FImportAssetObjectParams& Arguments)
 {
 	TRACE_CPUPROFILER_EVENT_SCOPE(UInterchangeSceneImportAssetFactory::BeginImportAsset_GameThread);
diff --git a/Engine/Plugins/Interchange/Runtime/Source/Import/Private/Texture/InterchangeTextureFactory.cpp b/Engine/Plugins/Interchange/Runtime/Source/Import/Private/Texture/InterchangeTextureFactory.cpp
--- a/Engine/Plugins/Interchange/Runtime/Source/Import/Private/Texture/InterchangeTextureFactory.cpp
+++ b/Engine/Plugins/Interchange/Runtime/Source/Import/Private/Texture/InterchangeTextureFactory.cpp
@@ -2039,4 +2039,34 @@ bool UInterchangeTextureFactory::SetSourceFilename(const UObject* Object, const
 	return false;
 }
 
+void UInterchangeTextureFactory::BackupSourceData(const UObject* Object) const
+{
+#if WITH_EDITORONLY_DATA
+	if (const UTexture* Texture = Cast<UTexture>(Object))
+	{
+		UE::Interchange::FFactoryCommon::BackupSourceData(Texture->AssetImportData.Get());
+	}
+#endif
+}
+
+void UInterchangeTextureFactory::ReinstateSourceData(const UObject* Object) const
+{
+#if WITH_EDITORONLY_DATA
+	if (const UTexture* Texture = Cast<UTexture>(Object))
+	{
+		UE::Interchange::FFactoryCommon::ReinstateSourceData(Texture->AssetImportData.Get());
+	}
+#endif
+}
+
+void UInterchangeTextureFactory::ClearBackupSourceData(const UObject* Object) const
+{
+#if WITH_EDITORONLY_DATA
+	if (const UTexture* Texture = Cast<UTexture>(Object))
+	{
+		UE::Interchange::FFactoryCommon::ClearBackupSourceData(Texture->AssetImportData.Get());
+	}
+#endif
+}
+
 #undef LOCTEXT_NAMESPACE
diff --git a/Engine/Plugins/Interchange/Runtime/Source/Import/Public/Animation/InterchangeAnimSequenceFactory.h b/Engine/Plugins/Interchange/Runtime/Source/Import/Public/Animation/InterchangeAnimSequenceFactory.h
--- a/Engine/Plugins/Interchange/Runtime/Source/Import/Public/Animation/InterchangeAnimSequenceFactory.h
+++ b/Engine/Plugins/Interchange/Runtime/Source/Import/Public/Animation/InterchangeAnimSequenceFactory.h
@@ -36,6 +36,9 @@ public:
 
 	virtual bool GetSourceFilenames(const UObject* Object, TArray<FString>& OutSourceFilenames) const override;
 	virtual bool SetSourceFilename(const UObject* Object, const FString& SourceFilename, int32 SourceIndex) const override;
+	virtual void BackupSourceData(const UObject* Object) const override;
+	virtual void ReinstateSourceData(const UObject* Object) const override;
+	virtual void ClearBackupSourceData(const UObject* Object) const override;
 	
 	struct FBoneTrackData
 	{
diff --git a/Engine/Plugins/Interchange/Runtime/Source/Import/Public/Material/InterchangeMaterialFactory.h b/Engine/Plugins/Interchange/Runtime/Source/Import/Public/Material/InterchangeMaterialFactory.h
--- a/Engine/Plugins/Interchange/Runtime/Source/Import/Public/Material/InterchangeMaterialFactory.h
+++ b/Engine/Plugins/Interchange/Runtime/Source/Import/Public/Material/InterchangeMaterialFactory.h
@@ -77,6 +77,9 @@ public:
 	virtual void SetupObject_GameThread(const FSetupObjectParams& Arguments) override;
 	virtual bool GetSourceFilenames(const UObject* Object, TArray<FString>& OutSourceFilenames) const override;
 	virtual bool SetSourceFilename(const UObject* Object, const FString& SourceFilename, int32 SourceIndex) const override;
+	virtual void BackupSourceData(const UObject* Object) const override;
+	virtual void ReinstateSourceData(const UObject* Object) const override;
+	virtual void ClearBackupSourceData(const UObject* Object) const override;
 
 	// Interchange factory base interface end
 	//////////////////////////////////////////////////////////////////////////
diff --git a/Engine/Plugins/Interchange/Runtime/Source/Import/Public/Mesh/InterchangeSkeletalMeshFactory.h b/Engine/Plugins/Interchange/Runtime/Source/Import/Public/Mesh/InterchangeSkeletalMeshFactory.h
--- a/Engine/Plugins/Interchange/Runtime/Source/Import/Public/Mesh/InterchangeSkeletalMeshFactory.h
+++ b/Engine/Plugins/Interchange/Runtime/Source/Import/Public/Mesh/InterchangeSkeletalMeshFactory.h
@@ -62,6 +62,9 @@ public:
 	virtual bool GetSourceFilenames(const UObject* Object, TArray<FString>& OutSourceFilenames) const override;
 	virtual bool SetSourceFilename(const UObject* Object, const FString& SourceFilename, int32 SourceIndex) const override;
 	virtual bool SetReimportSourceIndex(const UObject* Object, int32 SourceIndex) const override;
+	virtual void BackupSourceData(const UObject* Object) const override;
+	virtual void ReinstateSourceData(const UObject* Object) const override;
+	virtual void ClearBackupSourceData(const UObject* Object) const override;
 
 	// Interchange factory base interface end
 	//////////////////////////////////////////////////////////////////////////
diff --git a/Engine/Plugins/Interchange/Runtime/Source/Import/Public/Mesh/InterchangeStaticMeshFactory.h b/Engine/Plugins/Interchange/Runtime/Source/Import/Public/Mesh/InterchangeStaticMeshFactory.h
--- a/Engine/Plugins/Interchange/Runtime/Source/Import/Public/Mesh/InterchangeStaticMeshFactory.h
+++ b/Engine/Plugins/Interchange/Runtime/Source/Import/Public/Mesh/InterchangeStaticMeshFactory.h
@@ -36,6 +36,9 @@ public:
 	virtual void SetupObject_GameThread(const FSetupObjectParams& Arguments) override;
 	virtual bool GetSourceFilenames(const UObject* Object, TArray<FString>& OutSourceFilenames) const override;
 	virtual bool SetSourceFilename(const UObject* Object, const FString& SourceFilename, int32 SourceIndex) const override;
+	virtual void BackupSourceData(const UObject* Object) const override;
+	virtual void ReinstateSourceData(const UObject* Object) const override;
+	virtual void ClearBackupSourceData(const UObject* Object) const override;
 
 	// Interchange factory base interface end
 	//////////////////////////////////////////////////////////////////////////
diff --git a/Engine/Plugins/Interchange/Runtime/Source/Import/Public/Scene/InterchangeSceneImportAssetFactory.h b/Engine/Plugins/Interchange/Runtime/Source/Import/Public/Scene/InterchangeSceneImportAssetFactory.h
--- a/Engine/Plugins/Interchange/Runtime/Source/Import/Public/Scene/InterchangeSceneImportAssetFactory.h
+++ b/Engine/Plugins/Interchange/Runtime/Source/Import/Public/Scene/InterchangeSceneImportAssetFactory.h
@@ -24,6 +24,9 @@ public:
 	virtual void FinalizeObject_GameThread(const FSetupObjectParams& Arguments) override;
 	virtual bool GetSourceFilenames(const UObject* Object, TArray<FString>& OutSourceFilenames) const override;
 	virtual bool SetSourceFilename(const UObject* Object, const FString& SourceFilename, int32 SourceIndex) const override;
+	virtual void BackupSourceData(const UObject* Object) const override;
+	virtual void ReinstateSourceData(const UObject* Object) const override;
+	virtual void ClearBackupSourceData(const UObject* Object) const override;
 
 	// Interchange factory base interface end
 	//////////////////////////////////////////////////////////////////////////
diff --git a/Engine/Plugins/Interchange/Runtime/Source/Import/Public/Texture/InterchangeTextureFactory.h b/Engine/Plugins/Interchange/Runtime/Source/Import/Public/Texture/InterchangeTextureFactory.h
--- a/Engine/Plugins/Interchange/Runtime/Source/Import/Public/Texture/InterchangeTextureFactory.h
+++ b/Engine/Plugins/Interchange/Runtime/Source/Import/Public/Texture/InterchangeTextureFactory.h
@@ -61,6 +61,9 @@ public:
 	virtual void SetupObject_GameThread(const FSetupObjectParams& Arguments) override;
 	virtual bool GetSourceFilenames(const UObject* Object, TArray<FString>& OutSourceFilenames) const override;
 	virtual bool SetSourceFilename(const UObject* Object, const FString& SourceFilename, int32 SourceIndex) const override;
+	virtual void BackupSourceData(const UObject* Object) const override;
+	virtual void ReinstateSourceData(const UObject* Object) const override;
+	virtual void ClearBackupSourceData(const UObject* Object) const override;
 
 	// Interchange factory base interface end
 	//////////////////////////////////////////////////////////////////////////
diff --git a/Engine/Source/Editor/UnrealEd/Private/Factories/InterchangeReimportHandler.cpp b/Engine/Source/Editor/UnrealEd/Private/Factories/InterchangeReimportHandler.cpp
--- a/Engine/Source/Editor/UnrealEd/Private/Factories/InterchangeReimportHandler.cpp
+++ b/Engine/Source/Editor/UnrealEd/Private/Factories/InterchangeReimportHandler.cpp
@@ -36,6 +36,7 @@ void UInterchangeReimportHandler::SetReimportPaths(UObject* Obj, const FString&
 	if (const UClass* FactoryClass = InterchangeManager.GetRegisteredFactoryClass(Obj->GetClass()))
 	{
 		UInterchangeFactoryBase* FactoryBase = FactoryClass->GetDefaultObject<UInterchangeFactoryBase>();
+		FactoryBase->BackupSourceData(Obj);
 		FactoryBase->SetSourceFilename(Obj, NewReimportPath, SourceFileIndex);
 	}
 }
diff --git a/Engine/Source/Runtime/Interchange/Core/Public/InterchangeFactoryBase.h b/Engine/Source/Runtime/Interchange/Core/Public/InterchangeFactoryBase.h
--- a/Engine/Source/Runtime/Interchange/Core/Public/InterchangeFactoryBase.h
+++ b/Engine/Source/Runtime/Interchange/Core/Public/InterchangeFactoryBase.h
@@ -292,6 +292,27 @@ public:
 		return false;
 	}
 
+	/**
+	 * Backups SourceData from the AssetData.
+	 */
+	virtual void BackupSourceData(const UObject* Object) const
+	{
+	}
+
+	/**
+	 * Reinstates the backedup SourceData from the Backup to the AssetData.
+	 */
+	virtual void ReinstateSourceData(const UObject* Object) const
+	{
+	}
+
+	/**
+	 * Clears the backedup SourceData from the Backup to the AssetData.
+	 */
+	virtual void ClearBackupSourceData(const UObject* Object) const
+	{
+	}
+
 	/**
 	 * Set the object's reimport source at the specified index value.
 	 */
diff --git a/Engine/Source/Runtime/Interchange/Engine/Private/InterchangeAssetImportData.cpp b/Engine/Source/Runtime/Interchange/Engine/Private/InterchangeAssetImportData.cpp
--- a/Engine/Source/Runtime/Interchange/Engine/Private/InterchangeAssetImportData.cpp
+++ b/Engine/Source/Runtime/Interchange/Engine/Private/InterchangeAssetImportData.cpp
@@ -381,4 +381,28 @@ void UInterchangeAssetImportData::ProcessDeprecatedData() const
 			}
 		}
 	}
+}
+
+void UInterchangeAssetImportData::BackupSourceData() const
+{
+#if WITH_EDITORONLY_DATA
+	if (SourceDataBackup.SourceFiles.Num() == 0)
+	{
+		SourceDataBackup = SourceData;
+	}
+#endif
+}
+
+void UInterchangeAssetImportData::ClearBackupSourceData() const
+{
+#if WITH_EDITORONLY_DATA
+	SourceDataBackup = FAssetImportInfo();
+#endif
+}
+
+void UInterchangeAssetImportData::ReinstateBackupSourceData()
+{
+#if WITH_EDITORONLY_DATA
+	SourceData = SourceDataBackup;
+#endif
 }
\ No newline at end of file
diff --git a/Engine/Source/Runtime/Interchange/Engine/Private/InterchangeImportCommon.cpp b/Engine/Source/Runtime/Interchange/Engine/Private/InterchangeImportCommon.cpp
--- a/Engine/Source/Runtime/Interchange/Engine/Private/InterchangeImportCommon.cpp
+++ b/Engine/Source/Runtime/Interchange/Engine/Private/InterchangeImportCommon.cpp
@@ -248,6 +248,33 @@ namespace UE::Interchange
 		return AssetImportData;
 	}
 
+	void FFactoryCommon::BackupSourceData(const UAssetImportData* AssetImportData)
+	{
+		TRACE_CPUPROFILER_EVENT_SCOPE(UE::Interchange::FFactoryCommon::BackupSourceFilenames)
+		if (const UInterchangeAssetImportData* InterchangeAssetImportData = Cast<UInterchangeAssetImportData>(AssetImportData))
+		{
+			InterchangeAssetImportData->BackupSourceData();
+		}
+	}
+
+	void FFactoryCommon::ReinstateSourceData(UAssetImportData* AssetImportData)
+	{
+		TRACE_CPUPROFILER_EVENT_SCOPE(UE::Interchange::FFactoryCommon::ReinstateSourceData)
+		if (UInterchangeAssetImportData* InterchangeAssetImportData = Cast<UInterchangeAssetImportData>(AssetImportData))
+		{
+			InterchangeAssetImportData->ReinstateBackupSourceData();
+		}
+	}
+
+	void FFactoryCommon::ClearBackupSourceData(const UAssetImportData* AssetImportData)
+	{
+		TRACE_CPUPROFILER_EVENT_SCOPE(UE::Interchange::FFactoryCommon::ClearBackupSourceData)
+		if (const UInterchangeAssetImportData* InterchangeAssetImportData = Cast<UInterchangeAssetImportData>(AssetImportData))
+		{
+			InterchangeAssetImportData->ClearBackupSourceData();
+		}
+	}
+
 	bool FFactoryCommon::GetSourceFilenames(const UAssetImportData* AssetImportData, TArray<FString>& OutSourceFilenames)
 	{
 		TRACE_CPUPROFILER_EVENT_SCOPE(UE::Interchange::FFactoryCommon::GetSourceFilenames)
diff --git a/Engine/Source/Runtime/Interchange/Engine/Private/InterchangeManager.cpp b/Engine/Source/Runtime/Interchange/Engine/Private/InterchangeManager.cpp
--- a/Engine/Source/Runtime/Interchange/Engine/Private/InterchangeManager.cpp
+++ b/Engine/Source/Runtime/Interchange/Engine/Private/InterchangeManager.cpp
@@ -1984,6 +1984,28 @@ UInterchangeManager::ImportInternal(const FString& ContentPath, const UInterchan
 
 		AsyncHelper->InitCancel();
 		AsyncHelper->CleanUp();
+
+		//if it is a re-import, re-instate the backedup source data:
+		if (UObject* ReimportCandidate = AsyncHelper->TaskData.ReimportObject)
+		{
+			if (const UClass* FactoryClass = GetRegisteredFactoryClass(ReimportCandidate->GetClass()))
+			{
+				UInterchangeFactoryBase* FactoryBase = FactoryClass->GetDefaultObject<UInterchangeFactoryBase>();
+				FactoryBase->ReinstateSourceData(ReimportCandidate);
+			}
+		}
+	}
+	else
+	{
+		//If it is a re-import, clear the backedup source data:
+		if (UObject* ReimportCandidate = AsyncHelper->TaskData.ReimportObject)
+		{
+			if (const UClass* FactoryClass = GetRegisteredFactoryClass(ReimportCandidate->GetClass()))
+			{
+				UInterchangeFactoryBase* FactoryBase = FactoryClass->GetDefaultObject<UInterchangeFactoryBase>();
+				FactoryBase->ClearBackupSourceData(ReimportCandidate);
+			}
+		}
 	}
 
 	AsyncHelper->ContentBasePath = ContentBasePath;
diff --git a/Engine/Source/Runtime/Interchange/Engine/Public/InterchangeAssetImportData.h b/Engine/Source/Runtime/Interchange/Engine/Public/InterchangeAssetImportData.h
--- a/Engine/Source/Runtime/Interchange/Engine/Public/InterchangeAssetImportData.h
+++ b/Engine/Source/Runtime/Interchange/Engine/Public/InterchangeAssetImportData.h
@@ -166,6 +166,10 @@ public:
 	UFUNCTION(BlueprintCallable, Category = "Interchange | AssetImportData")
 	INTERCHANGEENGINE_API void SetTranslatorSettings(UInterchangeTranslatorSettings* TranslatorSettings) const;
 
+	/** Won't overwrite an existing backup. Backup has to be cleared before a new one can be stored. */
+	INTERCHANGEENGINE_API void BackupSourceData() const;
+	INTERCHANGEENGINE_API void ClearBackupSourceData() const;
+	INTERCHANGEENGINE_API void ReinstateBackupSourceData();
 
 private:
 	UPROPERTY(meta = (DeprecatedProperty, DeprecationMessage = "Use GetNodeContainer/SetNodeContainer instead."))
@@ -189,6 +193,11 @@ private:
 	mutable TArray64<uint8> CachedNodeContainer;
 	mutable TArray<TPair<FString, FString>> CachedPipelines; //Class, Data(serialized JSON) pair
 	mutable TPair<FString, FString> CachedTranslatorSettings;
+
+#if WITH_EDITORONLY_DATA
+	/** Source file data describing the files that were used to import this asset. Temporary, Primary usage is for Re-import cancellations. */
+	mutable FAssetImportInfo SourceDataBackup;
+#endif
 };
 
 /**
diff --git a/Engine/Source/Runtime/Interchange/Engine/Public/InterchangeImportCommon.h b/Engine/Source/Runtime/Interchange/Engine/Public/InterchangeImportCommon.h
--- a/Engine/Source/Runtime/Interchange/Engine/Public/InterchangeImportCommon.h
+++ b/Engine/Source/Runtime/Interchange/Engine/Public/InterchangeImportCommon.h
@@ -91,6 +91,21 @@ namespace UE::Interchange
 		 */
 		static INTERCHANGEENGINE_API bool GetSourceFilenames(const UAssetImportData* AssetImportData, TArray<FString>& OutSourceFilenames);
 
+		/**
+		 * Backups the SourceData. Primary usage for re-instating SourceData on Re-Import cancellation.
+		 */
+		static INTERCHANGEENGINE_API void BackupSourceData(const UAssetImportData* AssetImportData);
+
+		/**
+		 * Reinstates the backedup SourceData. Primary usage for re-instating SourceData on Re-Import cancellation.
+		 */
+		static INTERCHANGEENGINE_API void ReinstateSourceData(UAssetImportData* AssetImportData);
+
+		/**
+		 * Clears the backedup SourceData. Primary usage for re-instating SourceData on Re-Import cancellation.
+		 */
+		static INTERCHANGEENGINE_API void ClearBackupSourceData(const UAssetImportData* AssetImportData);
+
 		/**
 		 * Sets the SourceFileName value at the specified index.
 		 */
