Interchange FBX Material:
- Adding a scalar parameter for shininess instead of a float attribute

diff --git a/Engine/Plugins/Interchange/Runtime/Source/Parsers/Fbx/Private/FbxMaterial.cpp b/Engine/Plugins/Interchange/Runtime/Source/Parsers/Fbx/Private/FbxMaterial.cpp
--- a/Engine/Plugins/Interchange/Runtime/Source/Parsers/Fbx/Private/FbxMaterial.cpp
+++ b/Engine/Plugins/Interchange/Runtime/Source/Parsers/Fbx/Private/FbxMaterial.cpp
@@ -246,7 +246,11 @@ namespace UE
 						UInterchangeShaderPortsAPI::ConnectDefaultOuputToInput(ShaderGraphNode, InputName, MultiplyNode->GetUniqueID());
 
 						// Scale texture output from [0-1] to [0-1000]
-						MultiplyNode->AddFloatAttribute(UInterchangeShaderPortsAPI::MakeInputValueKey(Multiply::Inputs::B.ToString()), 1000.0f);
+						UInterchangeShaderNode* WeightNode = UInterchangeShaderNode::Create(&NodeContainer, TEXT("ShininessMapWeight"), MultiplyNode->GetUniqueID());
+						WeightNode->SetCustomShaderType(ScalarParameter::Name.ToString());
+						WeightNode->AddFloatAttribute(UInterchangeShaderPortsAPI::MakeInputParameterKey(ScalarParameter::Attributes::DefaultValue.ToString()), 1000.f);
+
+						UInterchangeShaderPortsAPI::ConnectDefaultOuputToInput(MultiplyNode, Multiply::Inputs::B.ToString(), WeightNode->GetUniqueID());
 
 						UInterchangeShaderPortsAPI::ConnectDefaultOuputToInput(MultiplyNode, Multiply::Inputs::A.ToString(), TextureSampleShader->GetUniqueID());
 					}
