Interchange: Fix fbx animation import
#rnx

diff --git a/Engine/Plugins/Interchange/Runtime/Source/Dispatcher/Private/InterchangeCommands.cpp b/Engine/Plugins/Interchange/Runtime/Source/Dispatcher/Private/InterchangeCommands.cpp
--- a/Engine/Plugins/Interchange/Runtime/Source/Dispatcher/Private/InterchangeCommands.cpp
+++ b/Engine/Plugins/Interchange/Runtime/Source/Dispatcher/Private/InterchangeCommands.cpp
@@ -12,7 +12,7 @@ namespace UE
 	{
 		int32 DispatcherCommandVersion::GetMajor()
 		{
-			return 19;
+			return 20;
 		}
 
 		int32 DispatcherCommandVersion::GetMinor()
diff --git a/Engine/Plugins/Interchange/Runtime/Source/Parsers/Fbx/Private/FbxAnimation.cpp b/Engine/Plugins/Interchange/Runtime/Source/Parsers/Fbx/Private/FbxAnimation.cpp
--- a/Engine/Plugins/Interchange/Runtime/Source/Parsers/Fbx/Private/FbxAnimation.cpp
+++ b/Engine/Plugins/Interchange/Runtime/Source/Parsers/Fbx/Private/FbxAnimation.cpp
@@ -355,10 +355,13 @@ namespace UE::Interchange::Private
 		DestinationCurve.StringKeyValues = StepCurveValues;
 	}
 
-	bool ImportBakeTransforms(FNodeTransformFetchPayloadData& FetchPayloadData, FAnimationPayloadData& AnimationBakeTransformPayloadData, FbxScene* SDKScene)
+	bool ImportBakeTransforms(FNodeTransformFetchPayloadData& FetchPayloadData, FAnimationPayloadData& AnimationBakeTransformPayloadData, FFbxParser& Parser)
 	{
 		if (!ensure(!FMath::IsNearlyZero(AnimationBakeTransformPayloadData.BakeFrequency)))
 		{
+			UInterchangeResultError_Generic* Message = Parser.AddMessage<UInterchangeResultError_Generic>();
+			Message->InterchangeKey = Parser.GetFbxHelper()->GetFbxNodeHierarchyName(FetchPayloadData.Node);
+			Message->Text = LOCTEXT("BakeFrequencyZero", "Cannot fetch FBX animation bake transforms payload because the bake frequency is zero.");
 			return false;
 		}
 		
@@ -368,6 +371,9 @@ namespace UE::Interchange::Private
 		EndTime.SetSecondDouble(AnimationBakeTransformPayloadData.RangeEndTime);
 		if (!ensure(AnimationBakeTransformPayloadData.RangeEndTime > AnimationBakeTransformPayloadData.RangeStartTime))
 		{
+			UInterchangeResultError_Generic* Message = Parser.AddMessage<UInterchangeResultError_Generic>();
+			Message->InterchangeKey = Parser.GetFbxHelper()->GetFbxNodeHierarchyName(FetchPayloadData.Node);
+			Message->Text = LOCTEXT("InvalidRange", "Cannot fetch FBX animation bake transforms payload because the bake range is invalid.");
 			return false;
 		}
 
@@ -382,19 +388,57 @@ namespace UE::Interchange::Private
 		const FbxTime TimeComparisonThreshold = (UE_DOUBLE_KINDA_SMALL_NUMBER * static_cast<double>(FBXSDK_TC_SECOND));
 		AnimationBakeTransformPayloadData.Transforms.Empty(NumFrame);
 
-		SDKScene->SetCurrentAnimationStack(FetchPayloadData.CurrentAnimStack);
+		Parser.GetSDKScene()->SetCurrentAnimationStack(FetchPayloadData.CurrentAnimStack);
+
+		bool bNanErrorLogged = false;
+		auto LogNanError = [&bNanErrorLogged, &FetchPayloadData, &Parser]()
+			{
+				if (bNanErrorLogged)
+				{
+					return;
+				}
+				UInterchangeResultError_Generic* Message = Parser.AddMessage<UInterchangeResultError_Generic>();
+				Message->InterchangeKey = Parser.GetFbxHelper()->GetFbxNodeHierarchyName(FetchPayloadData.Node);
+				Message->Text = LOCTEXT("BoneTransformNan", "Error when fetching FBX animation bake transforms payload, some transform contain NAN.");
+				bNanErrorLogged = true;
+			};
 
 		for (FbxTime CurTime = StartTime; CurTime < (EndTime + TimeComparisonThreshold); CurTime += TimeStep)
 		{
-			FbxAMatrix NodeTransform = FetchPayloadData.Node->EvaluateGlobalTransform(CurTime);
+			FTransform LocalTransform;
 			FbxNode* ParentNode = FetchPayloadData.Node->GetParent();
 			if (ParentNode)
 			{
+				FbxAMatrix NodeTransform = FetchPayloadData.Node->EvaluateGlobalTransform(CurTime);
+				FTransform GlobalTransform = UE::Interchange::Private::FFbxConvert::ConvertTransform<FTransform, FVector, FQuat>(NodeTransform);
+				if (GlobalTransform.ContainsNaN())
+				{
+					LogNanError();
+					continue;
+				}
+
 				FbxAMatrix ParentTransform = ParentNode->EvaluateGlobalTransform(CurTime);
-				NodeTransform = ParentTransform.Inverse() * NodeTransform;
-			}
+				if (GlobalTransform.ContainsNaN())
+				{
+					LogNanError();
+					continue;
+				}
+				FTransform ParentGlobalTransform = UE::Interchange::Private::FFbxConvert::ConvertTransform<FTransform, FVector, FQuat>(ParentTransform);
 
-			AnimationBakeTransformPayloadData.Transforms.Add(UE::Interchange::Private::FFbxConvert::ConvertTransform<FTransform, FVector, FQuat>(NodeTransform));
+				LocalTransform = GlobalTransform.GetRelativeTransform(ParentGlobalTransform);
+			}
+			else
+			{
+				FbxAMatrix& LocalMatrix = FetchPayloadData.Node->EvaluateLocalTransform(CurTime);
+				FbxVector4 NewLocalT = LocalMatrix.GetT();
+				FbxVector4 NewLocalS = LocalMatrix.GetS();
+				FbxQuaternion NewLocalQ = LocalMatrix.GetQ();
+
+				LocalTransform.SetTranslation(UE::Interchange::Private::FFbxConvert::ConvertPos<FVector>(NewLocalT));
+				LocalTransform.SetScale3D(UE::Interchange::Private::FFbxConvert::ConvertScale<FVector>(NewLocalS));
+				LocalTransform.SetRotation(UE::Interchange::Private::FFbxConvert::ConvertRotToQuat<FQuat>(NewLocalQ));
+			}
+			AnimationBakeTransformPayloadData.Transforms.Add(LocalTransform);
 		}
 
 		return true;
@@ -522,7 +566,7 @@ namespace UE::Interchange::Private
 		AnimationBakeTransformPayloadData.RangeStartTime = RangeStartTime;
 		AnimationBakeTransformPayloadData.RangeEndTime = RangeEndTime;
 
-		ImportBakeTransforms(FetchPayloadData, AnimationBakeTransformPayloadData, Parser.GetSDKScene());
+		ImportBakeTransforms(FetchPayloadData, AnimationBakeTransformPayloadData, Parser);
 		{
 			FLargeMemoryWriter Ar;
 			AnimationBakeTransformPayloadData.SerializeBaked(Ar);
