Interchange: Make sure skeleton conflict ui expand the skeleton conflict every time we open the dialog

#rnx

diff --git a/Engine/Plugins/Interchange/Runtime/Source/Pipelines/Private/InterchangeGenericAssetsPipeline.cpp b/Engine/Plugins/Interchange/Runtime/Source/Pipelines/Private/InterchangeGenericAssetsPipeline.cpp
--- a/Engine/Plugins/Interchange/Runtime/Source/Pipelines/Private/InterchangeGenericAssetsPipeline.cpp
+++ b/Engine/Plugins/Interchange/Runtime/Source/Pipelines/Private/InterchangeGenericAssetsPipeline.cpp
@@ -556,14 +556,14 @@ void UInterchangeGenericAssetsPipeline::InternalRecursiveFillJointsFromNodeConta
 namespace UE::Interchange::Private
 {
 	void SetParentChildConflict(TSharedPtr<UInterchangeGenericAssetsPipeline::FSkeletonJoint> ParentJoint)
+	{
+		TSharedPtr<UInterchangeGenericAssetsPipeline::FSkeletonJoint> ParentJointIter = ParentJoint;
+		while (ParentJointIter.IsValid() && !ParentJointIter->bChildConflict)
 		{
-			TSharedPtr<UInterchangeGenericAssetsPipeline::FSkeletonJoint> ParentJointIter = ParentJoint;
-			while (ParentJointIter.IsValid() && !ParentJointIter->bChildConflict)
-			{
-				ParentJointIter->bChildConflict = true;
-				ParentJointIter = ParentJointIter->Parent;
-			}
-		};
+			ParentJointIter->bChildConflict = true;
+			ParentJointIter = ParentJointIter->Parent;
+		}
+	}
 
 	void RecursivelyFillJointRemoved(TSharedPtr<UInterchangeGenericAssetsPipeline::FSkeletonJoint> AssetJoint
 		, TSharedPtr<UInterchangeGenericAssetsPipeline::FSkeletonJoint> ParentJoint
@@ -1772,32 +1772,25 @@ void SInterchangeGenericAssetSkeletonConflictWidget::Construct(const FArguments&
 	];
 }
 
-FReply SInterchangeGenericAssetSkeletonConflictWidget::SetSectionVisible(EInterchangeSkeletonCompareSection SectionIndex)
-{
-	bShowSectionFlag[SectionIndex] = !bShowSectionFlag[SectionIndex];
-	return FReply::Handled();
-}
-
-EVisibility SInterchangeGenericAssetSkeletonConflictWidget::IsSectionVisible(EInterchangeSkeletonCompareSection SectionIndex)
-{
-	return bShowSectionFlag[SectionIndex] ? EVisibility::Visible : EVisibility::Collapsed;
-}
-
-const FSlateBrush* SInterchangeGenericAssetSkeletonConflictWidget::GetCollapsableArrow(EInterchangeSkeletonCompareSection SectionIndex) const
-{
-	return bShowSectionFlag[SectionIndex] ? FAppStyle::GetBrush("Symbols.DownArrow") : FAppStyle::GetBrush("Symbols.RightArrow");
-}
-
 namespace UE::Interchange::Private
 {
+	void RecursivelyCollapseTreeItem(TSharedPtr<STreeView<TSharedPtr<UInterchangeGenericAssetsPipeline::FSkeletonJoint>>> CompareTree
+		, TSharedPtr<UInterchangeGenericAssetsPipeline::FSkeletonJoint> JointItem)
+	{
+		CompareTree->SetItemExpansion(JointItem, false);
+		for (TSharedPtr<UInterchangeGenericAssetsPipeline::FSkeletonJoint> ChildJoint : JointItem->Children)
+		{
+			RecursivelyCollapseTreeItem(CompareTree, ChildJoint);
+		}
+	}
+
 	void RecursivelyExpandTreeItem(TSharedPtr<STreeView<TSharedPtr<UInterchangeGenericAssetsPipeline::FSkeletonJoint>>> CompareTree
 		, TSharedPtr<UInterchangeGenericAssetsPipeline::FSkeletonJoint> JointItem)
 	{
-		if (JointItem->bInitialAutoExpand || !JointItem->bMatch || !JointItem->bChildConflict)
+		if (!JointItem->bMatch || !JointItem->bChildConflict)
 		{
 			return;
 		}
-		JointItem->bInitialAutoExpand = true;
 		CompareTree->SetItemExpansion(JointItem, true);
 		for (TSharedPtr<UInterchangeGenericAssetsPipeline::FSkeletonJoint> ChildJoint : JointItem->Children)
 		{
@@ -1806,6 +1799,36 @@ namespace UE::Interchange::Private
 	}
 }
 
+FReply SInterchangeGenericAssetSkeletonConflictWidget::OnExpandToConflict()
+{
+	for (TSharedPtr<UInterchangeGenericAssetsPipeline::FSkeletonJoint> Joint : Joints)
+	{
+		UE::Interchange::Private::RecursivelyCollapseTreeItem(CompareTree, Joint);
+	}
+
+	for (TSharedPtr<UInterchangeGenericAssetsPipeline::FSkeletonJoint> Joint : Joints)
+	{
+		UE::Interchange::Private::RecursivelyExpandTreeItem(CompareTree, Joint);
+	}
+	return FReply::Handled();
+}
+
+FReply SInterchangeGenericAssetSkeletonConflictWidget::SetSectionVisible(EInterchangeSkeletonCompareSection SectionIndex)
+{
+	bShowSectionFlag[SectionIndex] = !bShowSectionFlag[SectionIndex];
+	return FReply::Handled();
+}
+
+EVisibility SInterchangeGenericAssetSkeletonConflictWidget::IsSectionVisible(EInterchangeSkeletonCompareSection SectionIndex)
+{
+	return bShowSectionFlag[SectionIndex] ? EVisibility::Visible : EVisibility::Collapsed;
+}
+
+const FSlateBrush* SInterchangeGenericAssetSkeletonConflictWidget::GetCollapsableArrow(EInterchangeSkeletonCompareSection SectionIndex) const
+{
+	return bShowSectionFlag[SectionIndex] ? FAppStyle::GetBrush("Symbols.DownArrow") : FAppStyle::GetBrush("Symbols.RightArrow");
+}
+
 TSharedPtr<SWidget> SInterchangeGenericAssetSkeletonConflictWidget::ConstructSkeletonComparison()
 {
 	FText SkeletonStatus = NSLOCTEXT("SInterchangeGenericAssetSkeletonConflictWidget", "ConstructSkeletonComparison_SkeletonStatus", "The skeleton has some conflicts");
@@ -1872,10 +1895,26 @@ TSharedPtr<SWidget> SInterchangeGenericAssetSkeletonConflictWidget::ConstructSke
 						.AutoHeight()
 						.Padding(2)
 						[
-							SNew(STextBlock)
-							.Font(FAppStyle::GetFontStyle("DetailsView.CategoryFontStyle"))
-							.Text(SkeletonStatus)
-							.ColorAndOpacity(SInterchangeGenericAssetMaterialConflictWidget::SlateColorFullConflict)
+							SNew(SHorizontalBox)
+							+SHorizontalBox::Slot()
+							.AutoWidth()
+							[
+								SNew(SButton)
+								.HAlign(HAlign_Center)
+								.Text(NSLOCTEXT("SInterchangeGenericAssetMaterialConflictWidget", "SInterchangeGenericAssetMaterialConflic_ExpandToConflict", "Expand To Conflict"))
+								.OnClicked(this, &SInterchangeGenericAssetSkeletonConflictWidget::OnExpandToConflict)
+							]
+							+SHorizontalBox::Slot()
+							.FillWidth(1.0f)
+							.VAlign(VAlign_Center)
+							.HAlign(HAlign_Left)
+							.Padding(6.0f, 0.0f)
+							[
+								SNew(STextBlock)
+								.Font(FAppStyle::GetFontStyle("DetailsView.CategoryFontStyle"))
+								.Text(SkeletonStatus)
+								.ColorAndOpacity(SInterchangeGenericAssetMaterialConflictWidget::SlateColorFullConflict)
+							]
 						]
 						+SVerticalBox::Slot()
 						.AutoHeight()
diff --git a/Engine/Plugins/Interchange/Runtime/Source/Pipelines/Public/InterchangeGenericAssetsPipeline.h b/Engine/Plugins/Interchange/Runtime/Source/Pipelines/Public/InterchangeGenericAssetsPipeline.h
--- a/Engine/Plugins/Interchange/Runtime/Source/Pipelines/Public/InterchangeGenericAssetsPipeline.h
+++ b/Engine/Plugins/Interchange/Runtime/Source/Pipelines/Public/InterchangeGenericAssetsPipeline.h
@@ -126,7 +126,6 @@ public:
 		bool bMatch = false;
 		bool bConflict = false;
 		bool bChildConflict = false;
-		bool bInitialAutoExpand = false;
 		TSharedPtr<FSkeletonJoint> Parent;
 		TArray<TSharedPtr<FSkeletonJoint>> Children;
 	};
@@ -314,6 +313,7 @@ private:
 	//////////////////////////////////////////////////////////////////////////
 	//Collapse generic
 	bool bShowSectionFlag[EInterchangeSkeletonCompareSection::Count];
+	FReply OnExpandToConflict();
 	FReply SetSectionVisible(EInterchangeSkeletonCompareSection SectionIndex);
 	EVisibility IsSectionVisible(EInterchangeSkeletonCompareSection SectionIndex);
 	const FSlateBrush* GetCollapsableArrow(EInterchangeSkeletonCompareSection SectionIndex) const;
