Interchange Leve Sequence:
- Add Lights (Point, Spot,Rect, Directional, Sky) and Cameras (CineCamera and RailSplinComponents) property tracks (for the moment only Integer/Bool/Byte/Float/Double tracks are supported)

diff --git a/Engine/Plugins/Interchange/Runtime/Source/Import/Private/Animation/InterchangeLevelSequenceHelper.cpp b/Engine/Plugins/Interchange/Runtime/Source/Import/Private/Animation/InterchangeLevelSequenceHelper.cpp
--- a/Engine/Plugins/Interchange/Runtime/Source/Import/Private/Animation/InterchangeLevelSequenceHelper.cpp
+++ b/Engine/Plugins/Interchange/Runtime/Source/Import/Private/Animation/InterchangeLevelSequenceHelper.cpp
@@ -5,14 +5,22 @@
 
 #include "Camera/CameraComponent.h"
 #include "CineCameraComponent.h"
+#include "Components/DirectionalLightComponent.h"
 #include "Components/LocalLightComponent.h"
+#include "Components/PointLightComponent.h"
+#include "Components/RectLightComponent.h"
+#include "Components/SkyLightComponent.h"
+#include "Components/SpotLightComponent.h"
+#include "Components/SplineComponent.h"
 
 #include "MovieScene.h"
 #include "MovieSceneSection.h"
 #include "Tracks/MovieSceneBoolTrack.h"
 #include "Tracks/MovieSceneByteTrack.h"
 #include "Tracks/MovieSceneColorTrack.h"
+#include "Tracks/MovieSceneEnumTrack.h"
 #include "Tracks/MovieSceneFloatTrack.h"
+#include "Tracks/MovieSceneIntegerTrack.h"
 #include "Tracks/MovieSceneVisibilityTrack.h"
 
 namespace UE::Interchange::Private
@@ -21,34 +29,279 @@ namespace UE::Interchange::Private
 		: PropertyTracks{
 			//Common
 			{EInterchangePropertyTracks::Visibility, {UMovieSceneVisibilityTrack::StaticClass()->GetName(), AActor::GetHiddenPropertyName().ToString(), TEXT("Visibility")}},
+
 			{EInterchangePropertyTracks::AutoActivate, {UMovieSceneBoolTrack::StaticClass()->GetName(), GET_MEMBER_NAME_STRING_CHECKED(UActorComponent, bAutoActivate), TEXT("Auto Activate")}},
+			{EInterchangePropertyTracks::AffectDistanceFieldLighting, {UMovieSceneBoolTrack::StaticClass()->GetName(),  GET_MEMBER_NAME_STRING_CHECKED(UPrimitiveComponent, bAffectDistanceFieldLighting) , TEXT("Affect Distance Field Lighting")}},
+			{EInterchangePropertyTracks::AffectDynamicIndirectLighting, {UMovieSceneBoolTrack::StaticClass()->GetName(),  GET_MEMBER_NAME_STRING_CHECKED(UPrimitiveComponent, bAffectDynamicIndirectLighting) , TEXT("Affect Dynamic Indirect Lighting")}},
+			{EInterchangePropertyTracks::AffectIndirectLightingWhileHidden, {UMovieSceneBoolTrack::StaticClass()->GetName(),  GET_MEMBER_NAME_STRING_CHECKED(UPrimitiveComponent, bAffectIndirectLightingWhileHidden) , TEXT("Affect Indirect Lighting While Hidden")}},
+			{EInterchangePropertyTracks::BodyInstanceAngularDamping, {UMovieSceneFloatTrack::StaticClass()->GetName(),  GET_MEMBER_NAME_STRING_CHECKED(UPrimitiveComponent, BodyInstance) + FString{ TEXT(".") } + GET_MEMBER_NAME_STRING_CHECKED(FBodyInstance, AngularDamping), TEXT("Angular Damping (Body Instance)")}},
+			{EInterchangePropertyTracks::BodyInstancebEnableGravity, {UMovieSceneByteTrack::StaticClass()->GetName(),  GET_MEMBER_NAME_STRING_CHECKED(UPrimitiveComponent, BodyInstance) + FString{ TEXT(".") } + GET_MEMBER_NAME_STRING_CHECKED(FBodyInstance, bEnableGravity), TEXT("Enable Gravity (Body Instance)")}},
+			{EInterchangePropertyTracks::BodyInstancebNotifyRigidBodyCollision, {UMovieSceneByteTrack::StaticClass()->GetName(),  GET_MEMBER_NAME_STRING_CHECKED(UPrimitiveComponent, BodyInstance) + FString{ TEXT(".") } + GET_MEMBER_NAME_STRING_CHECKED(FBodyInstance, bNotifyRigidBodyCollision), TEXT("Simulation Generates Hit Events (Body Instance)")}},
+			{EInterchangePropertyTracks::BodyInstancebSimulatePhysics, {UMovieSceneByteTrack::StaticClass()->GetName(),  GET_MEMBER_NAME_STRING_CHECKED(UPrimitiveComponent, BodyInstance) + FString{ TEXT(".") } + GET_MEMBER_NAME_STRING_CHECKED(FBodyInstance, bSimulatePhysics), TEXT("Simulate Physics (Body Instance)")}},
+			{EInterchangePropertyTracks::BodyInstancebUpdateKinematicFromSimulation, {UMovieSceneByteTrack::StaticClass()->GetName(),  GET_MEMBER_NAME_STRING_CHECKED(UPrimitiveComponent, BodyInstance) + FString{ TEXT(".") } + GET_MEMBER_NAME_STRING_CHECKED(FBodyInstance, bUpdateKinematicFromSimulation), TEXT("Update Kinematic from Simulation (Body Instance)")}},
+			{EInterchangePropertyTracks::BodyInstancebUseCCD, {UMovieSceneByteTrack::StaticClass()->GetName(),  GET_MEMBER_NAME_STRING_CHECKED(UPrimitiveComponent, BodyInstance) + FString{ TEXT(".") } + GET_MEMBER_NAME_STRING_CHECKED(FBodyInstance, bUseCCD), TEXT("Use CCD (Body Instance)")}},
+			{EInterchangePropertyTracks::BodyInstanceLinearDamping, {UMovieSceneFloatTrack::StaticClass()->GetName(),  GET_MEMBER_NAME_STRING_CHECKED(UPrimitiveComponent, BodyInstance) + FString{ TEXT(".") } + GET_MEMBER_NAME_STRING_CHECKED(FBodyInstance, LinearDamping), TEXT("Linear Damping (Body Instance)")}},
+			{EInterchangePropertyTracks::BodyInstanceMassScale, {UMovieSceneFloatTrack::StaticClass()->GetName(),  GET_MEMBER_NAME_STRING_CHECKED(UPrimitiveComponent, BodyInstance) + FString{ TEXT(".") } + GET_MEMBER_NAME_STRING_CHECKED(FBodyInstance, MassScale), TEXT("Mass Scale (Body Instance)")}},
+			{EInterchangePropertyTracks::BoundsScale, {UMovieSceneFloatTrack::StaticClass()->GetName(),  GET_MEMBER_NAME_STRING_CHECKED(UPrimitiveComponent, BoundsScale) , TEXT("Bounds Scale")}},
+			{EInterchangePropertyTracks::CastContactShadow, {UMovieSceneBoolTrack::StaticClass()->GetName(),  GET_MEMBER_NAME_STRING_CHECKED(UPrimitiveComponent, bCastContactShadow) , TEXT("Contact Shadow")}},
+			{EInterchangePropertyTracks::CastHiddenShadow, {UMovieSceneBoolTrack::StaticClass()->GetName(),  GET_MEMBER_NAME_STRING_CHECKED(UPrimitiveComponent, bCastHiddenShadow) , TEXT("Hidden Shadow")}},
+			{EInterchangePropertyTracks::CastInsetShadow, {UMovieSceneBoolTrack::StaticClass()->GetName(),  GET_MEMBER_NAME_STRING_CHECKED(UPrimitiveComponent, bCastInsetShadow) , TEXT("Dynamic Inset Shadow")}},
+			{EInterchangePropertyTracks::CastShadow, {UMovieSceneByteTrack::StaticClass()->GetName(),  GET_MEMBER_NAME_STRING_CHECKED(UPrimitiveComponent, CastShadow) , TEXT("Cast Shadow")}},
+			{EInterchangePropertyTracks::CustomDepthStencilValue, {UMovieSceneIntegerTrack::StaticClass()->GetName(),  GET_MEMBER_NAME_STRING_CHECKED(UPrimitiveComponent, CustomDepthStencilValue) , TEXT("CustomDepth Stencil Value")}},
+			{EInterchangePropertyTracks::CustomDepthStencilWriteMask, {UMovieSceneEnumTrack::StaticClass()->GetName(),  GET_MEMBER_NAME_STRING_CHECKED(UPrimitiveComponent, CustomDepthStencilWriteMask) , TEXT("Custom Depth Stencil Write Mask"), StaticEnum<ERendererStencilMask>()}},
+			{EInterchangePropertyTracks::DrawDebug, {UMovieSceneBoolTrack::StaticClass()->GetName(),  GET_MEMBER_NAME_STRING_CHECKED(USplineComponent, bDrawDebug) , TEXT("Draw Debug")}},
+			{EInterchangePropertyTracks::EmissiveLightSource, {UMovieSceneBoolTrack::StaticClass()->GetName(),  GET_MEMBER_NAME_STRING_CHECKED(UPrimitiveComponent, bEmissiveLightSource) , TEXT("Emissive Light Source")}},
+			{EInterchangePropertyTracks::ExcludeFromLightAttachmentGroup, {UMovieSceneBoolTrack::StaticClass()->GetName(),  GET_MEMBER_NAME_STRING_CHECKED(UPrimitiveComponent, bExcludeFromLightAttachmentGroup) , TEXT("Exclude from Light Attachment Group")}},
+			{EInterchangePropertyTracks::HiddenInSceneCapture, {UMovieSceneBoolTrack::StaticClass()->GetName(),  GET_MEMBER_NAME_STRING_CHECKED(UPrimitiveComponent, bHiddenInSceneCapture) , TEXT("Hidden In Scene Capture")}},
+			{EInterchangePropertyTracks::Holdout, {UMovieSceneBoolTrack::StaticClass()->GetName(),  GET_MEMBER_NAME_STRING_CHECKED(UPrimitiveComponent, bHoldout) , TEXT("Holdout")}},
+			{EInterchangePropertyTracks::LightAttachmentsAsGroup, {UMovieSceneBoolTrack::StaticClass()->GetName(),  GET_MEMBER_NAME_STRING_CHECKED(UPrimitiveComponent, bLightAttachmentsAsGroup) , TEXT("Light Attachments as Group")}},
+			{EInterchangePropertyTracks::Mobility, {UMovieSceneEnumTrack::StaticClass()->GetName(),  GET_MEMBER_NAME_STRING_CHECKED(USceneComponent, Mobility) , TEXT("Mobility"), StaticEnum<EComponentMobility::Type>()}},
+			{EInterchangePropertyTracks::OnlyOwnerSee, {UMovieSceneBoolTrack::StaticClass()->GetName(),  GET_MEMBER_NAME_STRING_CHECKED(UPrimitiveComponent, bOnlyOwnerSee) , TEXT("Only Owner See")}},
+			{EInterchangePropertyTracks::OwnerNoSee, {UMovieSceneBoolTrack::StaticClass()->GetName(),  GET_MEMBER_NAME_STRING_CHECKED(UPrimitiveComponent, bOwnerNoSee) , TEXT("Owner No See")}},
+			{EInterchangePropertyTracks::ReceivesDecals, {UMovieSceneBoolTrack::StaticClass()->GetName(),  GET_MEMBER_NAME_STRING_CHECKED(UPrimitiveComponent, bReceivesDecals) , TEXT("Receives Decals")}},
+			{EInterchangePropertyTracks::RenderCustomDepth, {UMovieSceneBoolTrack::StaticClass()->GetName(),  GET_MEMBER_NAME_STRING_CHECKED(UPrimitiveComponent, bRenderCustomDepth) , TEXT("Render CustomDepth Pass")}},
+			{EInterchangePropertyTracks::RenderInDepthPass, {UMovieSceneBoolTrack::StaticClass()->GetName(),  GET_MEMBER_NAME_STRING_CHECKED(UPrimitiveComponent, bRenderInDepthPass) , TEXT("Render in Depth Pass")}},
+			{EInterchangePropertyTracks::RenderInMainPass, {UMovieSceneBoolTrack::StaticClass()->GetName(),  GET_MEMBER_NAME_STRING_CHECKED(UPrimitiveComponent, bRenderInMainPass) , TEXT("Render in Main Pass")}},
+			{EInterchangePropertyTracks::SingleSampleShadowFromStationaryLights, {UMovieSceneBoolTrack::StaticClass()->GetName(),  GET_MEMBER_NAME_STRING_CHECKED(UPrimitiveComponent, bSingleSampleShadowFromStationaryLights) , TEXT("Single Sample Shadow from Stationary Lights")}},
+			{EInterchangePropertyTracks::TranslucencySortDistanceOffset, {UMovieSceneFloatTrack::StaticClass()->GetName(),  GET_MEMBER_NAME_STRING_CHECKED(UPrimitiveComponent, TranslucencySortDistanceOffset) , TEXT("Translucency Sort Distance Offset")}},
+			{EInterchangePropertyTracks::VisibleInRayTracing, {UMovieSceneBoolTrack::StaticClass()->GetName(),  GET_MEMBER_NAME_STRING_CHECKED(UPrimitiveComponent, bVisibleInRayTracing) , TEXT("Visible in Ray Tracing")}},
+			{EInterchangePropertyTracks::VisibleInSceneCaptureOnly, {UMovieSceneBoolTrack::StaticClass()->GetName(),  GET_MEMBER_NAME_STRING_CHECKED(UPrimitiveComponent, bVisibleInSceneCaptureOnly) , TEXT("Visible In Scene Capture Only")}},
 
 			// Light
-			{EInterchangePropertyTracks::LightColor, {UMovieSceneColorTrack::StaticClass()->GetName(), GET_MEMBER_NAME_STRING_CHECKED(ULocalLightComponent, LightColor), TEXT("Light Color")}},
-			{EInterchangePropertyTracks::LightIntensity, {UMovieSceneFloatTrack::StaticClass()->GetName(), GET_MEMBER_NAME_STRING_CHECKED(ULightComponent, Intensity), TEXT("Intensity")}},
-			{EInterchangePropertyTracks::LightIntensityUnits, {UMovieSceneByteTrack::StaticClass()->GetName(), GET_MEMBER_NAME_STRING_CHECKED(ULocalLightComponent, IntensityUnits), TEXT("Intensity Units"), StaticEnum<ELightUnits>()}},
-			{EInterchangePropertyTracks::LightTemperature, {UMovieSceneFloatTrack::StaticClass()->GetName(), GET_MEMBER_NAME_STRING_CHECKED(ULocalLightComponent, Temperature), TEXT("Temperature")}},
-			{EInterchangePropertyTracks::LightUseTemperature, {UMovieSceneBoolTrack::StaticClass()->GetName(), GET_MEMBER_NAME_STRING_CHECKED(ULocalLightComponent, bUseTemperature), TEXT("Use Temperature")}},
+			{EInterchangePropertyTracks::LightAffectGlobalIllumination, {UMovieSceneBoolTrack::StaticClass()->GetName(),  GET_MEMBER_NAME_STRING_CHECKED(ULightComponentBase, bAffectGlobalIllumination) , TEXT("Affect Ray Tracing Global Illumination")}},
+			{EInterchangePropertyTracks::LightAffectReflection, {UMovieSceneBoolTrack::StaticClass()->GetName(),  GET_MEMBER_NAME_STRING_CHECKED(ULightComponentBase, bAffectReflection) , TEXT("Affect Ray Tracing Reflections")}},
+			{EInterchangePropertyTracks::LightAffectTranslucentLighting, {UMovieSceneBoolTrack::StaticClass()->GetName(),  GET_MEMBER_NAME_STRING_CHECKED(ULightComponentBase, bAffectTranslucentLighting) , TEXT("Affect Translucent Lighting")}},
+			{EInterchangePropertyTracks::LightAtmosphereSunDiskColorScale, {UMovieSceneColorTrack::StaticClass()->GetName(),  GET_MEMBER_NAME_STRING_CHECKED(UDirectionalLightComponent, AtmosphereSunDiskColorScale) , TEXT("Atmosphere Sun Disk Color Scale")}},
+			{EInterchangePropertyTracks::LightAtmosphereSunLight, {UMovieSceneBoolTrack::StaticClass()->GetName(),  GET_MEMBER_NAME_STRING_CHECKED(UDirectionalLightComponent, bAtmosphereSunLight) , TEXT("Atmosphere Sun Light")}},
+			{EInterchangePropertyTracks::LightAtmosphereSunLightIndex, {UMovieSceneIntegerTrack::StaticClass()->GetName(),  GET_MEMBER_NAME_STRING_CHECKED(UDirectionalLightComponent, AtmosphereSunLightIndex) , TEXT("Atmosphere Sun Light Index")}},
+			{EInterchangePropertyTracks::LightAttenuationRadius, {UMovieSceneFloatTrack::StaticClass()->GetName(),  GET_MEMBER_NAME_STRING_CHECKED(ULocalLightComponent, AttenuationRadius) , TEXT("Attenuation Radius")}},
+			{EInterchangePropertyTracks::LightBarnDoorAngle, {UMovieSceneFloatTrack::StaticClass()->GetName(),  GET_MEMBER_NAME_STRING_CHECKED(URectLightComponent, BarnDoorAngle) , TEXT("Barn Door Angle")}},
+			{EInterchangePropertyTracks::LightBarnDoorLength, {UMovieSceneFloatTrack::StaticClass()->GetName(),  GET_MEMBER_NAME_STRING_CHECKED(URectLightComponent, BarnDoorLength) , TEXT("Barn Door Length")}},
+			{EInterchangePropertyTracks::LightBloomMaxBrightness, {UMovieSceneFloatTrack::StaticClass()->GetName(),  GET_MEMBER_NAME_STRING_CHECKED(ULightComponent, BloomMaxBrightness) , TEXT("Bloom Max Brightness")}},
+			{EInterchangePropertyTracks::LightBloomScale, {UMovieSceneFloatTrack::StaticClass()->GetName(),  GET_MEMBER_NAME_STRING_CHECKED(ULightComponent, BloomScale) , TEXT("Bloom Scale")}},
+			{EInterchangePropertyTracks::LightBloomThreshold, {UMovieSceneFloatTrack::StaticClass()->GetName(),  GET_MEMBER_NAME_STRING_CHECKED(ULightComponent, BloomThreshold) , TEXT("Bloom Threshold")}},
+			{EInterchangePropertyTracks::LightBloomTint, {UMovieSceneColorTrack::StaticClass()->GetName(),  GET_MEMBER_NAME_STRING_CHECKED(ULightComponent, BloomTint) , TEXT("Bloom Tint")}},
+			{EInterchangePropertyTracks::LightCascadeDistributionExponent, {UMovieSceneFloatTrack::StaticClass()->GetName(),  GET_MEMBER_NAME_STRING_CHECKED(UDirectionalLightComponent, CascadeDistributionExponent) , TEXT("Distribution Exponent")}},
+			{EInterchangePropertyTracks::LightCascadeTransitionFraction, {UMovieSceneFloatTrack::StaticClass()->GetName(),  GET_MEMBER_NAME_STRING_CHECKED(UDirectionalLightComponent, CascadeTransitionFraction) , TEXT("Transition Fraction")}},
+			{EInterchangePropertyTracks::LightCastDeepShadow, {UMovieSceneBoolTrack::StaticClass()->GetName(),  GET_MEMBER_NAME_STRING_CHECKED(ULightComponentBase, bCastDeepShadow) , TEXT("Cast Deep Shadow")}},
+			{EInterchangePropertyTracks::LightCastShadows, {UMovieSceneByteTrack::StaticClass()->GetName(),  GET_MEMBER_NAME_STRING_CHECKED(ULightComponentBase, CastShadows) , TEXT("Cast Shadows")}},
+			{EInterchangePropertyTracks::LightCastVolumetricShadow, {UMovieSceneBoolTrack::StaticClass()->GetName(),  GET_MEMBER_NAME_STRING_CHECKED(ULightComponentBase, bCastVolumetricShadow) , TEXT("Cast Volumetric Shadow")}},
+			{EInterchangePropertyTracks::LightCloudAmbientOcclusionStrength, {UMovieSceneFloatTrack::StaticClass()->GetName(),  GET_MEMBER_NAME_STRING_CHECKED(USkyLightComponent, CloudAmbientOcclusionStrength) , TEXT("Cloud Ambient Occlusion Strength")}},
+			{EInterchangePropertyTracks::LightCloudScatteredLuminanceScale, {UMovieSceneColorTrack::StaticClass()->GetName(),  GET_MEMBER_NAME_STRING_CHECKED(UDirectionalLightComponent, CloudScatteredLuminanceScale) , TEXT("Cloud Scattered Luminance Scale")}},
+			{EInterchangePropertyTracks::LightCloudShadowOnAtmosphereStrength, {UMovieSceneFloatTrack::StaticClass()->GetName(),  GET_MEMBER_NAME_STRING_CHECKED(UDirectionalLightComponent, CloudShadowOnAtmosphereStrength) , TEXT("Cloud Shadow on Atmosphere Strength")}},
+			{EInterchangePropertyTracks::LightCloudShadowOnSurfaceStrength, {UMovieSceneFloatTrack::StaticClass()->GetName(),  GET_MEMBER_NAME_STRING_CHECKED(UDirectionalLightComponent, CloudShadowOnSurfaceStrength) , TEXT("Cloud Shadow on Surface Strength")}},
+			{EInterchangePropertyTracks::LightCloudShadowStrength, {UMovieSceneFloatTrack::StaticClass()->GetName(),  GET_MEMBER_NAME_STRING_CHECKED(UDirectionalLightComponent, CloudShadowStrength) , TEXT("Cloud Shadow Strength")}},
+			{EInterchangePropertyTracks::LightColor, {UMovieSceneColorTrack::StaticClass()->GetName(),  GET_MEMBER_NAME_STRING_CHECKED(ULightComponentBase, LightColor) , TEXT("Light Color")}},
+			{EInterchangePropertyTracks::LightDynamicShadowCascades, {UMovieSceneIntegerTrack::StaticClass()->GetName(),  GET_MEMBER_NAME_STRING_CHECKED(UDirectionalLightComponent, DynamicShadowCascades) , TEXT("Num Dynamic Shadow Cascades")}},
+			{EInterchangePropertyTracks::LightDynamicShadowDistanceMovableLight, {UMovieSceneFloatTrack::StaticClass()->GetName(),  GET_MEMBER_NAME_STRING_CHECKED(UDirectionalLightComponent, DynamicShadowDistanceMovableLight) , TEXT("Dynamic Shadow Distance MovableLight")}},
+			{EInterchangePropertyTracks::LightDynamicShadowDistanceStationaryLight, {UMovieSceneFloatTrack::StaticClass()->GetName(),  GET_MEMBER_NAME_STRING_CHECKED(UDirectionalLightComponent, DynamicShadowDistanceStationaryLight) , TEXT("Dynamic Shadow Distance StationaryLight")}},
+			{EInterchangePropertyTracks::LightEnableLightShaftBloom, {UMovieSceneBoolTrack::StaticClass()->GetName(),  GET_MEMBER_NAME_STRING_CHECKED(ULightComponent, bEnableLightShaftBloom) , TEXT("Light Shaft Bloom")}},
+			{EInterchangePropertyTracks::LightEnableLightShaftOcclusion, {UMovieSceneBoolTrack::StaticClass()->GetName(),  GET_MEMBER_NAME_STRING_CHECKED(UDirectionalLightComponent, bEnableLightShaftOcclusion) , TEXT("Light Shaft Occlusion")}},
+			{EInterchangePropertyTracks::LightFalloffExponent, {UMovieSceneFloatTrack::StaticClass()->GetName(),  GET_MEMBER_NAME_STRING_CHECKED(UPointLightComponent, LightFalloffExponent) , TEXT("Light Falloff Exponent")}},
+			{EInterchangePropertyTracks::LightForceCachedShadowsForMovablePrimitives, {UMovieSceneBoolTrack::StaticClass()->GetName(),  GET_MEMBER_NAME_STRING_CHECKED(ULightComponent, bForceCachedShadowsForMovablePrimitives) , TEXT("Force Cached Shadows for Movable Primitives")}},
+			{EInterchangePropertyTracks::LightForwardShadingPriority, {UMovieSceneIntegerTrack::StaticClass()->GetName(),  GET_MEMBER_NAME_STRING_CHECKED(UDirectionalLightComponent, ForwardShadingPriority) , TEXT("Forward Shading Priority")}},
+			{EInterchangePropertyTracks::LightFunctionFadeDistance, {UMovieSceneFloatTrack::StaticClass()->GetName(),  GET_MEMBER_NAME_STRING_CHECKED(ULightComponent, LightFunctionFadeDistance) , TEXT("Fade Distance")}},
+			{EInterchangePropertyTracks::LightIESBrightnessScale, {UMovieSceneFloatTrack::StaticClass()->GetName(),  GET_MEMBER_NAME_STRING_CHECKED(ULightComponent, IESBrightnessScale) , TEXT("IES Intensity Scale")}},
+			{EInterchangePropertyTracks::LightIndirectLightingIntensity, {UMovieSceneFloatTrack::StaticClass()->GetName(),  GET_MEMBER_NAME_STRING_CHECKED(ULightComponentBase, IndirectLightingIntensity) , TEXT("Indirect Lighting Intensity")}},
+			{EInterchangePropertyTracks::LightInnerConeAngle, {UMovieSceneFloatTrack::StaticClass()->GetName(),  GET_MEMBER_NAME_STRING_CHECKED(USpotLightComponent, InnerConeAngle) , TEXT("Inner Cone Angle")}},
+			{EInterchangePropertyTracks::LightIntensity, {UMovieSceneFloatTrack::StaticClass()->GetName(),  GET_MEMBER_NAME_STRING_CHECKED(ULightComponentBase, Intensity) , TEXT("Intensity")}},
+			{EInterchangePropertyTracks::LightIntensityUnits, {UMovieSceneEnumTrack::StaticClass()->GetName(),  GET_MEMBER_NAME_STRING_CHECKED(ULocalLightComponent, IntensityUnits) , TEXT("Intensity Units"), StaticEnum<ELightUnits>()}},
+			{EInterchangePropertyTracks::LightInverseExposureBlend, {UMovieSceneFloatTrack::StaticClass()->GetName(),  GET_MEMBER_NAME_STRING_CHECKED(ULocalLightComponent, InverseExposureBlend) , TEXT("Inverse Exposure Blend")}},
+			{EInterchangePropertyTracks::LightLowerHemisphereColor, {UMovieSceneColorTrack::StaticClass()->GetName(),  GET_MEMBER_NAME_STRING_CHECKED(USkyLightComponent, LowerHemisphereColor) , TEXT("Lower Hemisphere Color")}},
+			{EInterchangePropertyTracks::LightmassSettingsLightSourceAngle, {UMovieSceneFloatTrack::StaticClass()->GetName(),  GET_MEMBER_NAME_STRING_CHECKED(UDirectionalLightComponent, LightmassSettings) + FString{ TEXT(".") } + GET_MEMBER_NAME_STRING_CHECKED(FLightmassDirectionalLightSettings, LightSourceAngle), TEXT("Light Source Angle (Lightmass Settings)")}},
+			{EInterchangePropertyTracks::LightMinOcclusion, {UMovieSceneFloatTrack::StaticClass()->GetName(),  GET_MEMBER_NAME_STRING_CHECKED(USkyLightComponent, MinOcclusion) , TEXT("Min Occlusion")}},
+			{EInterchangePropertyTracks::LightModulatedShadowColor, {UMovieSceneColorTrack::StaticClass()->GetName(),  GET_MEMBER_NAME_STRING_CHECKED(UDirectionalLightComponent, ModulatedShadowColor) , TEXT("Modulated Shadow Color")}},
+			{EInterchangePropertyTracks::LightOcclusionDepthRange, {UMovieSceneFloatTrack::StaticClass()->GetName(),  GET_MEMBER_NAME_STRING_CHECKED(UDirectionalLightComponent, OcclusionDepthRange) , TEXT("Occlusion Depth Range")}},
+			{EInterchangePropertyTracks::LightOcclusionExponent, {UMovieSceneFloatTrack::StaticClass()->GetName(),  GET_MEMBER_NAME_STRING_CHECKED(USkyLightComponent, OcclusionExponent) , TEXT("Occlusion Exponent")}},
+			{EInterchangePropertyTracks::LightOcclusionMaskDarkness, {UMovieSceneFloatTrack::StaticClass()->GetName(),  GET_MEMBER_NAME_STRING_CHECKED(UDirectionalLightComponent, OcclusionMaskDarkness) , TEXT("Occlusion Mask Darkness")}},
+			{EInterchangePropertyTracks::LightOcclusionTint, {UMovieSceneColorTrack::StaticClass()->GetName(),  GET_MEMBER_NAME_STRING_CHECKED(USkyLightComponent, OcclusionTint) , TEXT("Occlusion Tint")}},
+			{EInterchangePropertyTracks::LightOuterConeAngle, {UMovieSceneFloatTrack::StaticClass()->GetName(),  GET_MEMBER_NAME_STRING_CHECKED(USpotLightComponent, OuterConeAngle) , TEXT("Outer Cone Angle")}},
+			{EInterchangePropertyTracks::LightSamplesPerPixel, {UMovieSceneIntegerTrack::StaticClass()->GetName(),  GET_MEMBER_NAME_STRING_CHECKED(ULightComponentBase, SamplesPerPixel) , TEXT("Samples Per Pixel")}},
+			{EInterchangePropertyTracks::LightShadowAmount, {UMovieSceneFloatTrack::StaticClass()->GetName(),  GET_MEMBER_NAME_STRING_CHECKED(UDirectionalLightComponent, ShadowAmount) , TEXT("Shadow Amount")}},
+			{EInterchangePropertyTracks::LightShadowBias, {UMovieSceneFloatTrack::StaticClass()->GetName(),  GET_MEMBER_NAME_STRING_CHECKED(ULightComponent, ShadowBias) , TEXT("Shadow Bias")}},
+			{EInterchangePropertyTracks::LightShadowCascadeBiasDistribution, {UMovieSceneFloatTrack::StaticClass()->GetName(),  GET_MEMBER_NAME_STRING_CHECKED(UDirectionalLightComponent, ShadowCascadeBiasDistribution) , TEXT("Shadow Cascade Bias Distribution")}},
+			{EInterchangePropertyTracks::LightShadowDistanceFadeoutFraction, {UMovieSceneFloatTrack::StaticClass()->GetName(),  GET_MEMBER_NAME_STRING_CHECKED(UDirectionalLightComponent, ShadowDistanceFadeoutFraction) , TEXT("Distance Fadeout Fraction")}},
+			{EInterchangePropertyTracks::LightShadowSlopeBias, {UMovieSceneFloatTrack::StaticClass()->GetName(),  GET_MEMBER_NAME_STRING_CHECKED(ULightComponent, ShadowSlopeBias) , TEXT("Shadow Slope Bias")}},
+			{EInterchangePropertyTracks::LightShadowSourceAngleFactor, {UMovieSceneFloatTrack::StaticClass()->GetName(),  GET_MEMBER_NAME_STRING_CHECKED(UDirectionalLightComponent, ShadowSourceAngleFactor) , TEXT("Shadow Source Angle Factor")}},
+			{EInterchangePropertyTracks::LightSoftSourceRadius, {UMovieSceneFloatTrack::StaticClass()->GetName(),  GET_MEMBER_NAME_STRING_CHECKED(UPointLightComponent, SoftSourceRadius) , TEXT("Soft Source Radius")}},
+			{EInterchangePropertyTracks::LightSourceAngle, {UMovieSceneFloatTrack::StaticClass()->GetName(),  GET_MEMBER_NAME_STRING_CHECKED(UDirectionalLightComponent, LightSourceAngle) , TEXT("Source Angle")}},
+			{EInterchangePropertyTracks::LightSourceCubemapAngle, {UMovieSceneFloatTrack::StaticClass()->GetName(),  GET_MEMBER_NAME_STRING_CHECKED(USkyLightComponent, SourceCubemapAngle) , TEXT("Source Cubemap Angle")}},
+			{EInterchangePropertyTracks::LightSourceHeight, {UMovieSceneFloatTrack::StaticClass()->GetName(),  GET_MEMBER_NAME_STRING_CHECKED(URectLightComponent, SourceHeight) , TEXT("Source Height")}},
+			{EInterchangePropertyTracks::LightSourceLength, {UMovieSceneFloatTrack::StaticClass()->GetName(),  GET_MEMBER_NAME_STRING_CHECKED(UPointLightComponent, SourceLength) , TEXT("Source Length")}},
+			{EInterchangePropertyTracks::LightSourceRadius, {UMovieSceneFloatTrack::StaticClass()->GetName(),  GET_MEMBER_NAME_STRING_CHECKED(UPointLightComponent, SourceRadius) , TEXT("Source Radius")}},
+			{EInterchangePropertyTracks::LightSourceSoftAngle, {UMovieSceneFloatTrack::StaticClass()->GetName(),  GET_MEMBER_NAME_STRING_CHECKED(UDirectionalLightComponent, LightSourceSoftAngle) , TEXT("Source Soft Angle")}},
+			{EInterchangePropertyTracks::LightSourceWidth, {UMovieSceneFloatTrack::StaticClass()->GetName(),  GET_MEMBER_NAME_STRING_CHECKED(URectLightComponent, SourceWidth) , TEXT("Source Width")}},
+			{EInterchangePropertyTracks::LightSpecularScale, {UMovieSceneFloatTrack::StaticClass()->GetName(),  GET_MEMBER_NAME_STRING_CHECKED(ULightComponent, SpecularScale) , TEXT("Specular Scale")}},
+			{EInterchangePropertyTracks::LightTemperature, {UMovieSceneFloatTrack::StaticClass()->GetName(),  GET_MEMBER_NAME_STRING_CHECKED(ULightComponent, Temperature) , TEXT("Temperature")}},
+			{EInterchangePropertyTracks::LightTransmission, {UMovieSceneBoolTrack::StaticClass()->GetName(),  GET_MEMBER_NAME_STRING_CHECKED(ULightComponentBase, bTransmission) , TEXT("Transmission")}},
+			{EInterchangePropertyTracks::LightUseIESBrightness, {UMovieSceneBoolTrack::StaticClass()->GetName(),  GET_MEMBER_NAME_STRING_CHECKED(ULightComponent, bUseIESBrightness) , TEXT("Use IES Intensity")}},
+			{EInterchangePropertyTracks::LightUseInverseSquaredFalloff, {UMovieSceneBoolTrack::StaticClass()->GetName(),  GET_MEMBER_NAME_STRING_CHECKED(UPointLightComponent, bUseInverseSquaredFalloff) , TEXT("Use Inverse Squared Falloff")}},
+			{EInterchangePropertyTracks::LightUseTemperature, {UMovieSceneBoolTrack::StaticClass()->GetName(),  GET_MEMBER_NAME_STRING_CHECKED(ULightComponent, bUseTemperature) , TEXT("Use Temperature")}},
+			{EInterchangePropertyTracks::LightVolumetricScatteringIntensity, {UMovieSceneFloatTrack::StaticClass()->GetName(),  GET_MEMBER_NAME_STRING_CHECKED(ULightComponentBase, VolumetricScatteringIntensity) , TEXT("Volumetric Scattering Intensity")}},
 
 			// Camera
-			{EInterchangePropertyTracks::CameraAutoCalculateOrthoPlanes, {UMovieSceneBoolTrack::StaticClass()->GetName(), GET_MEMBER_NAME_STRING_CHECKED(UCameraComponent, bAutoCalculateOrthoPlanes), TEXT("Auto Calculate Ortho Planes")}},
-			{EInterchangePropertyTracks::CameraAspectRatio, {UMovieSceneFloatTrack::StaticClass()->GetName(), GET_MEMBER_NAME_STRING_CHECKED(UCameraComponent, AspectRatio), TEXT("Aspect Ratio")}},
-			{EInterchangePropertyTracks::CameraAspectRatioAxisConstraint, {UMovieSceneByteTrack::StaticClass()->GetName(), GET_MEMBER_NAME_STRING_CHECKED(UCameraComponent, AspectRatioAxisConstraint), TEXT("Aspect Ratio Axis Constraint"), StaticEnum<EAspectRatioAxisConstraint>()}},
-			{EInterchangePropertyTracks::CameraConstrainAspectRatio, {UMovieSceneBoolTrack::StaticClass()->GetName(), GET_MEMBER_NAME_STRING_CHECKED(UCameraComponent, bConstrainAspectRatio), TEXT("Constrain Aspect Ratio")}},
-			{EInterchangePropertyTracks::CameraCurrentAperture, {UMovieSceneFloatTrack::StaticClass()->GetName(), GET_MEMBER_NAME_STRING_CHECKED(UCineCameraComponent, CurrentAperture), TEXT("Current Aperture")}},
-			{EInterchangePropertyTracks::CameraCurrentFocalLength, {UMovieSceneFloatTrack::StaticClass()->GetName(), GET_MEMBER_NAME_STRING_CHECKED(UCineCameraComponent, CurrentFocalLength), TEXT("Current Focal Length")}},
-			{EInterchangePropertyTracks::CameraCustomNearClippingPlane, {UMovieSceneFloatTrack::StaticClass()->GetName(), GET_MEMBER_NAME_STRING_CHECKED(UCineCameraComponent, CustomNearClippingPlane), TEXT("Custom Near Clipping Plane")}},
-			{EInterchangePropertyTracks::CameraFieldOfView, {UMovieSceneFloatTrack::StaticClass()->GetName(), GET_MEMBER_NAME_STRING_CHECKED(UCameraComponent, FieldOfView), TEXT("Field of View")}},
-			{EInterchangePropertyTracks::CameraFilmbackSensorAspectRatio, {UMovieSceneFloatTrack::StaticClass()->GetName(), GET_MEMBER_NAME_STRING_CHECKED(UCineCameraComponent, Filmback) + FString{TEXT(".")} + GET_MEMBER_NAME_STRING_CHECKED(FCameraFilmbackSettings, SensorAspectRatio) , TEXT("Sensor Aspect Ratio (Filmback)")}},
-			{EInterchangePropertyTracks::CameraFilmbackSensorHeight, {UMovieSceneFloatTrack::StaticClass()->GetName(), GET_MEMBER_NAME_STRING_CHECKED(UCineCameraComponent, Filmback) + FString{TEXT(".")} + GET_MEMBER_NAME_STRING_CHECKED(FCameraFilmbackSettings, SensorHeight) , TEXT("Sensor Height (Filmback)")}},
-			{EInterchangePropertyTracks::CameraFilmbackSensorWidth, {UMovieSceneFloatTrack::StaticClass()->GetName(), GET_MEMBER_NAME_STRING_CHECKED(UCineCameraComponent, Filmback) + FString{TEXT(".")} + GET_MEMBER_NAME_STRING_CHECKED(FCameraFilmbackSettings, SensorWidth) , TEXT("Sensor Width (Filmback)")}},
-			{EInterchangePropertyTracks::CameraFocusSettingsManualFocusDistance, {UMovieSceneFloatTrack::StaticClass()->GetName(), GET_MEMBER_NAME_STRING_CHECKED(UCineCameraComponent, FocusSettings) + FString{TEXT(".")} + GET_MEMBER_NAME_STRING_CHECKED(FCameraFocusSettings, ManualFocusDistance), TEXT("Manuel Focus Distance (Focus Settings)")}},
-			{EInterchangePropertyTracks::CameraOrthoFarClipPlane, {UMovieSceneFloatTrack::StaticClass()->GetName(), GET_MEMBER_NAME_STRING_CHECKED(UCameraComponent, OrthoFarClipPlane), TEXT("Ortho Far Clip Plane")}},
-			{EInterchangePropertyTracks::CameraOrthoNearClipPlane, {UMovieSceneFloatTrack::StaticClass()->GetName(), GET_MEMBER_NAME_STRING_CHECKED(UCameraComponent, OrthoNearClipPlane), TEXT("Ortho Near Clip Plane")}},
-			{EInterchangePropertyTracks::CameraOrthoWidth, {UMovieSceneFloatTrack::StaticClass()->GetName(), GET_MEMBER_NAME_STRING_CHECKED(UCameraComponent, OrthoWidth), TEXT("Ortho Width")}},
-			{EInterchangePropertyTracks::CameraPostProcessBlendWeight, {UMovieSceneFloatTrack::StaticClass()->GetName(), GET_MEMBER_NAME_STRING_CHECKED(UCameraComponent, PostProcessBlendWeight), TEXT("PostProcess Blend Weight")}},
-			{EInterchangePropertyTracks::CameraProjectionMode, {UMovieSceneByteTrack::StaticClass()->GetName(), GET_MEMBER_NAME_STRING_CHECKED(UCameraComponent, ProjectionMode), TEXT("Projection Mode"), StaticEnum<ECameraProjectionMode::Type>()}},
-			{EInterchangePropertyTracks::CameraUseFieldOfViewForLOD, {UMovieSceneBoolTrack::StaticClass()->GetName(), GET_MEMBER_NAME_STRING_CHECKED(UCameraComponent, bUseFieldOfViewForLOD), TEXT("Use Field Of View For LOD")}},
+			{EInterchangePropertyTracks::CameraAspectRatio, {UMovieSceneFloatTrack::StaticClass()->GetName(),  GET_MEMBER_NAME_STRING_CHECKED(UCameraComponent, AspectRatio) , TEXT("Aspect Ratio")}},
+			{EInterchangePropertyTracks::CameraAspectRatioAxisConstraint, {UMovieSceneEnumTrack::StaticClass()->GetName(),  GET_MEMBER_NAME_STRING_CHECKED(UCameraComponent, AspectRatioAxisConstraint) , TEXT("Aspect Ratio Axis Constraint"), StaticEnum<EAspectRatioAxisConstraint>()}},
+			{EInterchangePropertyTracks::CameraAutoCalculateOrthoPlanes, {UMovieSceneBoolTrack::StaticClass()->GetName(),  GET_MEMBER_NAME_STRING_CHECKED(UCameraComponent, bAutoCalculateOrthoPlanes) , TEXT("Auto Calculate Ortho Planes")}},
+			{EInterchangePropertyTracks::CameraConstrainAspectRatio, {UMovieSceneBoolTrack::StaticClass()->GetName(),  GET_MEMBER_NAME_STRING_CHECKED(UCameraComponent, bConstrainAspectRatio) , TEXT("Constrain Aspect Ratio")}},
+			{EInterchangePropertyTracks::CameraCurrentAperture, {UMovieSceneFloatTrack::StaticClass()->GetName(),  GET_MEMBER_NAME_STRING_CHECKED(UCineCameraComponent, CurrentAperture) , TEXT("Current Aperture")}},
+			{EInterchangePropertyTracks::CameraCurrentFocalLength, {UMovieSceneFloatTrack::StaticClass()->GetName(),  GET_MEMBER_NAME_STRING_CHECKED(UCineCameraComponent, CurrentFocalLength) , TEXT("Current Focal Length")}},
+			{EInterchangePropertyTracks::CameraCustomNearClippingPlane, {UMovieSceneFloatTrack::StaticClass()->GetName(),  GET_MEMBER_NAME_STRING_CHECKED(UCineCameraComponent, CustomNearClippingPlane) , TEXT("Custom Near Clipping Plane")}},
+			{EInterchangePropertyTracks::CameraFieldOfView, {UMovieSceneFloatTrack::StaticClass()->GetName(),  GET_MEMBER_NAME_STRING_CHECKED(UCameraComponent, FieldOfView) , TEXT("Field Of View")}},
+			{EInterchangePropertyTracks::CameraFilmbackSensorAspectRatio, {UMovieSceneFloatTrack::StaticClass()->GetName(),  GET_MEMBER_NAME_STRING_CHECKED(UCineCameraComponent, Filmback) + FString{ TEXT(".") } + GET_MEMBER_NAME_STRING_CHECKED(FCameraFilmbackSettings, SensorAspectRatio), TEXT("Sensor Aspect Ratio (Filmback)")}},
+			{EInterchangePropertyTracks::CameraFilmbackSensorHeight, {UMovieSceneFloatTrack::StaticClass()->GetName(),  GET_MEMBER_NAME_STRING_CHECKED(UCineCameraComponent, Filmback) + FString{ TEXT(".") } + GET_MEMBER_NAME_STRING_CHECKED(FCameraFilmbackSettings, SensorHeight), TEXT("Sensor Height (Filmback)")}},
+			{EInterchangePropertyTracks::CameraFilmbackSensorWidth, {UMovieSceneFloatTrack::StaticClass()->GetName(),  GET_MEMBER_NAME_STRING_CHECKED(UCineCameraComponent, Filmback) + FString{ TEXT(".") } + GET_MEMBER_NAME_STRING_CHECKED(FCameraFilmbackSettings, SensorWidth), TEXT("Sensor Width (Filmback)")}},
+			{EInterchangePropertyTracks::CameraFocusSettingsFocusOffset, {UMovieSceneFloatTrack::StaticClass()->GetName(),  GET_MEMBER_NAME_STRING_CHECKED(UCineCameraComponent, FocusSettings) + FString{ TEXT(".") } + GET_MEMBER_NAME_STRING_CHECKED(FCameraFocusSettings, FocusOffset), TEXT("Focus Offset (Focus Settings)")}},
+			{EInterchangePropertyTracks::CameraFocusSettingsManualFocusDistance, {UMovieSceneFloatTrack::StaticClass()->GetName(),  GET_MEMBER_NAME_STRING_CHECKED(UCineCameraComponent, FocusSettings) + FString{ TEXT(".") } + GET_MEMBER_NAME_STRING_CHECKED(FCameraFocusSettings, ManualFocusDistance), TEXT("Manual Focus Distance (Focus Settings)")}},
+			{EInterchangePropertyTracks::CameraOrthoFarClipPlane, {UMovieSceneFloatTrack::StaticClass()->GetName(),  GET_MEMBER_NAME_STRING_CHECKED(UCameraComponent, OrthoFarClipPlane) , TEXT("Ortho Far Clip Plane")}},
+			{EInterchangePropertyTracks::CameraOrthoNearClipPlane, {UMovieSceneFloatTrack::StaticClass()->GetName(),  GET_MEMBER_NAME_STRING_CHECKED(UCameraComponent, OrthoNearClipPlane) , TEXT("Ortho Near Clip Plane")}},
+			{EInterchangePropertyTracks::CameraOrthoWidth, {UMovieSceneFloatTrack::StaticClass()->GetName(),  GET_MEMBER_NAME_STRING_CHECKED(UCameraComponent, OrthoWidth) , TEXT("Ortho Width")}},
+			{EInterchangePropertyTracks::CameraOverrideAspectRatioAxisConstraint, {UMovieSceneBoolTrack::StaticClass()->GetName(),  GET_MEMBER_NAME_STRING_CHECKED(UCameraComponent, bOverrideAspectRatioAxisConstraint) , TEXT("Override Aspect Ratio Axis Constraint")}},
+			{EInterchangePropertyTracks::CameraPostProcessBlendWeight, {UMovieSceneFloatTrack::StaticClass()->GetName(),  GET_MEMBER_NAME_STRING_CHECKED(UCameraComponent, PostProcessBlendWeight) , TEXT("Post Process Blend Weight")}},
+			{EInterchangePropertyTracks::CameraPostProcessSettingsAmbientCubemapIntensity, {UMovieSceneFloatTrack::StaticClass()->GetName(),  GET_MEMBER_NAME_STRING_CHECKED(UCameraComponent, PostProcessSettings) + FString{ TEXT(".") } + GET_MEMBER_NAME_STRING_CHECKED(FPostProcessSettings, AmbientCubemapIntensity), TEXT("Intensity (Post Process Settings)")}},
+			{EInterchangePropertyTracks::CameraPostProcessSettingsAmbientCubemapTint, {UMovieSceneColorTrack::StaticClass()->GetName(),  GET_MEMBER_NAME_STRING_CHECKED(UCameraComponent, PostProcessSettings) + FString{ TEXT(".") } + GET_MEMBER_NAME_STRING_CHECKED(FPostProcessSettings, AmbientCubemapTint), TEXT("Tint (Post Process Settings)")}},
+			{EInterchangePropertyTracks::CameraPostProcessSettingsAmbientOcclusionBias, {UMovieSceneFloatTrack::StaticClass()->GetName(),  GET_MEMBER_NAME_STRING_CHECKED(UCameraComponent, PostProcessSettings) + FString{ TEXT(".") } + GET_MEMBER_NAME_STRING_CHECKED(FPostProcessSettings, AmbientOcclusionBias), TEXT("Bias (Post Process Settings)")}},
+			{EInterchangePropertyTracks::CameraPostProcessSettingsAmbientOcclusionFadeDistance, {UMovieSceneFloatTrack::StaticClass()->GetName(),  GET_MEMBER_NAME_STRING_CHECKED(UCameraComponent, PostProcessSettings) + FString{ TEXT(".") } + GET_MEMBER_NAME_STRING_CHECKED(FPostProcessSettings, AmbientOcclusionFadeDistance), TEXT("Fade Out Distance (Post Process Settings)")}},
+			{EInterchangePropertyTracks::CameraPostProcessSettingsAmbientOcclusionFadeRadius, {UMovieSceneFloatTrack::StaticClass()->GetName(),  GET_MEMBER_NAME_STRING_CHECKED(UCameraComponent, PostProcessSettings) + FString{ TEXT(".") } + GET_MEMBER_NAME_STRING_CHECKED(FPostProcessSettings, AmbientOcclusionFadeRadius), TEXT("Fade Out Radius (Post Process Settings)")}},
+			{EInterchangePropertyTracks::CameraPostProcessSettingsAmbientOcclusionIntensity, {UMovieSceneFloatTrack::StaticClass()->GetName(),  GET_MEMBER_NAME_STRING_CHECKED(UCameraComponent, PostProcessSettings) + FString{ TEXT(".") } + GET_MEMBER_NAME_STRING_CHECKED(FPostProcessSettings, AmbientOcclusionIntensity), TEXT("Intensity (Post Process Settings)")}},
+			{EInterchangePropertyTracks::CameraPostProcessSettingsAmbientOcclusionMipBlend, {UMovieSceneFloatTrack::StaticClass()->GetName(),  GET_MEMBER_NAME_STRING_CHECKED(UCameraComponent, PostProcessSettings) + FString{ TEXT(".") } + GET_MEMBER_NAME_STRING_CHECKED(FPostProcessSettings, AmbientOcclusionMipBlend), TEXT("Mip Blend (Post Process Settings)")}},
+			{EInterchangePropertyTracks::CameraPostProcessSettingsAmbientOcclusionMipScale, {UMovieSceneFloatTrack::StaticClass()->GetName(),  GET_MEMBER_NAME_STRING_CHECKED(UCameraComponent, PostProcessSettings) + FString{ TEXT(".") } + GET_MEMBER_NAME_STRING_CHECKED(FPostProcessSettings, AmbientOcclusionMipScale), TEXT("Mip Scale (Post Process Settings)")}},
+			{EInterchangePropertyTracks::CameraPostProcessSettingsAmbientOcclusionMipThreshold, {UMovieSceneFloatTrack::StaticClass()->GetName(),  GET_MEMBER_NAME_STRING_CHECKED(UCameraComponent, PostProcessSettings) + FString{ TEXT(".") } + GET_MEMBER_NAME_STRING_CHECKED(FPostProcessSettings, AmbientOcclusionMipThreshold), TEXT("Mip Threshold (Post Process Settings)")}},
+			{EInterchangePropertyTracks::CameraPostProcessSettingsAmbientOcclusionPower, {UMovieSceneFloatTrack::StaticClass()->GetName(),  GET_MEMBER_NAME_STRING_CHECKED(UCameraComponent, PostProcessSettings) + FString{ TEXT(".") } + GET_MEMBER_NAME_STRING_CHECKED(FPostProcessSettings, AmbientOcclusionPower), TEXT("Power (Post Process Settings)")}},
+			{EInterchangePropertyTracks::CameraPostProcessSettingsAmbientOcclusionQuality, {UMovieSceneFloatTrack::StaticClass()->GetName(),  GET_MEMBER_NAME_STRING_CHECKED(UCameraComponent, PostProcessSettings) + FString{ TEXT(".") } + GET_MEMBER_NAME_STRING_CHECKED(FPostProcessSettings, AmbientOcclusionQuality), TEXT("Quality (Post Process Settings)")}},
+			{EInterchangePropertyTracks::CameraPostProcessSettingsAmbientOcclusionRadius, {UMovieSceneFloatTrack::StaticClass()->GetName(),  GET_MEMBER_NAME_STRING_CHECKED(UCameraComponent, PostProcessSettings) + FString{ TEXT(".") } + GET_MEMBER_NAME_STRING_CHECKED(FPostProcessSettings, AmbientOcclusionRadius), TEXT("Radius (Post Process Settings)")}},
+			{EInterchangePropertyTracks::CameraPostProcessSettingsAmbientOcclusionStaticFraction, {UMovieSceneFloatTrack::StaticClass()->GetName(),  GET_MEMBER_NAME_STRING_CHECKED(UCameraComponent, PostProcessSettings) + FString{ TEXT(".") } + GET_MEMBER_NAME_STRING_CHECKED(FPostProcessSettings, AmbientOcclusionStaticFraction), TEXT("Static Fraction (Post Process Settings)")}},
+			{EInterchangePropertyTracks::CameraPostProcessSettingsAmbientOcclusionTemporalBlendWeight, {UMovieSceneFloatTrack::StaticClass()->GetName(),  GET_MEMBER_NAME_STRING_CHECKED(UCameraComponent, PostProcessSettings) + FString{ TEXT(".") } + GET_MEMBER_NAME_STRING_CHECKED(FPostProcessSettings, AmbientOcclusionTemporalBlendWeight), TEXT("Temporal Blend Weight (Post Process Settings)")}},
+			{EInterchangePropertyTracks::CameraPostProcessSettingsAutoExposureBias, {UMovieSceneFloatTrack::StaticClass()->GetName(),  GET_MEMBER_NAME_STRING_CHECKED(UCameraComponent, PostProcessSettings) + FString{ TEXT(".") } + GET_MEMBER_NAME_STRING_CHECKED(FPostProcessSettings, AutoExposureBias), TEXT("Exposure Compensation (Post Process Settings)")}},
+			{EInterchangePropertyTracks::CameraPostProcessSettingsAutoExposureHighPercent, {UMovieSceneFloatTrack::StaticClass()->GetName(),  GET_MEMBER_NAME_STRING_CHECKED(UCameraComponent, PostProcessSettings) + FString{ TEXT(".") } + GET_MEMBER_NAME_STRING_CHECKED(FPostProcessSettings, AutoExposureHighPercent), TEXT("High Percent (Post Process Settings)")}},
+			{EInterchangePropertyTracks::CameraPostProcessSettingsAutoExposureLowPercent, {UMovieSceneFloatTrack::StaticClass()->GetName(),  GET_MEMBER_NAME_STRING_CHECKED(UCameraComponent, PostProcessSettings) + FString{ TEXT(".") } + GET_MEMBER_NAME_STRING_CHECKED(FPostProcessSettings, AutoExposureLowPercent), TEXT("Low Percent (Post Process Settings)")}},
+			{EInterchangePropertyTracks::CameraPostProcessSettingsAutoExposureMaxBrightness, {UMovieSceneFloatTrack::StaticClass()->GetName(),  GET_MEMBER_NAME_STRING_CHECKED(UCameraComponent, PostProcessSettings) + FString{ TEXT(".") } + GET_MEMBER_NAME_STRING_CHECKED(FPostProcessSettings, AutoExposureMaxBrightness), TEXT("Max Brightness (Post Process Settings)")}},
+			{EInterchangePropertyTracks::CameraPostProcessSettingsAutoExposureMinBrightness, {UMovieSceneFloatTrack::StaticClass()->GetName(),  GET_MEMBER_NAME_STRING_CHECKED(UCameraComponent, PostProcessSettings) + FString{ TEXT(".") } + GET_MEMBER_NAME_STRING_CHECKED(FPostProcessSettings, AutoExposureMinBrightness), TEXT("Min Brightness (Post Process Settings)")}},
+			{EInterchangePropertyTracks::CameraPostProcessSettingsAutoExposureSpeedDown, {UMovieSceneFloatTrack::StaticClass()->GetName(),  GET_MEMBER_NAME_STRING_CHECKED(UCameraComponent, PostProcessSettings) + FString{ TEXT(".") } + GET_MEMBER_NAME_STRING_CHECKED(FPostProcessSettings, AutoExposureSpeedDown), TEXT("Speed Down (Post Process Settings)")}},
+			{EInterchangePropertyTracks::CameraPostProcessSettingsAutoExposureSpeedUp, {UMovieSceneFloatTrack::StaticClass()->GetName(),  GET_MEMBER_NAME_STRING_CHECKED(UCameraComponent, PostProcessSettings) + FString{ TEXT(".") } + GET_MEMBER_NAME_STRING_CHECKED(FPostProcessSettings, AutoExposureSpeedUp), TEXT("Speed Up (Post Process Settings)")}},
+			{EInterchangePropertyTracks::CameraPostProcessSettingsBloom1Size, {UMovieSceneFloatTrack::StaticClass()->GetName(),  GET_MEMBER_NAME_STRING_CHECKED(UCameraComponent, PostProcessSettings) + FString{ TEXT(".") } + GET_MEMBER_NAME_STRING_CHECKED(FPostProcessSettings, Bloom1Size), TEXT("#1 Size (Post Process Settings)")}},
+			{EInterchangePropertyTracks::CameraPostProcessSettingsBloom1Tint, {UMovieSceneColorTrack::StaticClass()->GetName(),  GET_MEMBER_NAME_STRING_CHECKED(UCameraComponent, PostProcessSettings) + FString{ TEXT(".") } + GET_MEMBER_NAME_STRING_CHECKED(FPostProcessSettings, Bloom1Tint), TEXT("#1 Tint (Post Process Settings)")}},
+			{EInterchangePropertyTracks::CameraPostProcessSettingsBloom2Size, {UMovieSceneFloatTrack::StaticClass()->GetName(),  GET_MEMBER_NAME_STRING_CHECKED(UCameraComponent, PostProcessSettings) + FString{ TEXT(".") } + GET_MEMBER_NAME_STRING_CHECKED(FPostProcessSettings, Bloom2Size), TEXT("#2 Size (Post Process Settings)")}},
+			{EInterchangePropertyTracks::CameraPostProcessSettingsBloom2Tint, {UMovieSceneColorTrack::StaticClass()->GetName(),  GET_MEMBER_NAME_STRING_CHECKED(UCameraComponent, PostProcessSettings) + FString{ TEXT(".") } + GET_MEMBER_NAME_STRING_CHECKED(FPostProcessSettings, Bloom2Tint), TEXT("#2 Tint (Post Process Settings)")}},
+			{EInterchangePropertyTracks::CameraPostProcessSettingsBloom3Size, {UMovieSceneFloatTrack::StaticClass()->GetName(),  GET_MEMBER_NAME_STRING_CHECKED(UCameraComponent, PostProcessSettings) + FString{ TEXT(".") } + GET_MEMBER_NAME_STRING_CHECKED(FPostProcessSettings, Bloom3Size), TEXT("#3 Size (Post Process Settings)")}},
+			{EInterchangePropertyTracks::CameraPostProcessSettingsBloom3Tint, {UMovieSceneColorTrack::StaticClass()->GetName(),  GET_MEMBER_NAME_STRING_CHECKED(UCameraComponent, PostProcessSettings) + FString{ TEXT(".") } + GET_MEMBER_NAME_STRING_CHECKED(FPostProcessSettings, Bloom3Tint), TEXT("#3 Tint (Post Process Settings)")}},
+			{EInterchangePropertyTracks::CameraPostProcessSettingsBloom4Size, {UMovieSceneFloatTrack::StaticClass()->GetName(),  GET_MEMBER_NAME_STRING_CHECKED(UCameraComponent, PostProcessSettings) + FString{ TEXT(".") } + GET_MEMBER_NAME_STRING_CHECKED(FPostProcessSettings, Bloom4Size), TEXT("#4 Size (Post Process Settings)")}},
+			{EInterchangePropertyTracks::CameraPostProcessSettingsBloom4Tint, {UMovieSceneColorTrack::StaticClass()->GetName(),  GET_MEMBER_NAME_STRING_CHECKED(UCameraComponent, PostProcessSettings) + FString{ TEXT(".") } + GET_MEMBER_NAME_STRING_CHECKED(FPostProcessSettings, Bloom4Tint), TEXT("#4 Tint (Post Process Settings)")}},
+			{EInterchangePropertyTracks::CameraPostProcessSettingsBloom5Size, {UMovieSceneFloatTrack::StaticClass()->GetName(),  GET_MEMBER_NAME_STRING_CHECKED(UCameraComponent, PostProcessSettings) + FString{ TEXT(".") } + GET_MEMBER_NAME_STRING_CHECKED(FPostProcessSettings, Bloom5Size), TEXT("#5 Size (Post Process Settings)")}},
+			{EInterchangePropertyTracks::CameraPostProcessSettingsBloom5Tint, {UMovieSceneColorTrack::StaticClass()->GetName(),  GET_MEMBER_NAME_STRING_CHECKED(UCameraComponent, PostProcessSettings) + FString{ TEXT(".") } + GET_MEMBER_NAME_STRING_CHECKED(FPostProcessSettings, Bloom5Tint), TEXT("#5 Tint (Post Process Settings)")}},
+			{EInterchangePropertyTracks::CameraPostProcessSettingsBloom6Size, {UMovieSceneFloatTrack::StaticClass()->GetName(),  GET_MEMBER_NAME_STRING_CHECKED(UCameraComponent, PostProcessSettings) + FString{ TEXT(".") } + GET_MEMBER_NAME_STRING_CHECKED(FPostProcessSettings, Bloom6Size), TEXT("#6 Size (Post Process Settings)")}},
+			{EInterchangePropertyTracks::CameraPostProcessSettingsBloom6Tint, {UMovieSceneColorTrack::StaticClass()->GetName(),  GET_MEMBER_NAME_STRING_CHECKED(UCameraComponent, PostProcessSettings) + FString{ TEXT(".") } + GET_MEMBER_NAME_STRING_CHECKED(FPostProcessSettings, Bloom6Tint), TEXT("#6 Tint (Post Process Settings)")}},
+			{EInterchangePropertyTracks::CameraPostProcessSettingsBloomConvolutionBufferScale, {UMovieSceneFloatTrack::StaticClass()->GetName(),  GET_MEMBER_NAME_STRING_CHECKED(UCameraComponent, PostProcessSettings) + FString{ TEXT(".") } + GET_MEMBER_NAME_STRING_CHECKED(FPostProcessSettings, BloomConvolutionBufferScale), TEXT("Convolution Buffer (Post Process Settings)")}},
+			{EInterchangePropertyTracks::CameraPostProcessSettingsBloomConvolutionPreFilterMax, {UMovieSceneFloatTrack::StaticClass()->GetName(),  GET_MEMBER_NAME_STRING_CHECKED(UCameraComponent, PostProcessSettings) + FString{ TEXT(".") } + GET_MEMBER_NAME_STRING_CHECKED(FPostProcessSettings, BloomConvolutionPreFilterMax), TEXT("Convolution Boost Max (Post Process Settings)")}},
+			{EInterchangePropertyTracks::CameraPostProcessSettingsBloomConvolutionPreFilterMin, {UMovieSceneFloatTrack::StaticClass()->GetName(),  GET_MEMBER_NAME_STRING_CHECKED(UCameraComponent, PostProcessSettings) + FString{ TEXT(".") } + GET_MEMBER_NAME_STRING_CHECKED(FPostProcessSettings, BloomConvolutionPreFilterMin), TEXT("Convolution Boost Min (Post Process Settings)")}},
+			{EInterchangePropertyTracks::CameraPostProcessSettingsBloomConvolutionPreFilterMult, {UMovieSceneFloatTrack::StaticClass()->GetName(),  GET_MEMBER_NAME_STRING_CHECKED(UCameraComponent, PostProcessSettings) + FString{ TEXT(".") } + GET_MEMBER_NAME_STRING_CHECKED(FPostProcessSettings, BloomConvolutionPreFilterMult), TEXT("Convolution Boost Mult (Post Process Settings)")}},
+			{EInterchangePropertyTracks::CameraPostProcessSettingsBloomConvolutionScatterDispersion, {UMovieSceneFloatTrack::StaticClass()->GetName(),  GET_MEMBER_NAME_STRING_CHECKED(UCameraComponent, PostProcessSettings) + FString{ TEXT(".") } + GET_MEMBER_NAME_STRING_CHECKED(FPostProcessSettings, BloomConvolutionScatterDispersion), TEXT("Convolution Scatter Dispersion (Post Process Settings)")}},
+			{EInterchangePropertyTracks::CameraPostProcessSettingsBloomConvolutionSize, {UMovieSceneFloatTrack::StaticClass()->GetName(),  GET_MEMBER_NAME_STRING_CHECKED(UCameraComponent, PostProcessSettings) + FString{ TEXT(".") } + GET_MEMBER_NAME_STRING_CHECKED(FPostProcessSettings, BloomConvolutionSize), TEXT("Convolution Scale (Post Process Settings)")}},
+			{EInterchangePropertyTracks::CameraPostProcessSettingsBloomDirtMaskIntensity, {UMovieSceneFloatTrack::StaticClass()->GetName(),  GET_MEMBER_NAME_STRING_CHECKED(UCameraComponent, PostProcessSettings) + FString{ TEXT(".") } + GET_MEMBER_NAME_STRING_CHECKED(FPostProcessSettings, BloomDirtMaskIntensity), TEXT("Dirt Mask Intensity (Post Process Settings)")}},
+			{EInterchangePropertyTracks::CameraPostProcessSettingsBloomDirtMaskTint, {UMovieSceneColorTrack::StaticClass()->GetName(),  GET_MEMBER_NAME_STRING_CHECKED(UCameraComponent, PostProcessSettings) + FString{ TEXT(".") } + GET_MEMBER_NAME_STRING_CHECKED(FPostProcessSettings, BloomDirtMaskTint), TEXT("Dirt Mask Tint (Post Process Settings)")}},
+			{EInterchangePropertyTracks::CameraPostProcessSettingsBloomIntensity, {UMovieSceneFloatTrack::StaticClass()->GetName(),  GET_MEMBER_NAME_STRING_CHECKED(UCameraComponent, PostProcessSettings) + FString{ TEXT(".") } + GET_MEMBER_NAME_STRING_CHECKED(FPostProcessSettings, BloomIntensity), TEXT("Intensity (Post Process Settings)")}},
+			{EInterchangePropertyTracks::CameraPostProcessSettingsBloomSizeScale, {UMovieSceneFloatTrack::StaticClass()->GetName(),  GET_MEMBER_NAME_STRING_CHECKED(UCameraComponent, PostProcessSettings) + FString{ TEXT(".") } + GET_MEMBER_NAME_STRING_CHECKED(FPostProcessSettings, BloomSizeScale), TEXT("Size scale (Post Process Settings)")}},
+			{EInterchangePropertyTracks::CameraPostProcessSettingsBloomThreshold, {UMovieSceneFloatTrack::StaticClass()->GetName(),  GET_MEMBER_NAME_STRING_CHECKED(UCameraComponent, PostProcessSettings) + FString{ TEXT(".") } + GET_MEMBER_NAME_STRING_CHECKED(FPostProcessSettings, BloomThreshold), TEXT("Threshold (Post Process Settings)")}},
+			{EInterchangePropertyTracks::CameraPostProcessSettingsBlueCorrection, {UMovieSceneFloatTrack::StaticClass()->GetName(),  GET_MEMBER_NAME_STRING_CHECKED(UCameraComponent, PostProcessSettings) + FString{ TEXT(".") } + GET_MEMBER_NAME_STRING_CHECKED(FPostProcessSettings, BlueCorrection), TEXT("Blue Correction (Post Process Settings)")}},
+			{EInterchangePropertyTracks::CameraPostProcessSettingsChromaticAberrationStartOffset, {UMovieSceneFloatTrack::StaticClass()->GetName(),  GET_MEMBER_NAME_STRING_CHECKED(UCameraComponent, PostProcessSettings) + FString{ TEXT(".") } + GET_MEMBER_NAME_STRING_CHECKED(FPostProcessSettings, ChromaticAberrationStartOffset), TEXT("Start Offset (Post Process Settings)")}},
+			{EInterchangePropertyTracks::CameraPostProcessSettingsColorCorrectionHighlightsMax, {UMovieSceneFloatTrack::StaticClass()->GetName(),  GET_MEMBER_NAME_STRING_CHECKED(UCameraComponent, PostProcessSettings) + FString{ TEXT(".") } + GET_MEMBER_NAME_STRING_CHECKED(FPostProcessSettings, ColorCorrectionHighlightsMax), TEXT("HighlightsMax (Post Process Settings)")}},
+			{EInterchangePropertyTracks::CameraPostProcessSettingsColorCorrectionHighlightsMin, {UMovieSceneFloatTrack::StaticClass()->GetName(),  GET_MEMBER_NAME_STRING_CHECKED(UCameraComponent, PostProcessSettings) + FString{ TEXT(".") } + GET_MEMBER_NAME_STRING_CHECKED(FPostProcessSettings, ColorCorrectionHighlightsMin), TEXT("HighlightsMin (Post Process Settings)")}},
+			{EInterchangePropertyTracks::CameraPostProcessSettingsColorCorrectionShadowsMax, {UMovieSceneFloatTrack::StaticClass()->GetName(),  GET_MEMBER_NAME_STRING_CHECKED(UCameraComponent, PostProcessSettings) + FString{ TEXT(".") } + GET_MEMBER_NAME_STRING_CHECKED(FPostProcessSettings, ColorCorrectionShadowsMax), TEXT("ShadowsMax (Post Process Settings)")}},
+			{EInterchangePropertyTracks::CameraPostProcessSettingsColorGradingIntensity, {UMovieSceneFloatTrack::StaticClass()->GetName(),  GET_MEMBER_NAME_STRING_CHECKED(UCameraComponent, PostProcessSettings) + FString{ TEXT(".") } + GET_MEMBER_NAME_STRING_CHECKED(FPostProcessSettings, ColorGradingIntensity), TEXT("Color Grading LUT Intensity (Post Process Settings)")}},
+			{EInterchangePropertyTracks::CameraPostProcessSettingsDepthOfFieldBladeCount, {UMovieSceneIntegerTrack::StaticClass()->GetName(),  GET_MEMBER_NAME_STRING_CHECKED(UCameraComponent, PostProcessSettings) + FString{ TEXT(".") } + GET_MEMBER_NAME_STRING_CHECKED(FPostProcessSettings, DepthOfFieldBladeCount), TEXT("Number of diaphragm blades (Post Process Settings)")}},
+			{EInterchangePropertyTracks::CameraPostProcessSettingsDepthOfFieldDepthBlurAmount, {UMovieSceneFloatTrack::StaticClass()->GetName(),  GET_MEMBER_NAME_STRING_CHECKED(UCameraComponent, PostProcessSettings) + FString{ TEXT(".") } + GET_MEMBER_NAME_STRING_CHECKED(FPostProcessSettings, DepthOfFieldDepthBlurAmount), TEXT("Depth Blur km for 50% (Post Process Settings)")}},
+			{EInterchangePropertyTracks::CameraPostProcessSettingsDepthOfFieldDepthBlurRadius, {UMovieSceneFloatTrack::StaticClass()->GetName(),  GET_MEMBER_NAME_STRING_CHECKED(UCameraComponent, PostProcessSettings) + FString{ TEXT(".") } + GET_MEMBER_NAME_STRING_CHECKED(FPostProcessSettings, DepthOfFieldDepthBlurRadius), TEXT("Depth Blur Radius (Post Process Settings)")}},
+			{EInterchangePropertyTracks::CameraPostProcessSettingsDepthOfFieldFarBlurSize, {UMovieSceneFloatTrack::StaticClass()->GetName(),  GET_MEMBER_NAME_STRING_CHECKED(UCameraComponent, PostProcessSettings) + FString{ TEXT(".") } + GET_MEMBER_NAME_STRING_CHECKED(FPostProcessSettings, DepthOfFieldFarBlurSize), TEXT("Far Blur Size (Post Process Settings)")}},
+			{EInterchangePropertyTracks::CameraPostProcessSettingsDepthOfFieldFarTransitionRegion, {UMovieSceneFloatTrack::StaticClass()->GetName(),  GET_MEMBER_NAME_STRING_CHECKED(UCameraComponent, PostProcessSettings) + FString{ TEXT(".") } + GET_MEMBER_NAME_STRING_CHECKED(FPostProcessSettings, DepthOfFieldFarTransitionRegion), TEXT("Far Transition Region (Post Process Settings)")}},
+			{EInterchangePropertyTracks::CameraPostProcessSettingsDepthOfFieldFocalDistance, {UMovieSceneFloatTrack::StaticClass()->GetName(),  GET_MEMBER_NAME_STRING_CHECKED(UCameraComponent, PostProcessSettings) + FString{ TEXT(".") } + GET_MEMBER_NAME_STRING_CHECKED(FPostProcessSettings, DepthOfFieldFocalDistance), TEXT("Focal Distance (Post Process Settings)")}},
+			{EInterchangePropertyTracks::CameraPostProcessSettingsDepthOfFieldFocalRegion, {UMovieSceneFloatTrack::StaticClass()->GetName(),  GET_MEMBER_NAME_STRING_CHECKED(UCameraComponent, PostProcessSettings) + FString{ TEXT(".") } + GET_MEMBER_NAME_STRING_CHECKED(FPostProcessSettings, DepthOfFieldFocalRegion), TEXT("Focal Region (Post Process Settings)")}},
+			{EInterchangePropertyTracks::CameraPostProcessSettingsDepthOfFieldFstop, {UMovieSceneFloatTrack::StaticClass()->GetName(),  GET_MEMBER_NAME_STRING_CHECKED(UCameraComponent, PostProcessSettings) + FString{ TEXT(".") } + GET_MEMBER_NAME_STRING_CHECKED(FPostProcessSettings, DepthOfFieldFstop), TEXT("Aperture (F-stop) (Post Process Settings)")}},
+			{EInterchangePropertyTracks::CameraPostProcessSettingsDepthOfFieldMinFstop, {UMovieSceneFloatTrack::StaticClass()->GetName(),  GET_MEMBER_NAME_STRING_CHECKED(UCameraComponent, PostProcessSettings) + FString{ TEXT(".") } + GET_MEMBER_NAME_STRING_CHECKED(FPostProcessSettings, DepthOfFieldMinFstop), TEXT("Maximum Aperture (min F-stop) (Post Process Settings)")}},
+			{EInterchangePropertyTracks::CameraPostProcessSettingsDepthOfFieldNearBlurSize, {UMovieSceneFloatTrack::StaticClass()->GetName(),  GET_MEMBER_NAME_STRING_CHECKED(UCameraComponent, PostProcessSettings) + FString{ TEXT(".") } + GET_MEMBER_NAME_STRING_CHECKED(FPostProcessSettings, DepthOfFieldNearBlurSize), TEXT("Near Blur Size (Post Process Settings)")}},
+			{EInterchangePropertyTracks::CameraPostProcessSettingsDepthOfFieldNearTransitionRegion, {UMovieSceneFloatTrack::StaticClass()->GetName(),  GET_MEMBER_NAME_STRING_CHECKED(UCameraComponent, PostProcessSettings) + FString{ TEXT(".") } + GET_MEMBER_NAME_STRING_CHECKED(FPostProcessSettings, DepthOfFieldNearTransitionRegion), TEXT("Near Transition Region (Post Process Settings)")}},
+			{EInterchangePropertyTracks::CameraPostProcessSettingsDepthOfFieldOcclusion, {UMovieSceneFloatTrack::StaticClass()->GetName(),  GET_MEMBER_NAME_STRING_CHECKED(UCameraComponent, PostProcessSettings) + FString{ TEXT(".") } + GET_MEMBER_NAME_STRING_CHECKED(FPostProcessSettings, DepthOfFieldOcclusion), TEXT("Occlusion (Post Process Settings)")}},
+			{EInterchangePropertyTracks::CameraPostProcessSettingsDepthOfFieldScale, {UMovieSceneFloatTrack::StaticClass()->GetName(),  GET_MEMBER_NAME_STRING_CHECKED(UCameraComponent, PostProcessSettings) + FString{ TEXT(".") } + GET_MEMBER_NAME_STRING_CHECKED(FPostProcessSettings, DepthOfFieldScale), TEXT("Scale (Post Process Settings)")}},
+			{EInterchangePropertyTracks::CameraPostProcessSettingsDepthOfFieldSkyFocusDistance, {UMovieSceneFloatTrack::StaticClass()->GetName(),  GET_MEMBER_NAME_STRING_CHECKED(UCameraComponent, PostProcessSettings) + FString{ TEXT(".") } + GET_MEMBER_NAME_STRING_CHECKED(FPostProcessSettings, DepthOfFieldSkyFocusDistance), TEXT("Sky Distance (Post Process Settings)")}},
+			{EInterchangePropertyTracks::CameraPostProcessSettingsDepthOfFieldUseHairDepth, {UMovieSceneByteTrack::StaticClass()->GetName(),  GET_MEMBER_NAME_STRING_CHECKED(UCameraComponent, PostProcessSettings) + FString{ TEXT(".") } + GET_MEMBER_NAME_STRING_CHECKED(FPostProcessSettings, DepthOfFieldUseHairDepth), TEXT("Use Hair Depth (Post Process Settings)")}},
+			{EInterchangePropertyTracks::CameraPostProcessSettingsDepthOfFieldVignetteSize, {UMovieSceneFloatTrack::StaticClass()->GetName(),  GET_MEMBER_NAME_STRING_CHECKED(UCameraComponent, PostProcessSettings) + FString{ TEXT(".") } + GET_MEMBER_NAME_STRING_CHECKED(FPostProcessSettings, DepthOfFieldVignetteSize), TEXT("Vignette Size (Post Process Settings)")}},
+			{EInterchangePropertyTracks::CameraPostProcessSettingsDynamicGlobalIlluminationMethod, {UMovieSceneEnumTrack::StaticClass()->GetName(),  GET_MEMBER_NAME_STRING_CHECKED(UCameraComponent, PostProcessSettings) + FString{ TEXT(".") } + GET_MEMBER_NAME_STRING_CHECKED(FPostProcessSettings, DynamicGlobalIlluminationMethod), TEXT("Method (Post Process Settings)"), StaticEnum<EDynamicGlobalIlluminationMethod::Type>()}},
+			{EInterchangePropertyTracks::CameraPostProcessSettingsExpandGamut, {UMovieSceneFloatTrack::StaticClass()->GetName(),  GET_MEMBER_NAME_STRING_CHECKED(UCameraComponent, PostProcessSettings) + FString{ TEXT(".") } + GET_MEMBER_NAME_STRING_CHECKED(FPostProcessSettings, ExpandGamut), TEXT("Expand Gamut (Post Process Settings)")}},
+			{EInterchangePropertyTracks::CameraPostProcessSettingsFilmBlackClip, {UMovieSceneFloatTrack::StaticClass()->GetName(),  GET_MEMBER_NAME_STRING_CHECKED(UCameraComponent, PostProcessSettings) + FString{ TEXT(".") } + GET_MEMBER_NAME_STRING_CHECKED(FPostProcessSettings, FilmBlackClip), TEXT("Black clip (Post Process Settings)")}},
+			{EInterchangePropertyTracks::CameraPostProcessSettingsFilmGrainHighlightsMax, {UMovieSceneFloatTrack::StaticClass()->GetName(),  GET_MEMBER_NAME_STRING_CHECKED(UCameraComponent, PostProcessSettings) + FString{ TEXT(".") } + GET_MEMBER_NAME_STRING_CHECKED(FPostProcessSettings, FilmGrainHighlightsMax), TEXT("Film Grain Highlights Max (Post Process Settings)")}},
+			{EInterchangePropertyTracks::CameraPostProcessSettingsFilmGrainHighlightsMin, {UMovieSceneFloatTrack::StaticClass()->GetName(),  GET_MEMBER_NAME_STRING_CHECKED(UCameraComponent, PostProcessSettings) + FString{ TEXT(".") } + GET_MEMBER_NAME_STRING_CHECKED(FPostProcessSettings, FilmGrainHighlightsMin), TEXT("Film Grain Highlights Min (Post Process Settings)")}},
+			{EInterchangePropertyTracks::CameraPostProcessSettingsFilmGrainIntensity, {UMovieSceneFloatTrack::StaticClass()->GetName(),  GET_MEMBER_NAME_STRING_CHECKED(UCameraComponent, PostProcessSettings) + FString{ TEXT(".") } + GET_MEMBER_NAME_STRING_CHECKED(FPostProcessSettings, FilmGrainIntensity), TEXT("Film Grain Intensity (Post Process Settings)")}},
+			{EInterchangePropertyTracks::CameraPostProcessSettingsFilmGrainIntensityHighlights, {UMovieSceneFloatTrack::StaticClass()->GetName(),  GET_MEMBER_NAME_STRING_CHECKED(UCameraComponent, PostProcessSettings) + FString{ TEXT(".") } + GET_MEMBER_NAME_STRING_CHECKED(FPostProcessSettings, FilmGrainIntensityHighlights), TEXT("Film Grain Intensity Highlights (Post Process Settings)")}},
+			{EInterchangePropertyTracks::CameraPostProcessSettingsFilmGrainIntensityMidtones, {UMovieSceneFloatTrack::StaticClass()->GetName(),  GET_MEMBER_NAME_STRING_CHECKED(UCameraComponent, PostProcessSettings) + FString{ TEXT(".") } + GET_MEMBER_NAME_STRING_CHECKED(FPostProcessSettings, FilmGrainIntensityMidtones), TEXT("Film Grain Intensity Midtones (Post Process Settings)")}},
+			{EInterchangePropertyTracks::CameraPostProcessSettingsFilmGrainIntensityShadows, {UMovieSceneFloatTrack::StaticClass()->GetName(),  GET_MEMBER_NAME_STRING_CHECKED(UCameraComponent, PostProcessSettings) + FString{ TEXT(".") } + GET_MEMBER_NAME_STRING_CHECKED(FPostProcessSettings, FilmGrainIntensityShadows), TEXT("Film Grain Intensity Shadows (Post Process Settings)")}},
+			{EInterchangePropertyTracks::CameraPostProcessSettingsFilmGrainShadowsMax, {UMovieSceneFloatTrack::StaticClass()->GetName(),  GET_MEMBER_NAME_STRING_CHECKED(UCameraComponent, PostProcessSettings) + FString{ TEXT(".") } + GET_MEMBER_NAME_STRING_CHECKED(FPostProcessSettings, FilmGrainShadowsMax), TEXT("Film Grain Shadows Max (Post Process Settings)")}},
+			{EInterchangePropertyTracks::CameraPostProcessSettingsFilmGrainTexelSize, {UMovieSceneFloatTrack::StaticClass()->GetName(),  GET_MEMBER_NAME_STRING_CHECKED(UCameraComponent, PostProcessSettings) + FString{ TEXT(".") } + GET_MEMBER_NAME_STRING_CHECKED(FPostProcessSettings, FilmGrainTexelSize), TEXT("Film Grain Texel Size (Post Process Settings)")}},
+			{EInterchangePropertyTracks::CameraPostProcessSettingsFilmShoulder, {UMovieSceneFloatTrack::StaticClass()->GetName(),  GET_MEMBER_NAME_STRING_CHECKED(UCameraComponent, PostProcessSettings) + FString{ TEXT(".") } + GET_MEMBER_NAME_STRING_CHECKED(FPostProcessSettings, FilmShoulder), TEXT("Shoulder (Post Process Settings)")}},
+			{EInterchangePropertyTracks::CameraPostProcessSettingsFilmSlope, {UMovieSceneFloatTrack::StaticClass()->GetName(),  GET_MEMBER_NAME_STRING_CHECKED(UCameraComponent, PostProcessSettings) + FString{ TEXT(".") } + GET_MEMBER_NAME_STRING_CHECKED(FPostProcessSettings, FilmSlope), TEXT("Slope (Post Process Settings)")}},
+			{EInterchangePropertyTracks::CameraPostProcessSettingsFilmToe, {UMovieSceneFloatTrack::StaticClass()->GetName(),  GET_MEMBER_NAME_STRING_CHECKED(UCameraComponent, PostProcessSettings) + FString{ TEXT(".") } + GET_MEMBER_NAME_STRING_CHECKED(FPostProcessSettings, FilmToe), TEXT("Toe (Post Process Settings)")}},
+			{EInterchangePropertyTracks::CameraPostProcessSettingsFilmWhiteClip, {UMovieSceneFloatTrack::StaticClass()->GetName(),  GET_MEMBER_NAME_STRING_CHECKED(UCameraComponent, PostProcessSettings) + FString{ TEXT(".") } + GET_MEMBER_NAME_STRING_CHECKED(FPostProcessSettings, FilmWhiteClip), TEXT("White clip (Post Process Settings)")}},
+			{EInterchangePropertyTracks::CameraPostProcessSettingsHistogramLogMax, {UMovieSceneFloatTrack::StaticClass()->GetName(),  GET_MEMBER_NAME_STRING_CHECKED(UCameraComponent, PostProcessSettings) + FString{ TEXT(".") } + GET_MEMBER_NAME_STRING_CHECKED(FPostProcessSettings, HistogramLogMax), TEXT("Histogram Log Max (Post Process Settings)")}},
+			{EInterchangePropertyTracks::CameraPostProcessSettingsHistogramLogMin, {UMovieSceneFloatTrack::StaticClass()->GetName(),  GET_MEMBER_NAME_STRING_CHECKED(UCameraComponent, PostProcessSettings) + FString{ TEXT(".") } + GET_MEMBER_NAME_STRING_CHECKED(FPostProcessSettings, HistogramLogMin), TEXT("Histogram Log Min (Post Process Settings)")}},
+			{EInterchangePropertyTracks::CameraPostProcessSettingsIndirectLightingColor, {UMovieSceneColorTrack::StaticClass()->GetName(),  GET_MEMBER_NAME_STRING_CHECKED(UCameraComponent, PostProcessSettings) + FString{ TEXT(".") } + GET_MEMBER_NAME_STRING_CHECKED(FPostProcessSettings, IndirectLightingColor), TEXT("Indirect Lighting Color (Post Process Settings)")}},
+			{EInterchangePropertyTracks::CameraPostProcessSettingsIndirectLightingIntensity, {UMovieSceneFloatTrack::StaticClass()->GetName(),  GET_MEMBER_NAME_STRING_CHECKED(UCameraComponent, PostProcessSettings) + FString{ TEXT(".") } + GET_MEMBER_NAME_STRING_CHECKED(FPostProcessSettings, IndirectLightingIntensity), TEXT("Indirect Lighting Intensity (Post Process Settings)")}},
+			{EInterchangePropertyTracks::CameraPostProcessSettingsLensFlareBokehSize, {UMovieSceneFloatTrack::StaticClass()->GetName(),  GET_MEMBER_NAME_STRING_CHECKED(UCameraComponent, PostProcessSettings) + FString{ TEXT(".") } + GET_MEMBER_NAME_STRING_CHECKED(FPostProcessSettings, LensFlareBokehSize), TEXT("BokehSize (Post Process Settings)")}},
+			{EInterchangePropertyTracks::CameraPostProcessSettingsLensFlareIntensity, {UMovieSceneFloatTrack::StaticClass()->GetName(),  GET_MEMBER_NAME_STRING_CHECKED(UCameraComponent, PostProcessSettings) + FString{ TEXT(".") } + GET_MEMBER_NAME_STRING_CHECKED(FPostProcessSettings, LensFlareIntensity), TEXT("Intensity (Post Process Settings)")}},
+			{EInterchangePropertyTracks::CameraPostProcessSettingsLensFlareThreshold, {UMovieSceneFloatTrack::StaticClass()->GetName(),  GET_MEMBER_NAME_STRING_CHECKED(UCameraComponent, PostProcessSettings) + FString{ TEXT(".") } + GET_MEMBER_NAME_STRING_CHECKED(FPostProcessSettings, LensFlareThreshold), TEXT("Threshold (Post Process Settings)")}},
+			{EInterchangePropertyTracks::CameraPostProcessSettingsLensFlareTint, {UMovieSceneColorTrack::StaticClass()->GetName(),  GET_MEMBER_NAME_STRING_CHECKED(UCameraComponent, PostProcessSettings) + FString{ TEXT(".") } + GET_MEMBER_NAME_STRING_CHECKED(FPostProcessSettings, LensFlareTint), TEXT("Tint (Post Process Settings)")}},
+			{EInterchangePropertyTracks::CameraPostProcessSettingsLocalExposureBlurredLuminanceBlend, {UMovieSceneFloatTrack::StaticClass()->GetName(),  GET_MEMBER_NAME_STRING_CHECKED(UCameraComponent, PostProcessSettings) + FString{ TEXT(".") } + GET_MEMBER_NAME_STRING_CHECKED(FPostProcessSettings, LocalExposureBlurredLuminanceBlend), TEXT("Blurred Luminance Blend (Post Process Settings)")}},
+			{EInterchangePropertyTracks::CameraPostProcessSettingsLocalExposureBlurredLuminanceKernelSizePercent, {UMovieSceneFloatTrack::StaticClass()->GetName(),  GET_MEMBER_NAME_STRING_CHECKED(UCameraComponent, PostProcessSettings) + FString{ TEXT(".") } + GET_MEMBER_NAME_STRING_CHECKED(FPostProcessSettings, LocalExposureBlurredLuminanceKernelSizePercent), TEXT("Blurred Luminance Kernel Size Percent (Post Process Settings)")}},
+			{EInterchangePropertyTracks::CameraPostProcessSettingsLocalExposureDetailStrength, {UMovieSceneFloatTrack::StaticClass()->GetName(),  GET_MEMBER_NAME_STRING_CHECKED(UCameraComponent, PostProcessSettings) + FString{ TEXT(".") } + GET_MEMBER_NAME_STRING_CHECKED(FPostProcessSettings, LocalExposureDetailStrength), TEXT("Detail Strength (Post Process Settings)")}},
+			{EInterchangePropertyTracks::CameraPostProcessSettingsLocalExposureHighlightContrastScale, {UMovieSceneFloatTrack::StaticClass()->GetName(),  GET_MEMBER_NAME_STRING_CHECKED(UCameraComponent, PostProcessSettings) + FString{ TEXT(".") } + GET_MEMBER_NAME_STRING_CHECKED(FPostProcessSettings, LocalExposureHighlightContrastScale), TEXT("Highlight Contrast (Post Process Settings)")}},
+			{EInterchangePropertyTracks::CameraPostProcessSettingsLocalExposureHighlightThreshold, {UMovieSceneFloatTrack::StaticClass()->GetName(),  GET_MEMBER_NAME_STRING_CHECKED(UCameraComponent, PostProcessSettings) + FString{ TEXT(".") } + GET_MEMBER_NAME_STRING_CHECKED(FPostProcessSettings, LocalExposureHighlightThreshold), TEXT("Highlight Threshold (Post Process Settings)")}},
+			{EInterchangePropertyTracks::CameraPostProcessSettingsLocalExposureMiddleGreyBias, {UMovieSceneFloatTrack::StaticClass()->GetName(),  GET_MEMBER_NAME_STRING_CHECKED(UCameraComponent, PostProcessSettings) + FString{ TEXT(".") } + GET_MEMBER_NAME_STRING_CHECKED(FPostProcessSettings, LocalExposureMiddleGreyBias), TEXT("Middle Grey Bias (Post Process Settings)")}},
+			{EInterchangePropertyTracks::CameraPostProcessSettingsLocalExposureShadowContrastScale, {UMovieSceneFloatTrack::StaticClass()->GetName(),  GET_MEMBER_NAME_STRING_CHECKED(UCameraComponent, PostProcessSettings) + FString{ TEXT(".") } + GET_MEMBER_NAME_STRING_CHECKED(FPostProcessSettings, LocalExposureShadowContrastScale), TEXT("Shadow Contrast (Post Process Settings)")}},
+			{EInterchangePropertyTracks::CameraPostProcessSettingsLocalExposureShadowThreshold, {UMovieSceneFloatTrack::StaticClass()->GetName(),  GET_MEMBER_NAME_STRING_CHECKED(UCameraComponent, PostProcessSettings) + FString{ TEXT(".") } + GET_MEMBER_NAME_STRING_CHECKED(FPostProcessSettings, LocalExposureShadowThreshold), TEXT("Shadow Threshold (Post Process Settings)")}},
+			{EInterchangePropertyTracks::CameraPostProcessSettingsLumenDiffuseColorBoost, {UMovieSceneFloatTrack::StaticClass()->GetName(),  GET_MEMBER_NAME_STRING_CHECKED(UCameraComponent, PostProcessSettings) + FString{ TEXT(".") } + GET_MEMBER_NAME_STRING_CHECKED(FPostProcessSettings, LumenDiffuseColorBoost), TEXT("Diffuse Color Boost (Post Process Settings)")}},
+			{EInterchangePropertyTracks::CameraPostProcessSettingsLumenFinalGatherLightingUpdateSpeed, {UMovieSceneFloatTrack::StaticClass()->GetName(),  GET_MEMBER_NAME_STRING_CHECKED(UCameraComponent, PostProcessSettings) + FString{ TEXT(".") } + GET_MEMBER_NAME_STRING_CHECKED(FPostProcessSettings, LumenFinalGatherLightingUpdateSpeed), TEXT("Final Gather Lighting Update Speed (Post Process Settings)")}},
+			{EInterchangePropertyTracks::CameraPostProcessSettingsLumenFinalGatherQuality, {UMovieSceneFloatTrack::StaticClass()->GetName(),  GET_MEMBER_NAME_STRING_CHECKED(UCameraComponent, PostProcessSettings) + FString{ TEXT(".") } + GET_MEMBER_NAME_STRING_CHECKED(FPostProcessSettings, LumenFinalGatherQuality), TEXT("Final Gather Quality (Post Process Settings)")}},
+			{EInterchangePropertyTracks::CameraPostProcessSettingsLumenFinalGatherScreenTraces, {UMovieSceneByteTrack::StaticClass()->GetName(),  GET_MEMBER_NAME_STRING_CHECKED(UCameraComponent, PostProcessSettings) + FString{ TEXT(".") } + GET_MEMBER_NAME_STRING_CHECKED(FPostProcessSettings, LumenFinalGatherScreenTraces), TEXT("Screen Traces (Post Process Settings)")}},
+			{EInterchangePropertyTracks::CameraPostProcessSettingsLumenFrontLayerTranslucencyReflections, {UMovieSceneByteTrack::StaticClass()->GetName(),  GET_MEMBER_NAME_STRING_CHECKED(UCameraComponent, PostProcessSettings) + FString{ TEXT(".") } + GET_MEMBER_NAME_STRING_CHECKED(FPostProcessSettings, LumenFrontLayerTranslucencyReflections), TEXT("High Quality Translucency Reflections (Post Process Settings)")}},
+			{EInterchangePropertyTracks::CameraPostProcessSettingsLumenFullSkylightLeakingDistance, {UMovieSceneFloatTrack::StaticClass()->GetName(),  GET_MEMBER_NAME_STRING_CHECKED(UCameraComponent, PostProcessSettings) + FString{ TEXT(".") } + GET_MEMBER_NAME_STRING_CHECKED(FPostProcessSettings, LumenFullSkylightLeakingDistance), TEXT("Full Skylight Leaking Distance (Post Process Settings)")}},
+			{EInterchangePropertyTracks::CameraPostProcessSettingsLumenMaxReflectionBounces, {UMovieSceneIntegerTrack::StaticClass()->GetName(),  GET_MEMBER_NAME_STRING_CHECKED(UCameraComponent, PostProcessSettings) + FString{ TEXT(".") } + GET_MEMBER_NAME_STRING_CHECKED(FPostProcessSettings, LumenMaxReflectionBounces), TEXT("Max Reflection Bounces (Post Process Settings)")}},
+			{EInterchangePropertyTracks::CameraPostProcessSettingsLumenMaxRefractionBounces, {UMovieSceneIntegerTrack::StaticClass()->GetName(),  GET_MEMBER_NAME_STRING_CHECKED(UCameraComponent, PostProcessSettings) + FString{ TEXT(".") } + GET_MEMBER_NAME_STRING_CHECKED(FPostProcessSettings, LumenMaxRefractionBounces), TEXT("Max Refraction Bounces (Post Process Settings)")}},
+			{EInterchangePropertyTracks::CameraPostProcessSettingsLumenMaxRoughnessToTraceReflections, {UMovieSceneFloatTrack::StaticClass()->GetName(),  GET_MEMBER_NAME_STRING_CHECKED(UCameraComponent, PostProcessSettings) + FString{ TEXT(".") } + GET_MEMBER_NAME_STRING_CHECKED(FPostProcessSettings, LumenMaxRoughnessToTraceReflections), TEXT("Max Roughness To Trace (Post Process Settings)")}},
+			{EInterchangePropertyTracks::CameraPostProcessSettingsLumenMaxTraceDistance, {UMovieSceneFloatTrack::StaticClass()->GetName(),  GET_MEMBER_NAME_STRING_CHECKED(UCameraComponent, PostProcessSettings) + FString{ TEXT(".") } + GET_MEMBER_NAME_STRING_CHECKED(FPostProcessSettings, LumenMaxTraceDistance), TEXT("Max Trace Distance (Post Process Settings)")}},
+			{EInterchangePropertyTracks::CameraPostProcessSettingsLumenRayLightingMode, {UMovieSceneEnumTrack::StaticClass()->GetName(),  GET_MEMBER_NAME_STRING_CHECKED(UCameraComponent, PostProcessSettings) + FString{ TEXT(".") } + GET_MEMBER_NAME_STRING_CHECKED(FPostProcessSettings, LumenRayLightingMode), TEXT("Ray Lighting Mode (Post Process Settings)"), StaticEnum<ELumenRayLightingModeOverride>()}},
+			{EInterchangePropertyTracks::CameraPostProcessSettingsLumenReflectionQuality, {UMovieSceneFloatTrack::StaticClass()->GetName(),  GET_MEMBER_NAME_STRING_CHECKED(UCameraComponent, PostProcessSettings) + FString{ TEXT(".") } + GET_MEMBER_NAME_STRING_CHECKED(FPostProcessSettings, LumenReflectionQuality), TEXT("Quality (Post Process Settings)")}},
+			{EInterchangePropertyTracks::CameraPostProcessSettingsLumenReflectionsScreenTraces, {UMovieSceneByteTrack::StaticClass()->GetName(),  GET_MEMBER_NAME_STRING_CHECKED(UCameraComponent, PostProcessSettings) + FString{ TEXT(".") } + GET_MEMBER_NAME_STRING_CHECKED(FPostProcessSettings, LumenReflectionsScreenTraces), TEXT("Screen Traces (Post Process Settings)")}},
+			{EInterchangePropertyTracks::CameraPostProcessSettingsLumenSceneDetail, {UMovieSceneFloatTrack::StaticClass()->GetName(),  GET_MEMBER_NAME_STRING_CHECKED(UCameraComponent, PostProcessSettings) + FString{ TEXT(".") } + GET_MEMBER_NAME_STRING_CHECKED(FPostProcessSettings, LumenSceneDetail), TEXT("Lumen Scene Detail (Post Process Settings)")}},
+			{EInterchangePropertyTracks::CameraPostProcessSettingsLumenSceneLightingQuality, {UMovieSceneFloatTrack::StaticClass()->GetName(),  GET_MEMBER_NAME_STRING_CHECKED(UCameraComponent, PostProcessSettings) + FString{ TEXT(".") } + GET_MEMBER_NAME_STRING_CHECKED(FPostProcessSettings, LumenSceneLightingQuality), TEXT("Lumen Scene Lighting Quality (Post Process Settings)")}},
+			{EInterchangePropertyTracks::CameraPostProcessSettingsLumenSceneLightingUpdateSpeed, {UMovieSceneFloatTrack::StaticClass()->GetName(),  GET_MEMBER_NAME_STRING_CHECKED(UCameraComponent, PostProcessSettings) + FString{ TEXT(".") } + GET_MEMBER_NAME_STRING_CHECKED(FPostProcessSettings, LumenSceneLightingUpdateSpeed), TEXT("Lumen Scene Lighting Update Speed (Post Process Settings)")}},
+			{EInterchangePropertyTracks::CameraPostProcessSettingsLumenSceneViewDistance, {UMovieSceneFloatTrack::StaticClass()->GetName(),  GET_MEMBER_NAME_STRING_CHECKED(UCameraComponent, PostProcessSettings) + FString{ TEXT(".") } + GET_MEMBER_NAME_STRING_CHECKED(FPostProcessSettings, LumenSceneViewDistance), TEXT("Lumen Scene View Distance (Post Process Settings)")}},
+			{EInterchangePropertyTracks::CameraPostProcessSettingsLumenSkylightLeaking, {UMovieSceneFloatTrack::StaticClass()->GetName(),  GET_MEMBER_NAME_STRING_CHECKED(UCameraComponent, PostProcessSettings) + FString{ TEXT(".") } + GET_MEMBER_NAME_STRING_CHECKED(FPostProcessSettings, LumenSkylightLeaking), TEXT("Skylight Leaking (Post Process Settings)")}},
+			{EInterchangePropertyTracks::CameraPostProcessSettingsMotionBlurAmount, {UMovieSceneFloatTrack::StaticClass()->GetName(),  GET_MEMBER_NAME_STRING_CHECKED(UCameraComponent, PostProcessSettings) + FString{ TEXT(".") } + GET_MEMBER_NAME_STRING_CHECKED(FPostProcessSettings, MotionBlurAmount), TEXT("Amount (Post Process Settings)")}},
+			{EInterchangePropertyTracks::CameraPostProcessSettingsMotionBlurMax, {UMovieSceneFloatTrack::StaticClass()->GetName(),  GET_MEMBER_NAME_STRING_CHECKED(UCameraComponent, PostProcessSettings) + FString{ TEXT(".") } + GET_MEMBER_NAME_STRING_CHECKED(FPostProcessSettings, MotionBlurMax), TEXT("Max (Post Process Settings)")}},
+			{EInterchangePropertyTracks::CameraPostProcessSettingsMotionBlurPerObjectSize, {UMovieSceneFloatTrack::StaticClass()->GetName(),  GET_MEMBER_NAME_STRING_CHECKED(UCameraComponent, PostProcessSettings) + FString{ TEXT(".") } + GET_MEMBER_NAME_STRING_CHECKED(FPostProcessSettings, MotionBlurPerObjectSize), TEXT("Per Object Size (Post Process Settings)")}},
+			{EInterchangePropertyTracks::CameraPostProcessSettingsPathTracingMaxBounces, {UMovieSceneIntegerTrack::StaticClass()->GetName(),  GET_MEMBER_NAME_STRING_CHECKED(UCameraComponent, PostProcessSettings) + FString{ TEXT(".") } + GET_MEMBER_NAME_STRING_CHECKED(FPostProcessSettings, PathTracingMaxBounces), TEXT("Max. Bounces (Post Process Settings)")}},
+			{EInterchangePropertyTracks::CameraPostProcessSettingsPathTracingMaxPathExposure, {UMovieSceneFloatTrack::StaticClass()->GetName(),  GET_MEMBER_NAME_STRING_CHECKED(UCameraComponent, PostProcessSettings) + FString{ TEXT(".") } + GET_MEMBER_NAME_STRING_CHECKED(FPostProcessSettings, PathTracingMaxPathExposure), TEXT("Max Path Exposure (Post Process Settings)")}},
+			{EInterchangePropertyTracks::CameraPostProcessSettingsRayTracingAO, {UMovieSceneByteTrack::StaticClass()->GetName(),  GET_MEMBER_NAME_STRING_CHECKED(UCameraComponent, PostProcessSettings) + FString{ TEXT(".") } + GET_MEMBER_NAME_STRING_CHECKED(FPostProcessSettings, RayTracingAO), TEXT("Enabled (Post Process Settings)")}},
+			{EInterchangePropertyTracks::CameraPostProcessSettingsRayTracingAOIntensity, {UMovieSceneFloatTrack::StaticClass()->GetName(),  GET_MEMBER_NAME_STRING_CHECKED(UCameraComponent, PostProcessSettings) + FString{ TEXT(".") } + GET_MEMBER_NAME_STRING_CHECKED(FPostProcessSettings, RayTracingAOIntensity), TEXT("Intensity (Post Process Settings)")}},
+			{EInterchangePropertyTracks::CameraPostProcessSettingsRayTracingAORadius, {UMovieSceneFloatTrack::StaticClass()->GetName(),  GET_MEMBER_NAME_STRING_CHECKED(UCameraComponent, PostProcessSettings) + FString{ TEXT(".") } + GET_MEMBER_NAME_STRING_CHECKED(FPostProcessSettings, RayTracingAORadius), TEXT("Radius (Post Process Settings)")}},
+			{EInterchangePropertyTracks::CameraPostProcessSettingsRayTracingAOSamplesPerPixel, {UMovieSceneIntegerTrack::StaticClass()->GetName(),  GET_MEMBER_NAME_STRING_CHECKED(UCameraComponent, PostProcessSettings) + FString{ TEXT(".") } + GET_MEMBER_NAME_STRING_CHECKED(FPostProcessSettings, RayTracingAOSamplesPerPixel), TEXT("Samples Per Pixel (Post Process Settings)")}},
+			{EInterchangePropertyTracks::CameraPostProcessSettingsRayTracingTranslucencyMaxRoughness, {UMovieSceneFloatTrack::StaticClass()->GetName(),  GET_MEMBER_NAME_STRING_CHECKED(UCameraComponent, PostProcessSettings) + FString{ TEXT(".") } + GET_MEMBER_NAME_STRING_CHECKED(FPostProcessSettings, RayTracingTranslucencyMaxRoughness), TEXT("Max Roughness (Post Process Settings)")}},
+			{EInterchangePropertyTracks::CameraPostProcessSettingsRayTracingTranslucencyRefraction, {UMovieSceneByteTrack::StaticClass()->GetName(),  GET_MEMBER_NAME_STRING_CHECKED(UCameraComponent, PostProcessSettings) + FString{ TEXT(".") } + GET_MEMBER_NAME_STRING_CHECKED(FPostProcessSettings, RayTracingTranslucencyRefraction), TEXT("Refraction (Post Process Settings)")}},
+			{EInterchangePropertyTracks::CameraPostProcessSettingsRayTracingTranslucencyRefractionRays, {UMovieSceneIntegerTrack::StaticClass()->GetName(),  GET_MEMBER_NAME_STRING_CHECKED(UCameraComponent, PostProcessSettings) + FString{ TEXT(".") } + GET_MEMBER_NAME_STRING_CHECKED(FPostProcessSettings, RayTracingTranslucencyRefractionRays), TEXT("Max. Refraction Rays (Post Process Settings)")}},
+			{EInterchangePropertyTracks::CameraPostProcessSettingsRayTracingTranslucencySamplesPerPixel, {UMovieSceneIntegerTrack::StaticClass()->GetName(),  GET_MEMBER_NAME_STRING_CHECKED(UCameraComponent, PostProcessSettings) + FString{ TEXT(".") } + GET_MEMBER_NAME_STRING_CHECKED(FPostProcessSettings, RayTracingTranslucencySamplesPerPixel), TEXT("Samples Per Pixel (Post Process Settings)")}},
+			{EInterchangePropertyTracks::CameraPostProcessSettingsRayTracingTranslucencyShadows, {UMovieSceneEnumTrack::StaticClass()->GetName(),  GET_MEMBER_NAME_STRING_CHECKED(UCameraComponent, PostProcessSettings) + FString{ TEXT(".") } + GET_MEMBER_NAME_STRING_CHECKED(FPostProcessSettings, RayTracingTranslucencyShadows), TEXT("Shadows (Post Process Settings)"), StaticEnum<EReflectedAndRefractedRayTracedShadows>()}},
+			{EInterchangePropertyTracks::CameraPostProcessSettingsReflectionMethod, {UMovieSceneEnumTrack::StaticClass()->GetName(),  GET_MEMBER_NAME_STRING_CHECKED(UCameraComponent, PostProcessSettings) + FString{ TEXT(".") } + GET_MEMBER_NAME_STRING_CHECKED(FPostProcessSettings, ReflectionMethod), TEXT("Method (Post Process Settings)"), StaticEnum<EReflectionMethod::Type>()}},
+			{EInterchangePropertyTracks::CameraPostProcessSettingsSceneColorTint, {UMovieSceneColorTrack::StaticClass()->GetName(),  GET_MEMBER_NAME_STRING_CHECKED(UCameraComponent, PostProcessSettings) + FString{ TEXT(".") } + GET_MEMBER_NAME_STRING_CHECKED(FPostProcessSettings, SceneColorTint), TEXT("Scene Color Tint (Post Process Settings)")}},
+			{EInterchangePropertyTracks::CameraPostProcessSettingsSceneFringeIntensity, {UMovieSceneFloatTrack::StaticClass()->GetName(),  GET_MEMBER_NAME_STRING_CHECKED(UCameraComponent, PostProcessSettings) + FString{ TEXT(".") } + GET_MEMBER_NAME_STRING_CHECKED(FPostProcessSettings, SceneFringeIntensity), TEXT("Intensity (Post Process Settings)")}},
+			{EInterchangePropertyTracks::CameraPostProcessSettingsScreenSpaceReflectionIntensity, {UMovieSceneFloatTrack::StaticClass()->GetName(),  GET_MEMBER_NAME_STRING_CHECKED(UCameraComponent, PostProcessSettings) + FString{ TEXT(".") } + GET_MEMBER_NAME_STRING_CHECKED(FPostProcessSettings, ScreenSpaceReflectionIntensity), TEXT("Intensity (Post Process Settings)")}},
+			{EInterchangePropertyTracks::CameraPostProcessSettingsScreenSpaceReflectionMaxRoughness, {UMovieSceneFloatTrack::StaticClass()->GetName(),  GET_MEMBER_NAME_STRING_CHECKED(UCameraComponent, PostProcessSettings) + FString{ TEXT(".") } + GET_MEMBER_NAME_STRING_CHECKED(FPostProcessSettings, ScreenSpaceReflectionMaxRoughness), TEXT("Max Roughness (Post Process Settings)")}},
+			{EInterchangePropertyTracks::CameraPostProcessSettingsScreenSpaceReflectionQuality, {UMovieSceneFloatTrack::StaticClass()->GetName(),  GET_MEMBER_NAME_STRING_CHECKED(UCameraComponent, PostProcessSettings) + FString{ TEXT(".") } + GET_MEMBER_NAME_STRING_CHECKED(FPostProcessSettings, ScreenSpaceReflectionQuality), TEXT("Quality (Post Process Settings)")}},
+			{EInterchangePropertyTracks::CameraPostProcessSettingsSharpen, {UMovieSceneFloatTrack::StaticClass()->GetName(),  GET_MEMBER_NAME_STRING_CHECKED(UCameraComponent, PostProcessSettings) + FString{ TEXT(".") } + GET_MEMBER_NAME_STRING_CHECKED(FPostProcessSettings, Sharpen), TEXT("Sharpen (Post Process Settings)")}},
+			{EInterchangePropertyTracks::CameraPostProcessSettingsTemperatureType, {UMovieSceneEnumTrack::StaticClass()->GetName(),  GET_MEMBER_NAME_STRING_CHECKED(UCameraComponent, PostProcessSettings) + FString{ TEXT(".") } + GET_MEMBER_NAME_STRING_CHECKED(FPostProcessSettings, TemperatureType), TEXT("Temperature Type (Post Process Settings)"), StaticEnum<ETemperatureMethod>()}},
+			{EInterchangePropertyTracks::CameraPostProcessSettingsToneCurveAmount, {UMovieSceneFloatTrack::StaticClass()->GetName(),  GET_MEMBER_NAME_STRING_CHECKED(UCameraComponent, PostProcessSettings) + FString{ TEXT(".") } + GET_MEMBER_NAME_STRING_CHECKED(FPostProcessSettings, ToneCurveAmount), TEXT("Tone Curve Amount (Post Process Settings)")}},
+			{EInterchangePropertyTracks::CameraPostProcessSettingsTranslucencyType, {UMovieSceneEnumTrack::StaticClass()->GetName(),  GET_MEMBER_NAME_STRING_CHECKED(UCameraComponent, PostProcessSettings) + FString{ TEXT(".") } + GET_MEMBER_NAME_STRING_CHECKED(FPostProcessSettings, TranslucencyType), TEXT("Type (Post Process Settings)"), StaticEnum<ETranslucencyType>()}},
+			{EInterchangePropertyTracks::CameraPostProcessSettingsVignetteIntensity, {UMovieSceneFloatTrack::StaticClass()->GetName(),  GET_MEMBER_NAME_STRING_CHECKED(UCameraComponent, PostProcessSettings) + FString{ TEXT(".") } + GET_MEMBER_NAME_STRING_CHECKED(FPostProcessSettings, VignetteIntensity), TEXT("Vignette Intensity (Post Process Settings)")}},
+			{EInterchangePropertyTracks::CameraPostProcessSettingsWhiteTemp, {UMovieSceneFloatTrack::StaticClass()->GetName(),  GET_MEMBER_NAME_STRING_CHECKED(UCameraComponent, PostProcessSettings) + FString{ TEXT(".") } + GET_MEMBER_NAME_STRING_CHECKED(FPostProcessSettings, WhiteTemp), TEXT("Temp (Post Process Settings)")}},
+			{EInterchangePropertyTracks::CameraPostProcessSettingsWhiteTint, {UMovieSceneFloatTrack::StaticClass()->GetName(),  GET_MEMBER_NAME_STRING_CHECKED(UCameraComponent, PostProcessSettings) + FString{ TEXT(".") } + GET_MEMBER_NAME_STRING_CHECKED(FPostProcessSettings, WhiteTint), TEXT("Tint (Post Process Settings)")}},
+			{EInterchangePropertyTracks::CameraProjectionMode, {UMovieSceneEnumTrack::StaticClass()->GetName(),  GET_MEMBER_NAME_STRING_CHECKED(UCameraComponent, ProjectionMode) , TEXT("Projection Mode"), StaticEnum<ECameraProjectionMode::Type>()}},
+			{EInterchangePropertyTracks::CameraUseFieldOfViewForLOD, {UMovieSceneBoolTrack::StaticClass()->GetName(),  GET_MEMBER_NAME_STRING_CHECKED(UCameraComponent, bUseFieldOfViewForLOD) , TEXT("Use Field Of View for LOD")}},
 		}
 	{}
 
diff --git a/Engine/Plugins/Interchange/Runtime/Source/Nodes/Public/InterchangeAnimationDefinitions.h b/Engine/Plugins/Interchange/Runtime/Source/Nodes/Public/InterchangeAnimationDefinitions.h
--- a/Engine/Plugins/Interchange/Runtime/Source/Nodes/Public/InterchangeAnimationDefinitions.h
+++ b/Engine/Plugins/Interchange/Runtime/Source/Nodes/Public/InterchangeAnimationDefinitions.h
@@ -11,20 +11,121 @@ UENUM()
 enum class EInterchangePropertyTracks : int32
 {
 	/** Common Properties*/
-	Visibility,
+	AffectDistanceFieldLighting,
+	AffectDynamicIndirectLighting,
+	AffectIndirectLightingWhileHidden,
 	AutoActivate,
+	BodyInstanceAngularDamping,
+	BodyInstancebEnableGravity,
+	BodyInstancebNotifyRigidBodyCollision,
+	BodyInstancebSimulatePhysics,
+	BodyInstancebUpdateKinematicFromSimulation,
+	BodyInstancebUseCCD,
+	BodyInstanceLinearDamping,
+	BodyInstanceMassScale,
+	BoundsScale,
+	CastContactShadow,
+	CastHiddenShadow,
+	CastInsetShadow,
+	CastShadow,
+	CustomDepthStencilValue,
+	CustomDepthStencilWriteMask,
+	DrawDebug,
+	EmissiveLightSource,
+	ExcludeFromLightAttachmentGroup,
+	HiddenInGame,
+	HiddenInSceneCapture,
+	Holdout,
+	LightAttachmentsAsGroup,
+	Mobility,
+	OnlyOwnerSee,
+	OwnerNoSee,
+	ReceivesDecals,
+	RenderCustomDepth,
+	RenderInDepthPass,
+	RenderInMainPass,
+	SingleSampleShadowFromStationaryLights,
+	TranslucencySortDistanceOffset,
+	VisibleInRayTracing,
+	VisibleInSceneCaptureOnly,
+	Visibility,
 
 	/** Light Properties */
+	LightAffectGlobalIllumination,
+	LightAffectReflection,
+	LightAffectTranslucentLighting,
+	LightAtmosphereSunDiskColorScale,
+	LightAtmosphereSunLight,
+	LightAtmosphereSunLightIndex,
+	LightAttenuationRadius,
+	LightBarnDoorAngle,
+	LightBarnDoorLength,
+	LightBloomMaxBrightness,
+	LightBloomScale,
+	LightBloomThreshold,
+	LightBloomTint,
+	LightCascadeDistributionExponent,
+	LightCascadeTransitionFraction,
+	LightCastDeepShadow,
+	LightCastShadows,
+	LightCastVolumetricShadow,
+	LightCloudAmbientOcclusionStrength,
+	LightCloudScatteredLuminanceScale,
+	LightCloudShadowOnAtmosphereStrength,
+	LightCloudShadowOnSurfaceStrength,
+	LightCloudShadowStrength,
 	LightColor,
+	LightDynamicShadowCascades,
+	LightDynamicShadowDistanceMovableLight,
+	LightDynamicShadowDistanceStationaryLight,
+	LightEnableLightShaftBloom,
+	LightEnableLightShaftOcclusion,
+	LightFalloffExponent,
+	LightForceCachedShadowsForMovablePrimitives,
+	LightForwardShadingPriority,
+	LightFunctionFadeDistance,
+	LightIESBrightnessScale,
+	LightIndirectLightingIntensity,
+	LightInnerConeAngle,
 	LightIntensity,
 	LightIntensityUnits,
+	LightInverseExposureBlend,
+	LightLowerHemisphereColor,
+	LightmassSettingsLightSourceAngle,
+	LightMinOcclusion,
+	LightModulatedShadowColor,
+	LightOcclusionDepthRange,
+	LightOcclusionExponent,
+	LightOcclusionMaskDarkness,
+	LightOcclusionTint,
+	LightOuterConeAngle,
+	LightSamplesPerPixel,
+	LightShadowAmount,
+	LightShadowBias,
+	LightShadowCascadeBiasDistribution,
+	LightShadowDistanceFadeoutFraction,
+	LightShadowSlopeBias,
+	LightShadowSourceAngleFactor,
+	LightSoftSourceRadius,
+	LightSourceAngle,
+	LightSourceCubemapAngle,
+	LightSourceHeight,
+	LightSourceLength,
+	LightSourceRadius,
+	LightSourceSoftAngle,
+	LightSourceWidth,
+	LightSpecularScale,
 	LightTemperature,
+	LightTransmission,
+	LightUseIESBrightness,
+	LightUseInverseSquaredFalloff,
 	LightUseTemperature,
+	LightVolumetricScatteringIntensity,
 
 	/** Camera Properties*/
 	CameraAspectRatio,
-	CameraAutoCalculateOrthoPlanes,
 	CameraAspectRatioAxisConstraint,
+	CameraAutoCalculateOrthoPlanes,
 	CameraConstrainAspectRatio,
 	CameraCurrentAperture,
 	CameraCurrentFocalLength,
@@ -33,14 +134,156 @@ enum class EInterchangePropertyTracks : int32
 	CameraFilmbackSensorAspectRatio,
 	CameraFilmbackSensorHeight,
 	CameraFilmbackSensorWidth,
+	CameraFocusSettingsFocusOffset,
 	CameraFocusSettingsManualFocusDistance,
-	CameraMobility,
 	CameraOrthoFarClipPlane,
 	CameraOrthoNearClipPlane,
 	CameraOrthoWidth,
+	CameraOverrideAspectRatioAxisConstraint,
 	CameraPostProcessBlendWeight,
+	CameraPostProcessSettingsAmbientCubemapIntensity,
+	CameraPostProcessSettingsAmbientCubemapTint,
+	CameraPostProcessSettingsAmbientOcclusionBias,
+	CameraPostProcessSettingsAmbientOcclusionFadeDistance,
+	CameraPostProcessSettingsAmbientOcclusionFadeRadius,
+	CameraPostProcessSettingsAmbientOcclusionIntensity,
+	CameraPostProcessSettingsAmbientOcclusionMipBlend,
+	CameraPostProcessSettingsAmbientOcclusionMipScale,
+	CameraPostProcessSettingsAmbientOcclusionMipThreshold,
+	CameraPostProcessSettingsAmbientOcclusionPower,
+	CameraPostProcessSettingsAmbientOcclusionQuality,
+	CameraPostProcessSettingsAmbientOcclusionRadius,
+	CameraPostProcessSettingsAmbientOcclusionStaticFraction,
+	CameraPostProcessSettingsAmbientOcclusionTemporalBlendWeight,
+	CameraPostProcessSettingsAutoExposureBias,
+	CameraPostProcessSettingsAutoExposureHighPercent,
+	CameraPostProcessSettingsAutoExposureLowPercent,
+	CameraPostProcessSettingsAutoExposureMaxBrightness,
+	CameraPostProcessSettingsAutoExposureMinBrightness,
+	CameraPostProcessSettingsAutoExposureSpeedDown,
+	CameraPostProcessSettingsAutoExposureSpeedUp,
+	CameraPostProcessSettingsBloom1Size,
+	CameraPostProcessSettingsBloom1Tint,
+	CameraPostProcessSettingsBloom2Size,
+	CameraPostProcessSettingsBloom2Tint,
+	CameraPostProcessSettingsBloom3Size,
+	CameraPostProcessSettingsBloom3Tint,
+	CameraPostProcessSettingsBloom4Size,
+	CameraPostProcessSettingsBloom4Tint,
+	CameraPostProcessSettingsBloom5Size,
+	CameraPostProcessSettingsBloom5Tint,
+	CameraPostProcessSettingsBloom6Size,
+	CameraPostProcessSettingsBloom6Tint,
+	CameraPostProcessSettingsBloomConvolutionBufferScale,
+	CameraPostProcessSettingsBloomConvolutionPreFilterMax,
+	CameraPostProcessSettingsBloomConvolutionPreFilterMin,
+	CameraPostProcessSettingsBloomConvolutionPreFilterMult,
+	CameraPostProcessSettingsBloomConvolutionScatterDispersion,
+	CameraPostProcessSettingsBloomConvolutionSize,
+	CameraPostProcessSettingsBloomDirtMaskIntensity,
+	CameraPostProcessSettingsBloomDirtMaskTint,
+	CameraPostProcessSettingsBloomIntensity,
+	CameraPostProcessSettingsBloomSizeScale,
+	CameraPostProcessSettingsBloomThreshold,
+	CameraPostProcessSettingsBlueCorrection,
+	CameraPostProcessSettingsChromaticAberrationStartOffset,
+	CameraPostProcessSettingsColorCorrectionHighlightsMax,
+	CameraPostProcessSettingsColorCorrectionHighlightsMin,
+	CameraPostProcessSettingsColorCorrectionShadowsMax,
+	CameraPostProcessSettingsColorGradingIntensity,
+	CameraPostProcessSettingsDepthOfFieldBladeCount,
+	CameraPostProcessSettingsDepthOfFieldDepthBlurAmount,
+	CameraPostProcessSettingsDepthOfFieldDepthBlurRadius,
+	CameraPostProcessSettingsDepthOfFieldFarBlurSize,
+	CameraPostProcessSettingsDepthOfFieldFarTransitionRegion,
+	CameraPostProcessSettingsDepthOfFieldFocalDistance,
+	CameraPostProcessSettingsDepthOfFieldFocalRegion,
+	CameraPostProcessSettingsDepthOfFieldFstop,
+	CameraPostProcessSettingsDepthOfFieldMinFstop,
+	CameraPostProcessSettingsDepthOfFieldNearBlurSize,
+	CameraPostProcessSettingsDepthOfFieldNearTransitionRegion,
+	CameraPostProcessSettingsDepthOfFieldOcclusion,
+	CameraPostProcessSettingsDepthOfFieldScale,
+	CameraPostProcessSettingsDepthOfFieldSkyFocusDistance,
+	CameraPostProcessSettingsDepthOfFieldUseHairDepth,
+	CameraPostProcessSettingsDepthOfFieldVignetteSize,
+	CameraPostProcessSettingsDynamicGlobalIlluminationMethod,
+	CameraPostProcessSettingsExpandGamut,
+	CameraPostProcessSettingsFilmBlackClip,
+	CameraPostProcessSettingsFilmGrainHighlightsMax,
+	CameraPostProcessSettingsFilmGrainHighlightsMin,
+	CameraPostProcessSettingsFilmGrainIntensity,
+	CameraPostProcessSettingsFilmGrainIntensityHighlights,
+	CameraPostProcessSettingsFilmGrainIntensityMidtones,
+	CameraPostProcessSettingsFilmGrainIntensityShadows,
+	CameraPostProcessSettingsFilmGrainShadowsMax,
+	CameraPostProcessSettingsFilmGrainTexelSize,
+	CameraPostProcessSettingsFilmShoulder,
+	CameraPostProcessSettingsFilmSlope,
+	CameraPostProcessSettingsFilmToe,
+	CameraPostProcessSettingsFilmWhiteClip,
+	CameraPostProcessSettingsHistogramLogMax,
+	CameraPostProcessSettingsHistogramLogMin,
+	CameraPostProcessSettingsIndirectLightingColor,
+	CameraPostProcessSettingsIndirectLightingIntensity,
+	CameraPostProcessSettingsLensFlareBokehSize,
+	CameraPostProcessSettingsLensFlareIntensity,
+	CameraPostProcessSettingsLensFlareThreshold,
+	CameraPostProcessSettingsLensFlareTint,
+	CameraPostProcessSettingsLocalExposureBlurredLuminanceBlend,
+	CameraPostProcessSettingsLocalExposureBlurredLuminanceKernelSizePercent,
+	CameraPostProcessSettingsLocalExposureDetailStrength,
+	CameraPostProcessSettingsLocalExposureHighlightContrastScale,
+	CameraPostProcessSettingsLocalExposureHighlightThreshold,
+	CameraPostProcessSettingsLocalExposureMiddleGreyBias,
+	CameraPostProcessSettingsLocalExposureShadowContrastScale,
+	CameraPostProcessSettingsLocalExposureShadowThreshold,
+	CameraPostProcessSettingsLumenDiffuseColorBoost,
+	CameraPostProcessSettingsLumenFinalGatherLightingUpdateSpeed,
+	CameraPostProcessSettingsLumenFinalGatherQuality,
+	CameraPostProcessSettingsLumenFinalGatherScreenTraces,
+	CameraPostProcessSettingsLumenFrontLayerTranslucencyReflections,
+	CameraPostProcessSettingsLumenFullSkylightLeakingDistance,
+	CameraPostProcessSettingsLumenMaxReflectionBounces,
+	CameraPostProcessSettingsLumenMaxRefractionBounces,
+	CameraPostProcessSettingsLumenMaxRoughnessToTraceReflections,
+	CameraPostProcessSettingsLumenMaxTraceDistance,
+	CameraPostProcessSettingsLumenRayLightingMode,
+	CameraPostProcessSettingsLumenReflectionQuality,
+	CameraPostProcessSettingsLumenReflectionsScreenTraces,
+	CameraPostProcessSettingsLumenSceneDetail,
+	CameraPostProcessSettingsLumenSceneLightingQuality,
+	CameraPostProcessSettingsLumenSceneLightingUpdateSpeed,
+	CameraPostProcessSettingsLumenSceneViewDistance,
+	CameraPostProcessSettingsLumenSkylightLeaking,
+	CameraPostProcessSettingsMotionBlurAmount,
+	CameraPostProcessSettingsMotionBlurMax,
+	CameraPostProcessSettingsMotionBlurPerObjectSize,
+	CameraPostProcessSettingsPathTracingMaxBounces,
+	CameraPostProcessSettingsPathTracingMaxPathExposure,
+	CameraPostProcessSettingsRayTracingAO,
+	CameraPostProcessSettingsRayTracingAOIntensity,
+	CameraPostProcessSettingsRayTracingAORadius,
+	CameraPostProcessSettingsRayTracingAOSamplesPerPixel,
+	CameraPostProcessSettingsRayTracingTranslucencyMaxRoughness,
+	CameraPostProcessSettingsRayTracingTranslucencyRefraction,
+	CameraPostProcessSettingsRayTracingTranslucencyRefractionRays,
+	CameraPostProcessSettingsRayTracingTranslucencySamplesPerPixel,
+	CameraPostProcessSettingsRayTracingTranslucencyShadows,
+	CameraPostProcessSettingsReflectionMethod,
+	CameraPostProcessSettingsSceneColorTint,
+	CameraPostProcessSettingsSceneFringeIntensity,
+	CameraPostProcessSettingsScreenSpaceReflectionIntensity,
+	CameraPostProcessSettingsScreenSpaceReflectionMaxRoughness,
+	CameraPostProcessSettingsScreenSpaceReflectionQuality,
+	CameraPostProcessSettingsSharpen,
+	CameraPostProcessSettingsTemperatureType,
+	CameraPostProcessSettingsToneCurveAmount,
+	CameraPostProcessSettingsTranslucencyType,
+	CameraPostProcessSettingsVignetteIntensity,
+	CameraPostProcessSettingsWhiteTemp,
+	CameraPostProcessSettingsWhiteTint,
 	CameraProjectionMode,
-	CameraShouldUpdatePhysicsVolume,
 	CameraUseFieldOfViewForLOD,
 
 	None = -1 UMETA(hidden),
