Interchange: Fix smooth group importing, the fbx sdk have a bug when triangulating edge hardness so we have to convert them to smooth group before doing the triangulation.

#rnx

diff --git a/Engine/Plugins/Interchange/Runtime/Source/Parsers/Fbx/Private/FbxMesh.cpp b/Engine/Plugins/Interchange/Runtime/Source/Parsers/Fbx/Private/FbxMesh.cpp
--- a/Engine/Plugins/Interchange/Runtime/Source/Parsers/Fbx/Private/FbxMesh.cpp
+++ b/Engine/Plugins/Interchange/Runtime/Source/Parsers/Fbx/Private/FbxMesh.cpp
@@ -636,12 +636,12 @@ bool FMeshDescriptionImporter::FillMeshDescriptionFromFbxMesh(FbxMesh* Mesh, con
 		Message->Text = LOCTEXT("MissingSmoothGroup", "No smoothing group information was found for this mesh '{MeshName}' in the FBX file. Please make sure to enable the 'Export Smoothing Groups' option in the FBX Exporter before exporting the file.");
 	}
 
-	for (int32 i = 0; i < LayerSmoothingCount; i++)
+	for (int32 SmoothingLayerIndex = 0; SmoothingLayerIndex < LayerSmoothingCount; SmoothingLayerIndex++)
 	{
-		FbxLayerElementSmoothing const* SmoothingInfo = Mesh->GetLayer(i)->GetSmoothing();
-		if (SmoothingInfo && SmoothingInfo->GetMappingMode() != FbxLayerElement::eByPolygon)
+		FbxLayerElementSmoothing const* SmoothingInfoTmp = Mesh->GetLayer(SmoothingLayerIndex)->GetSmoothing();
+		if (SmoothingInfoTmp && SmoothingInfoTmp->GetMappingMode() != FbxLayerElement::eByPolygon)
 		{
-			SDKGeometryConverter->ComputePolygonSmoothingFromEdgeSmoothing(Mesh, i);
+			SDKGeometryConverter->ComputePolygonSmoothingFromEdgeSmoothing(Mesh, SmoothingLayerIndex);
 		}
 	}
 
@@ -682,7 +682,15 @@ bool FMeshDescriptionImporter::FillMeshDescriptionFromFbxMesh(FbxMesh* Mesh, con
 		if (SmoothingInfo->GetMappingMode() == FbxLayerElement::eByPolygon)
 		{
 			//Convert the base layer to edge smoothing
-			SDKGeometryConverter->ComputeEdgeSmoothingFromPolygonSmoothing(Mesh, 0);
+			for (int32 SmoothingLayerIndex = 0; SmoothingLayerIndex < LayerSmoothingCount; SmoothingLayerIndex++)
+			{
+				FbxLayerElementSmoothing const* SmoothingInfoTmp = Mesh->GetLayer(SmoothingLayerIndex)->GetSmoothing();
+				if (SmoothingInfoTmp && SmoothingInfoTmp->GetMappingMode() != FbxLayerElement::eByEdge)
+				{
+					SDKGeometryConverter->ComputeEdgeSmoothingFromPolygonSmoothing(Mesh, SmoothingLayerIndex);
+				}
+			}
+			
 			BaseLayer = Mesh->GetLayer(0);
 			SmoothingInfo = BaseLayer->GetSmoothing();
 		}
@@ -989,6 +997,7 @@ bool FMeshDescriptionImporter::FillMeshDescriptionFromFbxMesh(FbxMesh* Mesh, con
 			TArray<FVertexID> CornerVerticesIDs;
 			TArray<FVector3f, TInlineAllocator<3>> P;
 
+			bool bUnsupportedSmoothingGroupErrorDisplayed = false;
 			bool bCorruptedMsgDone = false;
 			//Polygons
 			for (int32 PolygonIndex = 0; PolygonIndex < PolygonCount; PolygonIndex++)
@@ -1264,24 +1273,21 @@ bool FMeshDescriptionImporter::FillMeshDescriptionFromFbxMesh(FbxMesh* Mesh, con
 
 						if (!EdgeHardnesses[MatchEdgeId])
 						{
-							if (bSmoothingAvailable && SmoothingInfo)
+							if (bSmoothingAvailable && SmoothingInfo && SmoothingMappingMode == FbxLayerElement::eByEdge)
 							{
-								if (SmoothingMappingMode == FbxLayerElement::eByEdge)
-								{
-									int32 lSmoothingIndex = (SmoothingReferenceMode == FbxLayerElement::eDirect) ? EdgeIndex : SmoothingInfo->GetIndexArray().GetAt(EdgeIndex);
-									//Set the hard edges
-									int32 SmoothingFlag = SmoothingInfo->GetDirectArray().GetAt(lSmoothingIndex);
-									EdgeHardnesses[MatchEdgeId] = (SmoothingFlag == 0);
-								}
-								else
-								{
-									EdgeHardnesses[MatchEdgeId] = false;
-									//TODO add an error log
-									//AddTokenizedErrorMessage(FTokenizedMessage::Create(EMessageSeverity::Warning, FText::Format(LOCTEXT("Error_UnsupportedSmoothingGroup", "Unsupported Smoothing group mapping mode on mesh  '{0}'"), FText::FromString(Mesh->GetName()))), FFbxErrors::Generic_Mesh_UnsupportingSmoothingGroup);
-								}
+								int32 lSmoothingIndex = (SmoothingReferenceMode == FbxLayerElement::eDirect) ? EdgeIndex : SmoothingInfo->GetIndexArray().GetAt(EdgeIndex);
+								//Set the hard edges
+								int32 SmoothingFlag = SmoothingInfo->GetDirectArray().GetAt(lSmoothingIndex);
+								EdgeHardnesses[MatchEdgeId] = (SmoothingFlag == 0);
 							}
 							else
 							{
+								if (!bUnsupportedSmoothingGroupErrorDisplayed && SmoothingMappingMode != FbxLayerElement::eByEdge)
+								{
+									bUnsupportedSmoothingGroupErrorDisplayed = true;
+									UInterchangeResultMeshError_Generic* Message = AddMessage<UInterchangeResultMeshError_Generic>(Mesh);
+									Message->Text = LOCTEXT("Error_UnsupportedSmoothingGroup", "Unsupported Smoothing group mapping mode on mesh '{MeshName}'.");
+								}
 								//When there is no smoothing group we set all edge to: hard (faceted mesh) for static mesh and smooth for skinned and rigid
 								EdgeHardnesses[MatchEdgeId] = MeshType == EMeshType::Static ? !bStaticMeshUseSmoothEdgesIfSmoothingInformationIsMissing : false;
 							}
@@ -1728,6 +1734,16 @@ void FFbxMesh::AddAllMeshes(FbxScene* SDKScene, FbxGeometryConverter* SDKGeometr
 		}
 		for (FbxMesh* ToTriangulateMesh : ToTriangulateMeshes)
 		{
+			// Must do this before triangulating the mesh due to an FBX bug in Triangulate. Edge hardnees triangulation give wrong edge hardness so we compute them to smooth group during the triangulation.
+			int32 LayerSmoothingCount = ToTriangulateMesh->GetLayerCount(FbxLayerElement::eSmoothing);
+			for (int32 SmoothingLayerIndex = 0; SmoothingLayerIndex < LayerSmoothingCount; SmoothingLayerIndex++)
+			{
+				FbxLayerElementSmoothing const* SmoothingInfoTmp = ToTriangulateMesh->GetLayer(SmoothingLayerIndex)->GetSmoothing();
+				if (SmoothingInfoTmp && SmoothingInfoTmp->GetMappingMode() != FbxLayerElement::eByPolygon)
+				{
+					SDKGeometryConverter->ComputePolygonSmoothingFromEdgeSmoothing(ToTriangulateMesh, SmoothingLayerIndex);
+				}
+			}
 			const bool bReplace = true;
 			SDKGeometryConverter->Triangulate(ToTriangulateMesh, bReplace);
 		}
