Interchange Level Sequence:
- Add support of Vector tracks

diff --git a/Engine/Plugins/Interchange/Runtime/Source/Import/Private/Animation/InterchangeLevelSequenceHelper.cpp b/Engine/Plugins/Interchange/Runtime/Source/Import/Private/Animation/InterchangeLevelSequenceHelper.cpp
--- a/Engine/Plugins/Interchange/Runtime/Source/Import/Private/Animation/InterchangeLevelSequenceHelper.cpp
+++ b/Engine/Plugins/Interchange/Runtime/Source/Import/Private/Animation/InterchangeLevelSequenceHelper.cpp
@@ -18,6 +18,7 @@
 #include "Tracks/MovieSceneBoolTrack.h"
 #include "Tracks/MovieSceneByteTrack.h"
 #include "Tracks/MovieSceneColorTrack.h"
+#include "Tracks/MovieSceneVectorTrack.h"
 #include "Tracks/MovieSceneEnumTrack.h"
 #include "Tracks/MovieSceneFloatTrack.h"
 #include "Tracks/MovieSceneIntegerTrack.h"
@@ -28,12 +29,10 @@ namespace UE::Interchange::Private
 	FInterchangePropertyTracksHelper::FInterchangePropertyTracksHelper()
 		: PropertyTracks{
 			//Common
-			{EInterchangePropertyTracks::Visibility, {UMovieSceneVisibilityTrack::StaticClass()->GetName(), AActor::GetHiddenPropertyName().ToString(), TEXT("Visibility")}},
-
-			{EInterchangePropertyTracks::AutoActivate, {UMovieSceneBoolTrack::StaticClass()->GetName(), GET_MEMBER_NAME_STRING_CHECKED(UActorComponent, bAutoActivate), TEXT("Auto Activate")}},
 			{EInterchangePropertyTracks::AffectDistanceFieldLighting, {UMovieSceneBoolTrack::StaticClass()->GetName(),  GET_MEMBER_NAME_STRING_CHECKED(UPrimitiveComponent, bAffectDistanceFieldLighting) , TEXT("Affect Distance Field Lighting")}},
 			{EInterchangePropertyTracks::AffectDynamicIndirectLighting, {UMovieSceneBoolTrack::StaticClass()->GetName(),  GET_MEMBER_NAME_STRING_CHECKED(UPrimitiveComponent, bAffectDynamicIndirectLighting) , TEXT("Affect Dynamic Indirect Lighting")}},
 			{EInterchangePropertyTracks::AffectIndirectLightingWhileHidden, {UMovieSceneBoolTrack::StaticClass()->GetName(),  GET_MEMBER_NAME_STRING_CHECKED(UPrimitiveComponent, bAffectIndirectLightingWhileHidden) , TEXT("Affect Indirect Lighting While Hidden")}},
+			{EInterchangePropertyTracks::AutoActivate, {UMovieSceneBoolTrack::StaticClass()->GetName(),  GET_MEMBER_NAME_STRING_CHECKED(UActorComponent, bAutoActivate) , TEXT("Auto Activate")}},
 			{EInterchangePropertyTracks::BodyInstanceAngularDamping, {UMovieSceneFloatTrack::StaticClass()->GetName(),  GET_MEMBER_NAME_STRING_CHECKED(UPrimitiveComponent, BodyInstance) + FString{ TEXT(".") } + GET_MEMBER_NAME_STRING_CHECKED(FBodyInstance, AngularDamping), TEXT("Angular Damping (Body Instance)")}},
 			{EInterchangePropertyTracks::BodyInstancebEnableGravity, {UMovieSceneByteTrack::StaticClass()->GetName(),  GET_MEMBER_NAME_STRING_CHECKED(UPrimitiveComponent, BodyInstance) + FString{ TEXT(".") } + GET_MEMBER_NAME_STRING_CHECKED(FBodyInstance, bEnableGravity), TEXT("Enable Gravity (Body Instance)")}},
 			{EInterchangePropertyTracks::BodyInstancebNotifyRigidBodyCollision, {UMovieSceneByteTrack::StaticClass()->GetName(),  GET_MEMBER_NAME_STRING_CHECKED(UPrimitiveComponent, BodyInstance) + FString{ TEXT(".") } + GET_MEMBER_NAME_STRING_CHECKED(FBodyInstance, bNotifyRigidBodyCollision), TEXT("Simulation Generates Hit Events (Body Instance)")}},
@@ -49,6 +48,7 @@ namespace UE::Interchange::Private
 			{EInterchangePropertyTracks::CastShadow, {UMovieSceneByteTrack::StaticClass()->GetName(),  GET_MEMBER_NAME_STRING_CHECKED(UPrimitiveComponent, CastShadow) , TEXT("Cast Shadow")}},
 			{EInterchangePropertyTracks::CustomDepthStencilValue, {UMovieSceneIntegerTrack::StaticClass()->GetName(),  GET_MEMBER_NAME_STRING_CHECKED(UPrimitiveComponent, CustomDepthStencilValue) , TEXT("CustomDepth Stencil Value")}},
 			{EInterchangePropertyTracks::CustomDepthStencilWriteMask, {UMovieSceneEnumTrack::StaticClass()->GetName(),  GET_MEMBER_NAME_STRING_CHECKED(UPrimitiveComponent, CustomDepthStencilWriteMask) , TEXT("Custom Depth Stencil Write Mask"), StaticEnum<ERendererStencilMask>()}},
+			{EInterchangePropertyTracks::DefaultUpVector, {UMovieSceneDoubleVectorTrack::StaticClass()->GetName(),  GET_MEMBER_NAME_STRING_CHECKED(USplineComponent, DefaultUpVector) , TEXT("Default Up Vector"), 3}},
 			{EInterchangePropertyTracks::DrawDebug, {UMovieSceneBoolTrack::StaticClass()->GetName(),  GET_MEMBER_NAME_STRING_CHECKED(USplineComponent, bDrawDebug) , TEXT("Draw Debug")}},
 			{EInterchangePropertyTracks::EmissiveLightSource, {UMovieSceneBoolTrack::StaticClass()->GetName(),  GET_MEMBER_NAME_STRING_CHECKED(UPrimitiveComponent, bEmissiveLightSource) , TEXT("Emissive Light Source")}},
 			{EInterchangePropertyTracks::ExcludeFromLightAttachmentGroup, {UMovieSceneBoolTrack::StaticClass()->GetName(),  GET_MEMBER_NAME_STRING_CHECKED(UPrimitiveComponent, bExcludeFromLightAttachmentGroup) , TEXT("Exclude from Light Attachment Group")}},
@@ -66,6 +66,7 @@ namespace UE::Interchange::Private
 			{EInterchangePropertyTracks::TranslucencySortDistanceOffset, {UMovieSceneFloatTrack::StaticClass()->GetName(),  GET_MEMBER_NAME_STRING_CHECKED(UPrimitiveComponent, TranslucencySortDistanceOffset) , TEXT("Translucency Sort Distance Offset")}},
 			{EInterchangePropertyTracks::VisibleInRayTracing, {UMovieSceneBoolTrack::StaticClass()->GetName(),  GET_MEMBER_NAME_STRING_CHECKED(UPrimitiveComponent, bVisibleInRayTracing) , TEXT("Visible in Ray Tracing")}},
 			{EInterchangePropertyTracks::VisibleInSceneCaptureOnly, {UMovieSceneBoolTrack::StaticClass()->GetName(),  GET_MEMBER_NAME_STRING_CHECKED(UPrimitiveComponent, bVisibleInSceneCaptureOnly) , TEXT("Visible In Scene Capture Only")}},
+			{EInterchangePropertyTracks::Visibility, {UMovieSceneVisibilityTrack::StaticClass()->GetName(), AActor::GetHiddenPropertyName().ToString(), TEXT("Visibility")}},
 
 			// Light
 			{EInterchangePropertyTracks::LightAffectGlobalIllumination, {UMovieSceneBoolTrack::StaticClass()->GetName(),  GET_MEMBER_NAME_STRING_CHECKED(ULightComponentBase, bAffectGlobalIllumination) , TEXT("Affect Ray Tracing Global Illumination")}},
@@ -101,6 +102,7 @@ namespace UE::Interchange::Private
 			{EInterchangePropertyTracks::LightForceCachedShadowsForMovablePrimitives, {UMovieSceneBoolTrack::StaticClass()->GetName(),  GET_MEMBER_NAME_STRING_CHECKED(ULightComponent, bForceCachedShadowsForMovablePrimitives) , TEXT("Force Cached Shadows for Movable Primitives")}},
 			{EInterchangePropertyTracks::LightForwardShadingPriority, {UMovieSceneIntegerTrack::StaticClass()->GetName(),  GET_MEMBER_NAME_STRING_CHECKED(UDirectionalLightComponent, ForwardShadingPriority) , TEXT("Forward Shading Priority")}},
 			{EInterchangePropertyTracks::LightFunctionFadeDistance, {UMovieSceneFloatTrack::StaticClass()->GetName(),  GET_MEMBER_NAME_STRING_CHECKED(ULightComponent, LightFunctionFadeDistance) , TEXT("Fade Distance")}},
+			{EInterchangePropertyTracks::LightFunctionScale, {UMovieSceneDoubleVectorTrack::StaticClass()->GetName(),  GET_MEMBER_NAME_STRING_CHECKED(ULightComponent, LightFunctionScale) , TEXT("Light Function Scale"), 3}},
 			{EInterchangePropertyTracks::LightIESBrightnessScale, {UMovieSceneFloatTrack::StaticClass()->GetName(),  GET_MEMBER_NAME_STRING_CHECKED(ULightComponent, IESBrightnessScale) , TEXT("IES Intensity Scale")}},
 			{EInterchangePropertyTracks::LightIndirectLightingIntensity, {UMovieSceneFloatTrack::StaticClass()->GetName(),  GET_MEMBER_NAME_STRING_CHECKED(ULightComponentBase, IndirectLightingIntensity) , TEXT("Indirect Lighting Intensity")}},
 			{EInterchangePropertyTracks::LightInnerConeAngle, {UMovieSceneFloatTrack::StaticClass()->GetName(),  GET_MEMBER_NAME_STRING_CHECKED(USpotLightComponent, InnerConeAngle) , TEXT("Inner Cone Angle")}},
@@ -123,6 +125,7 @@ namespace UE::Interchange::Private
 			{EInterchangePropertyTracks::LightShadowDistanceFadeoutFraction, {UMovieSceneFloatTrack::StaticClass()->GetName(),  GET_MEMBER_NAME_STRING_CHECKED(UDirectionalLightComponent, ShadowDistanceFadeoutFraction) , TEXT("Distance Fadeout Fraction")}},
 			{EInterchangePropertyTracks::LightShadowSlopeBias, {UMovieSceneFloatTrack::StaticClass()->GetName(),  GET_MEMBER_NAME_STRING_CHECKED(ULightComponent, ShadowSlopeBias) , TEXT("Shadow Slope Bias")}},
 			{EInterchangePropertyTracks::LightShadowSourceAngleFactor, {UMovieSceneFloatTrack::StaticClass()->GetName(),  GET_MEMBER_NAME_STRING_CHECKED(UDirectionalLightComponent, ShadowSourceAngleFactor) , TEXT("Shadow Source Angle Factor")}},
+			{EInterchangePropertyTracks::LightShaftOverrideDirection, {UMovieSceneDoubleVectorTrack::StaticClass()->GetName(),  GET_MEMBER_NAME_STRING_CHECKED(UDirectionalLightComponent, LightShaftOverrideDirection) , TEXT("Light Shaft Override Direction"), 3}},
 			{EInterchangePropertyTracks::LightSoftSourceRadius, {UMovieSceneFloatTrack::StaticClass()->GetName(),  GET_MEMBER_NAME_STRING_CHECKED(UPointLightComponent, SoftSourceRadius) , TEXT("Soft Source Radius")}},
 			{EInterchangePropertyTracks::LightSourceAngle, {UMovieSceneFloatTrack::StaticClass()->GetName(),  GET_MEMBER_NAME_STRING_CHECKED(UDirectionalLightComponent, LightSourceAngle) , TEXT("Source Angle")}},
 			{EInterchangePropertyTracks::LightSourceCubemapAngle, {UMovieSceneFloatTrack::StaticClass()->GetName(),  GET_MEMBER_NAME_STRING_CHECKED(USkyLightComponent, SourceCubemapAngle) , TEXT("Source Cubemap Angle")}},
@@ -143,6 +146,7 @@ namespace UE::Interchange::Private
 			{EInterchangePropertyTracks::CameraAspectRatio, {UMovieSceneFloatTrack::StaticClass()->GetName(),  GET_MEMBER_NAME_STRING_CHECKED(UCameraComponent, AspectRatio) , TEXT("Aspect Ratio")}},
 			{EInterchangePropertyTracks::CameraAspectRatioAxisConstraint, {UMovieSceneEnumTrack::StaticClass()->GetName(),  GET_MEMBER_NAME_STRING_CHECKED(UCameraComponent, AspectRatioAxisConstraint) , TEXT("Aspect Ratio Axis Constraint"), StaticEnum<EAspectRatioAxisConstraint>()}},
 			{EInterchangePropertyTracks::CameraAutoCalculateOrthoPlanes, {UMovieSceneBoolTrack::StaticClass()->GetName(),  GET_MEMBER_NAME_STRING_CHECKED(UCameraComponent, bAutoCalculateOrthoPlanes) , TEXT("Auto Calculate Ortho Planes")}},
+			{EInterchangePropertyTracks::CameraAutoPlaneShift, {UMovieSceneFloatTrack::StaticClass()->GetName(),  GET_MEMBER_NAME_STRING_CHECKED(UCameraComponent, AutoPlaneShift) , TEXT("Auto Plane Shift")}},
 			{EInterchangePropertyTracks::CameraConstrainAspectRatio, {UMovieSceneBoolTrack::StaticClass()->GetName(),  GET_MEMBER_NAME_STRING_CHECKED(UCameraComponent, bConstrainAspectRatio) , TEXT("Constrain Aspect Ratio")}},
 			{EInterchangePropertyTracks::CameraCurrentAperture, {UMovieSceneFloatTrack::StaticClass()->GetName(),  GET_MEMBER_NAME_STRING_CHECKED(UCineCameraComponent, CurrentAperture) , TEXT("Current Aperture")}},
 			{EInterchangePropertyTracks::CameraCurrentFocalLength, {UMovieSceneFloatTrack::StaticClass()->GetName(),  GET_MEMBER_NAME_STRING_CHECKED(UCineCameraComponent, CurrentFocalLength) , TEXT("Current Focal Length")}},
@@ -153,6 +157,7 @@ namespace UE::Interchange::Private
 			{EInterchangePropertyTracks::CameraFilmbackSensorWidth, {UMovieSceneFloatTrack::StaticClass()->GetName(),  GET_MEMBER_NAME_STRING_CHECKED(UCineCameraComponent, Filmback) + FString{ TEXT(".") } + GET_MEMBER_NAME_STRING_CHECKED(FCameraFilmbackSettings, SensorWidth), TEXT("Sensor Width (Filmback)")}},
 			{EInterchangePropertyTracks::CameraFocusSettingsFocusOffset, {UMovieSceneFloatTrack::StaticClass()->GetName(),  GET_MEMBER_NAME_STRING_CHECKED(UCineCameraComponent, FocusSettings) + FString{ TEXT(".") } + GET_MEMBER_NAME_STRING_CHECKED(FCameraFocusSettings, FocusOffset), TEXT("Focus Offset (Focus Settings)")}},
 			{EInterchangePropertyTracks::CameraFocusSettingsManualFocusDistance, {UMovieSceneFloatTrack::StaticClass()->GetName(),  GET_MEMBER_NAME_STRING_CHECKED(UCineCameraComponent, FocusSettings) + FString{ TEXT(".") } + GET_MEMBER_NAME_STRING_CHECKED(FCameraFocusSettings, ManualFocusDistance), TEXT("Manual Focus Distance (Focus Settings)")}},
+			{EInterchangePropertyTracks::CameraFocusSettingsTrackingFocusSettingsRelativeOffset, {UMovieSceneDoubleVectorTrack::StaticClass()->GetName(),  GET_MEMBER_NAME_STRING_CHECKED(UCineCameraComponent, FocusSettings) + FString{ TEXT(".") } + GET_MEMBER_NAME_STRING_CHECKED(FCameraFocusSettings, TrackingFocusSettings)+ FString{ TEXT(".") } + GET_MEMBER_NAME_STRING_CHECKED(FCameraTrackingFocusSettings, RelativeOffset), TEXT("Relative Offset (Tracking Focus Settings)"), 3}},
 			{EInterchangePropertyTracks::CameraOrthoFarClipPlane, {UMovieSceneFloatTrack::StaticClass()->GetName(),  GET_MEMBER_NAME_STRING_CHECKED(UCameraComponent, OrthoFarClipPlane) , TEXT("Ortho Far Clip Plane")}},
 			{EInterchangePropertyTracks::CameraOrthoNearClipPlane, {UMovieSceneFloatTrack::StaticClass()->GetName(),  GET_MEMBER_NAME_STRING_CHECKED(UCameraComponent, OrthoNearClipPlane) , TEXT("Ortho Near Clip Plane")}},
 			{EInterchangePropertyTracks::CameraOrthoWidth, {UMovieSceneFloatTrack::StaticClass()->GetName(),  GET_MEMBER_NAME_STRING_CHECKED(UCameraComponent, OrthoWidth) , TEXT("Ortho Width")}},
@@ -192,6 +197,7 @@ namespace UE::Interchange::Private
 			{EInterchangePropertyTracks::CameraPostProcessSettingsBloom6Size, {UMovieSceneFloatTrack::StaticClass()->GetName(),  GET_MEMBER_NAME_STRING_CHECKED(UCameraComponent, PostProcessSettings) + FString{ TEXT(".") } + GET_MEMBER_NAME_STRING_CHECKED(FPostProcessSettings, Bloom6Size), TEXT("#6 Size (Post Process Settings)")}},
 			{EInterchangePropertyTracks::CameraPostProcessSettingsBloom6Tint, {UMovieSceneColorTrack::StaticClass()->GetName(),  GET_MEMBER_NAME_STRING_CHECKED(UCameraComponent, PostProcessSettings) + FString{ TEXT(".") } + GET_MEMBER_NAME_STRING_CHECKED(FPostProcessSettings, Bloom6Tint), TEXT("#6 Tint (Post Process Settings)")}},
 			{EInterchangePropertyTracks::CameraPostProcessSettingsBloomConvolutionBufferScale, {UMovieSceneFloatTrack::StaticClass()->GetName(),  GET_MEMBER_NAME_STRING_CHECKED(UCameraComponent, PostProcessSettings) + FString{ TEXT(".") } + GET_MEMBER_NAME_STRING_CHECKED(FPostProcessSettings, BloomConvolutionBufferScale), TEXT("Convolution Buffer (Post Process Settings)")}},
+			{EInterchangePropertyTracks::CameraPostProcessSettingsBloomConvolutionCenterUV, {UMovieSceneDoubleVectorTrack::StaticClass()->GetName(),  GET_MEMBER_NAME_STRING_CHECKED(UCameraComponent, PostProcessSettings) + FString{ TEXT(".") } + GET_MEMBER_NAME_STRING_CHECKED(FPostProcessSettings, BloomConvolutionCenterUV), TEXT("Convolution Center (Post Process Settings)"), 2}},
 			{EInterchangePropertyTracks::CameraPostProcessSettingsBloomConvolutionPreFilterMax, {UMovieSceneFloatTrack::StaticClass()->GetName(),  GET_MEMBER_NAME_STRING_CHECKED(UCameraComponent, PostProcessSettings) + FString{ TEXT(".") } + GET_MEMBER_NAME_STRING_CHECKED(FPostProcessSettings, BloomConvolutionPreFilterMax), TEXT("Convolution Boost Max (Post Process Settings)")}},
 			{EInterchangePropertyTracks::CameraPostProcessSettingsBloomConvolutionPreFilterMin, {UMovieSceneFloatTrack::StaticClass()->GetName(),  GET_MEMBER_NAME_STRING_CHECKED(UCameraComponent, PostProcessSettings) + FString{ TEXT(".") } + GET_MEMBER_NAME_STRING_CHECKED(FPostProcessSettings, BloomConvolutionPreFilterMin), TEXT("Convolution Boost Min (Post Process Settings)")}},
 			{EInterchangePropertyTracks::CameraPostProcessSettingsBloomConvolutionPreFilterMult, {UMovieSceneFloatTrack::StaticClass()->GetName(),  GET_MEMBER_NAME_STRING_CHECKED(UCameraComponent, PostProcessSettings) + FString{ TEXT(".") } + GET_MEMBER_NAME_STRING_CHECKED(FPostProcessSettings, BloomConvolutionPreFilterMult), TEXT("Convolution Boost Mult (Post Process Settings)")}},
@@ -204,10 +210,30 @@ namespace UE::Interchange::Private
 			{EInterchangePropertyTracks::CameraPostProcessSettingsBloomThreshold, {UMovieSceneFloatTrack::StaticClass()->GetName(),  GET_MEMBER_NAME_STRING_CHECKED(UCameraComponent, PostProcessSettings) + FString{ TEXT(".") } + GET_MEMBER_NAME_STRING_CHECKED(FPostProcessSettings, BloomThreshold), TEXT("Threshold (Post Process Settings)")}},
 			{EInterchangePropertyTracks::CameraPostProcessSettingsBlueCorrection, {UMovieSceneFloatTrack::StaticClass()->GetName(),  GET_MEMBER_NAME_STRING_CHECKED(UCameraComponent, PostProcessSettings) + FString{ TEXT(".") } + GET_MEMBER_NAME_STRING_CHECKED(FPostProcessSettings, BlueCorrection), TEXT("Blue Correction (Post Process Settings)")}},
 			{EInterchangePropertyTracks::CameraPostProcessSettingsChromaticAberrationStartOffset, {UMovieSceneFloatTrack::StaticClass()->GetName(),  GET_MEMBER_NAME_STRING_CHECKED(UCameraComponent, PostProcessSettings) + FString{ TEXT(".") } + GET_MEMBER_NAME_STRING_CHECKED(FPostProcessSettings, ChromaticAberrationStartOffset), TEXT("Start Offset (Post Process Settings)")}},
+			{EInterchangePropertyTracks::CameraPostProcessSettingsColorContrast, {UMovieSceneDoubleVectorTrack::StaticClass()->GetName(),  GET_MEMBER_NAME_STRING_CHECKED(UCameraComponent, PostProcessSettings) + FString{ TEXT(".") } + GET_MEMBER_NAME_STRING_CHECKED(FPostProcessSettings, ColorContrast), TEXT("Contrast (Post Process Settings)"), 4}},
+			{EInterchangePropertyTracks::CameraPostProcessSettingsColorContrastHighlights, {UMovieSceneDoubleVectorTrack::StaticClass()->GetName(),  GET_MEMBER_NAME_STRING_CHECKED(UCameraComponent, PostProcessSettings) + FString{ TEXT(".") } + GET_MEMBER_NAME_STRING_CHECKED(FPostProcessSettings, ColorContrastHighlights), TEXT("Contrast (Post Process Settings)"), 4}},
+			{EInterchangePropertyTracks::CameraPostProcessSettingsColorContrastMidtones, {UMovieSceneDoubleVectorTrack::StaticClass()->GetName(),  GET_MEMBER_NAME_STRING_CHECKED(UCameraComponent, PostProcessSettings) + FString{ TEXT(".") } + GET_MEMBER_NAME_STRING_CHECKED(FPostProcessSettings, ColorContrastMidtones), TEXT("Contrast (Post Process Settings)"), 4}},
+			{EInterchangePropertyTracks::CameraPostProcessSettingsColorContrastShadows, {UMovieSceneDoubleVectorTrack::StaticClass()->GetName(),  GET_MEMBER_NAME_STRING_CHECKED(UCameraComponent, PostProcessSettings) + FString{ TEXT(".") } + GET_MEMBER_NAME_STRING_CHECKED(FPostProcessSettings, ColorContrastShadows), TEXT("Contrast (Post Process Settings)"), 4}},
 			{EInterchangePropertyTracks::CameraPostProcessSettingsColorCorrectionHighlightsMax, {UMovieSceneFloatTrack::StaticClass()->GetName(),  GET_MEMBER_NAME_STRING_CHECKED(UCameraComponent, PostProcessSettings) + FString{ TEXT(".") } + GET_MEMBER_NAME_STRING_CHECKED(FPostProcessSettings, ColorCorrectionHighlightsMax), TEXT("HighlightsMax (Post Process Settings)")}},
 			{EInterchangePropertyTracks::CameraPostProcessSettingsColorCorrectionHighlightsMin, {UMovieSceneFloatTrack::StaticClass()->GetName(),  GET_MEMBER_NAME_STRING_CHECKED(UCameraComponent, PostProcessSettings) + FString{ TEXT(".") } + GET_MEMBER_NAME_STRING_CHECKED(FPostProcessSettings, ColorCorrectionHighlightsMin), TEXT("HighlightsMin (Post Process Settings)")}},
 			{EInterchangePropertyTracks::CameraPostProcessSettingsColorCorrectionShadowsMax, {UMovieSceneFloatTrack::StaticClass()->GetName(),  GET_MEMBER_NAME_STRING_CHECKED(UCameraComponent, PostProcessSettings) + FString{ TEXT(".") } + GET_MEMBER_NAME_STRING_CHECKED(FPostProcessSettings, ColorCorrectionShadowsMax), TEXT("ShadowsMax (Post Process Settings)")}},
+			{EInterchangePropertyTracks::CameraPostProcessSettingsColorGain, {UMovieSceneDoubleVectorTrack::StaticClass()->GetName(),  GET_MEMBER_NAME_STRING_CHECKED(UCameraComponent, PostProcessSettings) + FString{ TEXT(".") } + GET_MEMBER_NAME_STRING_CHECKED(FPostProcessSettings, ColorGain), TEXT("Gain (Post Process Settings)"), 4}},
+			{EInterchangePropertyTracks::CameraPostProcessSettingsColorGainHighlights, {UMovieSceneDoubleVectorTrack::StaticClass()->GetName(),  GET_MEMBER_NAME_STRING_CHECKED(UCameraComponent, PostProcessSettings) + FString{ TEXT(".") } + GET_MEMBER_NAME_STRING_CHECKED(FPostProcessSettings, ColorGainHighlights), TEXT("Gain (Post Process Settings)"), 4}},
+			{EInterchangePropertyTracks::CameraPostProcessSettingsColorGainMidtones, {UMovieSceneDoubleVectorTrack::StaticClass()->GetName(),  GET_MEMBER_NAME_STRING_CHECKED(UCameraComponent, PostProcessSettings) + FString{ TEXT(".") } + GET_MEMBER_NAME_STRING_CHECKED(FPostProcessSettings, ColorGainMidtones), TEXT("Gain (Post Process Settings)"), 4}},
+			{EInterchangePropertyTracks::CameraPostProcessSettingsColorGainShadows, {UMovieSceneDoubleVectorTrack::StaticClass()->GetName(),  GET_MEMBER_NAME_STRING_CHECKED(UCameraComponent, PostProcessSettings) + FString{ TEXT(".") } + GET_MEMBER_NAME_STRING_CHECKED(FPostProcessSettings, ColorGainShadows), TEXT("Gain (Post Process Settings)"), 4}},
+			{EInterchangePropertyTracks::CameraPostProcessSettingsColorGamma, {UMovieSceneDoubleVectorTrack::StaticClass()->GetName(),  GET_MEMBER_NAME_STRING_CHECKED(UCameraComponent, PostProcessSettings) + FString{ TEXT(".") } + GET_MEMBER_NAME_STRING_CHECKED(FPostProcessSettings, ColorGamma), TEXT("Gamma (Post Process Settings)"), 4}},
+			{EInterchangePropertyTracks::CameraPostProcessSettingsColorGammaHighlights, {UMovieSceneDoubleVectorTrack::StaticClass()->GetName(),  GET_MEMBER_NAME_STRING_CHECKED(UCameraComponent, PostProcessSettings) + FString{ TEXT(".") } + GET_MEMBER_NAME_STRING_CHECKED(FPostProcessSettings, ColorGammaHighlights), TEXT("Gamma (Post Process Settings)"), 4}},
+			{EInterchangePropertyTracks::CameraPostProcessSettingsColorGammaMidtones, {UMovieSceneDoubleVectorTrack::StaticClass()->GetName(),  GET_MEMBER_NAME_STRING_CHECKED(UCameraComponent, PostProcessSettings) + FString{ TEXT(".") } + GET_MEMBER_NAME_STRING_CHECKED(FPostProcessSettings, ColorGammaMidtones), TEXT("Gamma (Post Process Settings)"), 4}},
+			{EInterchangePropertyTracks::CameraPostProcessSettingsColorGammaShadows, {UMovieSceneDoubleVectorTrack::StaticClass()->GetName(),  GET_MEMBER_NAME_STRING_CHECKED(UCameraComponent, PostProcessSettings) + FString{ TEXT(".") } + GET_MEMBER_NAME_STRING_CHECKED(FPostProcessSettings, ColorGammaShadows), TEXT("Gamma (Post Process Settings)"), 4}},
 			{EInterchangePropertyTracks::CameraPostProcessSettingsColorGradingIntensity, {UMovieSceneFloatTrack::StaticClass()->GetName(),  GET_MEMBER_NAME_STRING_CHECKED(UCameraComponent, PostProcessSettings) + FString{ TEXT(".") } + GET_MEMBER_NAME_STRING_CHECKED(FPostProcessSettings, ColorGradingIntensity), TEXT("Color Grading LUT Intensity (Post Process Settings)")}},
+			{EInterchangePropertyTracks::CameraPostProcessSettingsColorOffset, {UMovieSceneDoubleVectorTrack::StaticClass()->GetName(),  GET_MEMBER_NAME_STRING_CHECKED(UCameraComponent, PostProcessSettings) + FString{ TEXT(".") } + GET_MEMBER_NAME_STRING_CHECKED(FPostProcessSettings, ColorOffset), TEXT("Offset (Post Process Settings)"), 4}},
+			{EInterchangePropertyTracks::CameraPostProcessSettingsColorOffsetHighlights, {UMovieSceneDoubleVectorTrack::StaticClass()->GetName(),  GET_MEMBER_NAME_STRING_CHECKED(UCameraComponent, PostProcessSettings) + FString{ TEXT(".") } + GET_MEMBER_NAME_STRING_CHECKED(FPostProcessSettings, ColorOffsetHighlights), TEXT("Offset (Post Process Settings)"), 4}},
+			{EInterchangePropertyTracks::CameraPostProcessSettingsColorOffsetMidtones, {UMovieSceneDoubleVectorTrack::StaticClass()->GetName(),  GET_MEMBER_NAME_STRING_CHECKED(UCameraComponent, PostProcessSettings) + FString{ TEXT(".") } + GET_MEMBER_NAME_STRING_CHECKED(FPostProcessSettings, ColorOffsetMidtones), TEXT("Offset (Post Process Settings)"), 4}},
+			{EInterchangePropertyTracks::CameraPostProcessSettingsColorOffsetShadows, {UMovieSceneDoubleVectorTrack::StaticClass()->GetName(),  GET_MEMBER_NAME_STRING_CHECKED(UCameraComponent, PostProcessSettings) + FString{ TEXT(".") } + GET_MEMBER_NAME_STRING_CHECKED(FPostProcessSettings, ColorOffsetShadows), TEXT("Offset (Post Process Settings)"), 4}},
+			{EInterchangePropertyTracks::CameraPostProcessSettingsColorSaturation, {UMovieSceneDoubleVectorTrack::StaticClass()->GetName(),  GET_MEMBER_NAME_STRING_CHECKED(UCameraComponent, PostProcessSettings) + FString{ TEXT(".") } + GET_MEMBER_NAME_STRING_CHECKED(FPostProcessSettings, ColorSaturation), TEXT("Saturation (Post Process Settings)"), 4}},
+			{EInterchangePropertyTracks::CameraPostProcessSettingsColorSaturationHighlights, {UMovieSceneDoubleVectorTrack::StaticClass()->GetName(),  GET_MEMBER_NAME_STRING_CHECKED(UCameraComponent, PostProcessSettings) + FString{ TEXT(".") } + GET_MEMBER_NAME_STRING_CHECKED(FPostProcessSettings, ColorSaturationHighlights), TEXT("Saturation (Post Process Settings)"), 4}},
+			{EInterchangePropertyTracks::CameraPostProcessSettingsColorSaturationMidtones, {UMovieSceneDoubleVectorTrack::StaticClass()->GetName(),  GET_MEMBER_NAME_STRING_CHECKED(UCameraComponent, PostProcessSettings) + FString{ TEXT(".") } + GET_MEMBER_NAME_STRING_CHECKED(FPostProcessSettings, ColorSaturationMidtones), TEXT("Saturation (Post Process Settings)"), 4}},
+			{EInterchangePropertyTracks::CameraPostProcessSettingsColorSaturationShadows, {UMovieSceneDoubleVectorTrack::StaticClass()->GetName(),  GET_MEMBER_NAME_STRING_CHECKED(UCameraComponent, PostProcessSettings) + FString{ TEXT(".") } + GET_MEMBER_NAME_STRING_CHECKED(FPostProcessSettings, ColorSaturationShadows), TEXT("Saturation (Post Process Settings)"), 4}},
 			{EInterchangePropertyTracks::CameraPostProcessSettingsDepthOfFieldBladeCount, {UMovieSceneIntegerTrack::StaticClass()->GetName(),  GET_MEMBER_NAME_STRING_CHECKED(UCameraComponent, PostProcessSettings) + FString{ TEXT(".") } + GET_MEMBER_NAME_STRING_CHECKED(FPostProcessSettings, DepthOfFieldBladeCount), TEXT("Number of diaphragm blades (Post Process Settings)")}},
 			{EInterchangePropertyTracks::CameraPostProcessSettingsDepthOfFieldDepthBlurAmount, {UMovieSceneFloatTrack::StaticClass()->GetName(),  GET_MEMBER_NAME_STRING_CHECKED(UCameraComponent, PostProcessSettings) + FString{ TEXT(".") } + GET_MEMBER_NAME_STRING_CHECKED(FPostProcessSettings, DepthOfFieldDepthBlurAmount), TEXT("Depth Blur km for 50% (Post Process Settings)")}},
 			{EInterchangePropertyTracks::CameraPostProcessSettingsDepthOfFieldDepthBlurRadius, {UMovieSceneFloatTrack::StaticClass()->GetName(),  GET_MEMBER_NAME_STRING_CHECKED(UCameraComponent, PostProcessSettings) + FString{ TEXT(".") } + GET_MEMBER_NAME_STRING_CHECKED(FPostProcessSettings, DepthOfFieldDepthBlurRadius), TEXT("Depth Blur Radius (Post Process Settings)")}},
@@ -301,6 +327,8 @@ namespace UE::Interchange::Private
 			{EInterchangePropertyTracks::CameraPostProcessSettingsWhiteTemp, {UMovieSceneFloatTrack::StaticClass()->GetName(),  GET_MEMBER_NAME_STRING_CHECKED(UCameraComponent, PostProcessSettings) + FString{ TEXT(".") } + GET_MEMBER_NAME_STRING_CHECKED(FPostProcessSettings, WhiteTemp), TEXT("Temp (Post Process Settings)")}},
 			{EInterchangePropertyTracks::CameraPostProcessSettingsWhiteTint, {UMovieSceneFloatTrack::StaticClass()->GetName(),  GET_MEMBER_NAME_STRING_CHECKED(UCameraComponent, PostProcessSettings) + FString{ TEXT(".") } + GET_MEMBER_NAME_STRING_CHECKED(FPostProcessSettings, WhiteTint), TEXT("Tint (Post Process Settings)")}},
 			{EInterchangePropertyTracks::CameraProjectionMode, {UMovieSceneEnumTrack::StaticClass()->GetName(),  GET_MEMBER_NAME_STRING_CHECKED(UCameraComponent, ProjectionMode) , TEXT("Projection Mode"), StaticEnum<ECameraProjectionMode::Type>()}},
+			{EInterchangePropertyTracks::CameraUpdateOrthoPlanes, {UMovieSceneBoolTrack::StaticClass()->GetName(),  GET_MEMBER_NAME_STRING_CHECKED(UCameraComponent, bUpdateOrthoPlanes) , TEXT("Update Ortho Planes")}},
+			{EInterchangePropertyTracks::CameraUseCameraHeightAsViewTarget, {UMovieSceneBoolTrack::StaticClass()->GetName(),  GET_MEMBER_NAME_STRING_CHECKED(UCameraComponent, bUseCameraHeightAsViewTarget) , TEXT("Use Camera Height as View Target")}},
 			{EInterchangePropertyTracks::CameraUseFieldOfViewForLOD, {UMovieSceneBoolTrack::StaticClass()->GetName(),  GET_MEMBER_NAME_STRING_CHECKED(UCameraComponent, bUseFieldOfViewForLOD) , TEXT("Use Field Of View for LOD")}},
 		}
 	{}
@@ -348,6 +376,18 @@ namespace UE::Interchange::Private
 			}
 		}
 
+		if(InterchangePropertyTrack->NumChannelsUsed.IsSet())
+		{
+			if(UMovieSceneDoubleVectorTrack* DoubleVectorTrack = Cast<UMovieSceneDoubleVectorTrack>(PropertyTrack))
+			{
+				DoubleVectorTrack->SetNumChannelsUsed(*InterchangePropertyTrack->NumChannelsUsed);
+			}
+			else if(UMovieSceneFloatVectorTrack* FloatVectorTrack = Cast<UMovieSceneFloatVectorTrack>(PropertyTrack))
+			{
+				FloatVectorTrack->SetNumChannelsUsed(*InterchangePropertyTrack->NumChannelsUsed);
+			}
+		}
+
 		bool bSectionAdded = false;
 
 		UMovieSceneSection* Section = PropertyTrack->FindOrAddSection(0, bSectionAdded);
diff --git a/Engine/Plugins/Interchange/Runtime/Source/Import/Private/Animation/InterchangeLevelSequenceHelper.h b/Engine/Plugins/Interchange/Runtime/Source/Import/Private/Animation/InterchangeLevelSequenceHelper.h
--- a/Engine/Plugins/Interchange/Runtime/Source/Import/Private/Animation/InterchangeLevelSequenceHelper.h
+++ b/Engine/Plugins/Interchange/Runtime/Source/Import/Private/Animation/InterchangeLevelSequenceHelper.h
@@ -57,10 +57,31 @@ namespace UE::Interchange::Private
 
 		struct FInterchangeProperty
 		{
+			FInterchangeProperty(FString&& ClassType, FString&& Path, FString&& Name)
+				: ClassType{ MoveTemp(ClassType) }
+				, Path{ MoveTemp(Path) }
+				, Name{ MoveTemp(Name) }
+			{}
+
+			FInterchangeProperty(FString&& ClassType, FString&& Path, FString&& Name, UEnum* EnumClass)
+				: ClassType{ MoveTemp(ClassType) }
+				, Path{ MoveTemp(Path) }
+				, Name{ MoveTemp(Name) }
+				, EnumClass{EnumClass}
+			{}
+
+			FInterchangeProperty(FString&& ClassType, FString&& Path, FString&& Name, int32 NumChannelsUsed)
+				: ClassType{ std::move(ClassType) }
+				, Path{ MoveTemp(Path) }
+				, Name{ MoveTemp(Name) }
+				, NumChannelsUsed{ NumChannelsUsed }
+			{}
+
 			FString ClassType; // Float, Double, Byte, etc. Basically the class name of the UMovieSceneTrack
 			FString Path;
 			FName Name;
 			UEnum* EnumClass = nullptr; // Only used for Enum property tracks
+			TOptional<int32> NumChannelsUsed; // Only used for Vector property tracks (Vector tracks can only have 2-4 channels)
 		};
 
 		TMap<EInterchangePropertyTracks, FInterchangeProperty> PropertyTracks;
diff --git a/Engine/Plugins/Interchange/Runtime/Source/Nodes/Public/InterchangeAnimationDefinitions.h b/Engine/Plugins/Interchange/Runtime/Source/Nodes/Public/InterchangeAnimationDefinitions.h
--- a/Engine/Plugins/Interchange/Runtime/Source/Nodes/Public/InterchangeAnimationDefinitions.h
+++ b/Engine/Plugins/Interchange/Runtime/Source/Nodes/Public/InterchangeAnimationDefinitions.h
@@ -30,10 +30,10 @@ enum class EInterchangePropertyTracks : int32
 	CastShadow,
 	CustomDepthStencilValue,
 	CustomDepthStencilWriteMask,
+	DefaultUpVector,
 	DrawDebug,
 	EmissiveLightSource,
 	ExcludeFromLightAttachmentGroup,
-	HiddenInGame,
 	HiddenInSceneCapture,
 	Holdout,
 	LightAttachmentsAsGroup,
@@ -84,6 +84,7 @@ enum class EInterchangePropertyTracks : int32
 	LightForceCachedShadowsForMovablePrimitives,
 	LightForwardShadingPriority,
 	LightFunctionFadeDistance,
+	LightFunctionScale,
 	LightIESBrightnessScale,
 	LightIndirectLightingIntensity,
 	LightInnerConeAngle,
@@ -106,6 +107,7 @@ enum class EInterchangePropertyTracks : int32
 	LightShadowDistanceFadeoutFraction,
 	LightShadowSlopeBias,
 	LightShadowSourceAngleFactor,
+	LightShaftOverrideDirection,
 	LightSoftSourceRadius,
 	LightSourceAngle,
 	LightSourceCubemapAngle,
@@ -126,6 +128,7 @@ enum class EInterchangePropertyTracks : int32
 	CameraAspectRatio,
 	CameraAspectRatioAxisConstraint,
 	CameraAutoCalculateOrthoPlanes,
+	CameraAutoPlaneShift,
 	CameraConstrainAspectRatio,
 	CameraCurrentAperture,
 	CameraCurrentFocalLength,
@@ -136,6 +139,7 @@ enum class EInterchangePropertyTracks : int32
 	CameraFilmbackSensorWidth,
 	CameraFocusSettingsFocusOffset,
 	CameraFocusSettingsManualFocusDistance,
+	CameraFocusSettingsTrackingFocusSettingsRelativeOffset,
 	CameraOrthoFarClipPlane,
 	CameraOrthoNearClipPlane,
 	CameraOrthoWidth,
@@ -175,6 +179,7 @@ enum class EInterchangePropertyTracks : int32
 	CameraPostProcessSettingsBloom6Size,
 	CameraPostProcessSettingsBloom6Tint,
 	CameraPostProcessSettingsBloomConvolutionBufferScale,
+	CameraPostProcessSettingsBloomConvolutionCenterUV,
 	CameraPostProcessSettingsBloomConvolutionPreFilterMax,
 	CameraPostProcessSettingsBloomConvolutionPreFilterMin,
 	CameraPostProcessSettingsBloomConvolutionPreFilterMult,
@@ -187,10 +192,30 @@ enum class EInterchangePropertyTracks : int32
 	CameraPostProcessSettingsBloomThreshold,
 	CameraPostProcessSettingsBlueCorrection,
 	CameraPostProcessSettingsChromaticAberrationStartOffset,
+	CameraPostProcessSettingsColorContrast,
+	CameraPostProcessSettingsColorContrastHighlights,
+	CameraPostProcessSettingsColorContrastMidtones,
+	CameraPostProcessSettingsColorContrastShadows,
 	CameraPostProcessSettingsColorCorrectionHighlightsMax,
 	CameraPostProcessSettingsColorCorrectionHighlightsMin,
 	CameraPostProcessSettingsColorCorrectionShadowsMax,
+	CameraPostProcessSettingsColorGain,
+	CameraPostProcessSettingsColorGainHighlights,
+	CameraPostProcessSettingsColorGainMidtones,
+	CameraPostProcessSettingsColorGainShadows,
+	CameraPostProcessSettingsColorGamma,
+	CameraPostProcessSettingsColorGammaHighlights,
+	CameraPostProcessSettingsColorGammaMidtones,
+	CameraPostProcessSettingsColorGammaShadows,
 	CameraPostProcessSettingsColorGradingIntensity,
+	CameraPostProcessSettingsColorOffset,
+	CameraPostProcessSettingsColorOffsetHighlights,
+	CameraPostProcessSettingsColorOffsetMidtones,
+	CameraPostProcessSettingsColorOffsetShadows,
+	CameraPostProcessSettingsColorSaturation,
+	CameraPostProcessSettingsColorSaturationHighlights,
+	CameraPostProcessSettingsColorSaturationMidtones,
+	CameraPostProcessSettingsColorSaturationShadows,
 	CameraPostProcessSettingsDepthOfFieldBladeCount,
 	CameraPostProcessSettingsDepthOfFieldDepthBlurAmount,
 	CameraPostProcessSettingsDepthOfFieldDepthBlurRadius,
@@ -284,6 +309,8 @@ enum class EInterchangePropertyTracks : int32
 	CameraPostProcessSettingsWhiteTemp,
 	CameraPostProcessSettingsWhiteTint,
 	CameraProjectionMode,
+	CameraUpdateOrthoPlanes,
+	CameraUseCameraHeightAsViewTarget,
 	CameraUseFieldOfViewForLOD,
 
 	None = -1 UMETA(hidden),
