Interchange: Enable fbx translator for interchange to replace the legacy fbx importer.

#rnx

diff --git a/Engine/Plugins/Interchange/Runtime/Source/Import/Private/Fbx/InterchangeFbxTranslator.cpp b/Engine/Plugins/Interchange/Runtime/Source/Import/Private/Fbx/InterchangeFbxTranslator.cpp
--- a/Engine/Plugins/Interchange/Runtime/Source/Import/Private/Fbx/InterchangeFbxTranslator.cpp
+++ b/Engine/Plugins/Interchange/Runtime/Source/Import/Private/Fbx/InterchangeFbxTranslator.cpp
@@ -26,7 +26,7 @@
 
 #include UE_INLINE_GENERATED_CPP_BY_NAME(InterchangeFbxTranslator)
 
-static bool GInterchangeEnableFBXImport = false;
+static bool GInterchangeEnableFBXImport = true;
 static FAutoConsoleVariableRef CCvarInterchangeEnableFBXImport(
 	TEXT("Interchange.FeatureFlags.Import.FBX"),
 	GInterchangeEnableFBXImport,
