Interchange: Don't append _Anim to AnimSequence if ImportOnlyAnimation settings is true.

#rnx

diff --git a/Engine/Plugins/Interchange/Runtime/Source/Pipelines/Private/InterchangeGenericAssetsPipeline.cpp b/Engine/Plugins/Interchange/Runtime/Source/Pipelines/Private/InterchangeGenericAssetsPipeline.cpp
--- a/Engine/Plugins/Interchange/Runtime/Source/Pipelines/Private/InterchangeGenericAssetsPipeline.cpp
+++ b/Engine/Plugins/Interchange/Runtime/Source/Pipelines/Private/InterchangeGenericAssetsPipeline.cpp
@@ -1536,7 +1536,12 @@ void UInterchangeGenericAssetsPipeline::ImplementUseSourceNameForAssetOption(UIn
 		if (AnimSequenceNodeUids.Num() == 1)
 		{
 			UInterchangeAnimSequenceFactoryNode* AnimSequenceNode = Cast<UInterchangeAnimSequenceFactoryNode>(InBaseNodeContainer->GetFactoryNode(AnimSequenceNodeUids[0]));
-			const FString DisplayLabelName = (OverrideAssetName.IsEmpty() ? FPaths::GetBaseFilename(InSourceDatas[0]->GetFilename()) : OverrideAssetName) + TEXT("_Anim");
+			FString DisplayLabelName = (OverrideAssetName.IsEmpty() ? FPaths::GetBaseFilename(InSourceDatas[0]->GetFilename()) : OverrideAssetName);
+			const bool bImportOnlyAnimations = CommonSkeletalMeshesAndAnimationsProperties && CommonSkeletalMeshesAndAnimationsProperties->bImportOnlyAnimations;
+			if (!bImportOnlyAnimations)
+			{
+				DisplayLabelName += TEXT("_Anim");
+			}
 			AnimSequenceNode->SetDisplayLabel(DisplayLabelName);
 		}
 	}
