Interchange: fix re-import LOD failed

#rnx

diff --git a/Engine/Plugins/Interchange/Runtime/Source/Pipelines/Private/InterchangeGenericMeshPipeline.cpp b/Engine/Plugins/Interchange/Runtime/Source/Pipelines/Private/InterchangeGenericMeshPipeline.cpp
--- a/Engine/Plugins/Interchange/Runtime/Source/Pipelines/Private/InterchangeGenericMeshPipeline.cpp
+++ b/Engine/Plugins/Interchange/Runtime/Source/Pipelines/Private/InterchangeGenericMeshPipeline.cpp
@@ -44,6 +44,7 @@ void UInterchangeGenericMeshPipeline::AdjustSettingsForContext(EInterchangePipel
 	{
 		bCreatePhysicsAsset = false;
 		PhysicsAsset = nullptr;
+		LodGroup = NAME_None;
 		
 		if (ImportType == EInterchangePipelineContext::AssetAlternateSkinningImport
 			|| ImportType == EInterchangePipelineContext::AssetAlternateSkinningReimport)
