Interchange: Do not create empty skeletal mesh

#rnx

diff --git a/Engine/Plugins/Interchange/Runtime/Source/Import/Private/Mesh/InterchangeSkeletalMeshFactory.cpp b/Engine/Plugins/Interchange/Runtime/Source/Import/Private/Mesh/InterchangeSkeletalMeshFactory.cpp
--- a/Engine/Plugins/Interchange/Runtime/Source/Import/Private/Mesh/InterchangeSkeletalMeshFactory.cpp
+++ b/Engine/Plugins/Interchange/Runtime/Source/Import/Private/Mesh/InterchangeSkeletalMeshFactory.cpp
@@ -1678,7 +1678,7 @@ UInterchangeFactoryBase::FImportAssetResult UInterchangeSkeletalMeshFactory::End
 	FSkinnedAssetAsyncBuildScope AsyncBuildScope(SkeletalMesh);
 
 	//Finish the import in the game thread so we can pop dialog if needed
-	if (ensure(ImportAssetObjectData.IsValid()))
+	if (ImportAssetObjectData.IsValid())
 	{
 		for (FImportAssetObjectLODData& ImportAssetObjectLODData : ImportAssetObjectData.LodDatas)
 		{
diff --git a/Engine/Plugins/Interchange/Runtime/Source/Pipelines/Private/InterchangeGenericMeshPipeline.cpp b/Engine/Plugins/Interchange/Runtime/Source/Pipelines/Private/InterchangeGenericMeshPipeline.cpp
--- a/Engine/Plugins/Interchange/Runtime/Source/Pipelines/Private/InterchangeGenericMeshPipeline.cpp
+++ b/Engine/Plugins/Interchange/Runtime/Source/Pipelines/Private/InterchangeGenericMeshPipeline.cpp
@@ -208,10 +208,11 @@ UInterchangePipelineMeshesUtilities* UInterchangeGenericMeshPipeline::CreateMesh
 	UInterchangePipelineMeshesUtilities* CreatedPipelineMeshesUtilities = UInterchangePipelineMeshesUtilities::CreateInterchangePipelineMeshesUtilities(InBaseNodeContainer);
 
 	bool bAutoDetectConvertStaticMeshToSkeletalMesh = false;
-	if (bAutoDetectType && Pipeline->CommonMeshesProperties->ForceAllMeshAsType == EInterchangeForceMeshType::IFMT_None)
+	bool bContainStaticMesh = false;
+	bool bContainSkeletalMesh = false;
+	bool bContainStaticMeshAnimationNode = false;
 	{
 		TArray<FString> StaticMeshNodeUids;
-		bool bContainSkeletalMesh = false;
 		InBaseNodeContainer->IterateNodesOfType<UInterchangeMeshNode>([&bContainSkeletalMesh, &StaticMeshNodeUids](const FString& NodeUid, UInterchangeMeshNode* MeshNode)
 			{
 				if (!MeshNode->IsMorphTarget())
@@ -219,36 +220,45 @@ UInterchangePipelineMeshesUtilities* UInterchangeGenericMeshPipeline::CreateMesh
 					MeshNode->IsSkinnedMesh() ? bContainSkeletalMesh = true : StaticMeshNodeUids.Add(NodeUid);
 				}
 			});
-		
-		bool bContainAnimationNode = false;
-		if (!bContainSkeletalMesh && StaticMeshNodeUids.Num() > 0)
-		{
-			TMap<const UInterchangeSceneNode*, bool> CacheProcessSceneNodes;
-			InBaseNodeContainer->BreakableIterateNodesOfType<UInterchangeTransformAnimationTrackNode>([&InBaseNodeContainer, &bContainAnimationNode, &StaticMeshNodeUids, &CacheProcessSceneNodes](const FString& NodeUid, UInterchangeTransformAnimationTrackNode* AnimationNode)
+		bContainStaticMesh = !StaticMeshNodeUids.IsEmpty();
+	
+		TMap<const UInterchangeSceneNode*, bool> CacheProcessSceneNodes;
+		InBaseNodeContainer->BreakableIterateNodesOfType<UInterchangeTransformAnimationTrackNode>([&InBaseNodeContainer, &bContainStaticMeshAnimationNode, &StaticMeshNodeUids, &CacheProcessSceneNodes](const FString& NodeUid, UInterchangeTransformAnimationTrackNode* AnimationNode)
+			{
+				FString SceneNodeUid;
+				if (AnimationNode->GetCustomActorDependencyUid(SceneNodeUid))
 				{
-					FString SceneNodeUid;
-					if (AnimationNode->GetCustomActorDependencyUid(SceneNodeUid))
+					if (const UInterchangeSceneNode* SceneNode = Cast<UInterchangeSceneNode>(InBaseNodeContainer->GetNode(SceneNodeUid)))
 					{
-						if (const UInterchangeSceneNode* SceneNode = Cast<UInterchangeSceneNode>(InBaseNodeContainer->GetNode(SceneNodeUid)))
+						if (IsImpactingAnyMeshesRecursive(SceneNode, InBaseNodeContainer, StaticMeshNodeUids, CacheProcessSceneNodes))
 						{
-							if (IsImpactingAnyMeshesRecursive(SceneNode, InBaseNodeContainer, StaticMeshNodeUids, CacheProcessSceneNodes))
-							{
-								bContainAnimationNode = true;
-							}
+							bContainStaticMeshAnimationNode = true;
 						}
 					}
-					return bContainAnimationNode;
-				});
-		}
+				}
+				return bContainStaticMeshAnimationNode;
+			});
+	}
 
-		//Auto detect some static mesh transform animations, we need to force the skeletal mesh type and recompute
-		bAutoDetectConvertStaticMeshToSkeletalMesh = bContainAnimationNode;
+	if (bAutoDetectType && Pipeline->CommonMeshesProperties->ForceAllMeshAsType == EInterchangeForceMeshType::IFMT_None)
+	{
+		if (!bContainSkeletalMesh && bContainStaticMesh)
+		{
+			//Auto detect some static mesh transform animations, we need to force the skeletal mesh type and recompute
+			bAutoDetectConvertStaticMeshToSkeletalMesh = bContainStaticMeshAnimationNode;
+		}
 	}
 
 	//Set the context option to use when querying the pipeline mesh utilities
 	FInterchangePipelineMeshesUtilitiesContext DataContext;
-	DataContext.bConvertStaticMeshToSkeletalMesh = bAutoDetectConvertStaticMeshToSkeletalMesh || (Pipeline->CommonMeshesProperties->ForceAllMeshAsType == EInterchangeForceMeshType::IFMT_SkeletalMesh);
+	
+	//We convert to skeletal mesh, only if the translated data do not have skeletal mesh
+	//Rigid mesh import is a fallback when there is no skinned mesh
+	DataContext.bConvertStaticMeshToSkeletalMesh = !bContainSkeletalMesh && (bAutoDetectConvertStaticMeshToSkeletalMesh || (Pipeline->CommonMeshesProperties->ForceAllMeshAsType == EInterchangeForceMeshType::IFMT_SkeletalMesh));
+
+	//Force static mesh convert all mesh to static mesh
 	DataContext.bConvertSkeletalMeshToStaticMesh = (Pipeline->CommonMeshesProperties->ForceAllMeshAsType == EInterchangeForceMeshType::IFMT_StaticMesh);
+
 	DataContext.bConvertStaticsWithMorphTargetsToSkeletals = Pipeline->CommonSkeletalMeshesAndAnimationsProperties->bConvertStaticsWithMorphTargetsToSkeletals;
 	DataContext.bImportMeshesInBoneHierarchy = Pipeline->CommonSkeletalMeshesAndAnimationsProperties->bImportMeshesInBoneHierarchy;
 	DataContext.bQueryGeometryOnlyIfNoInstance = Pipeline->CommonMeshesProperties->bBakeMeshes || Pipeline->CommonMeshesProperties->bBakePivotMeshes;
diff --git a/Engine/Plugins/Interchange/Runtime/Source/Pipelines/Public/InterchangeGenericAssetsPipelineSharedSettings.h b/Engine/Plugins/Interchange/Runtime/Source/Pipelines/Public/InterchangeGenericAssetsPipelineSharedSettings.h
--- a/Engine/Plugins/Interchange/Runtime/Source/Pipelines/Public/InterchangeGenericAssetsPipelineSharedSettings.h
+++ b/Engine/Plugins/Interchange/Runtime/Source/Pipelines/Public/InterchangeGenericAssetsPipelineSharedSettings.h
@@ -53,7 +53,11 @@ class INTERCHANGEPIPELINES_API UInterchangeGenericCommonMeshesProperties : publi
 public:
 	//////	COMMON_MESHES_CATEGORY Properties //////
 
-	/** If set, imports all meshes in the source as either static meshes or skeletal meshes. */
+	/**
+	 * If set, imports all meshes in the source as either static meshes or skeletal meshes.
+	 * For skeletal meshes the conversion will happen only if there is no skinned meshes.
+	 * Mixing rigid skeletal mesh with skinned mesh is not good and will result in multiple skeletal meshes.
+	 */
  	UPROPERTY(EditAnywhere, BlueprintReadWrite, Category = "Common Meshes")
  	EInterchangeForceMeshType ForceAllMeshAsType = EInterchangeForceMeshType::IFMT_None;
 
