Interchange: Make the autodetect of the generic asset pipeline set the import staticmesh to false if there is some skeletalmesh in the translated data.

#rnx

diff --git a/Engine/Plugins/Interchange/Runtime/Source/Pipelines/Private/InterchangeGenericAssetsPipeline.cpp b/Engine/Plugins/Interchange/Runtime/Source/Pipelines/Private/InterchangeGenericAssetsPipeline.cpp
--- a/Engine/Plugins/Interchange/Runtime/Source/Pipelines/Private/InterchangeGenericAssetsPipeline.cpp
+++ b/Engine/Plugins/Interchange/Runtime/Source/Pipelines/Private/InterchangeGenericAssetsPipeline.cpp
@@ -163,7 +163,7 @@ void UInterchangeGenericAssetsPipeline::FilterPropertiesFromTranslatedData(UInte
 		MeshPipeline->FilterPropertiesFromTranslatedData(InBaseNodeContainer);
 		AnimationPipeline->FilterPropertiesFromTranslatedData(InBaseNodeContainer);
 
-		UInterchangePipelineMeshesUtilities* PipelineMeshesUtilities = UInterchangeGenericMeshPipeline::CreateMeshPipelineUtilities(InBaseNodeContainer, MeshPipeline, CommonMeshesProperties->bAutoDetectMeshType);
+		UInterchangePipelineMeshesUtilities* PipelineMeshesUtilities = UInterchangeGenericMeshPipeline::CreateMeshPipelineUtilities(InBaseNodeContainer, MeshPipeline);
 
 		TArray<FString> SkeletalMeshes;
 		PipelineMeshesUtilities->GetAllSkinnedMeshInstance(SkeletalMeshes);
diff --git a/Engine/Plugins/Interchange/Runtime/Source/Pipelines/Private/InterchangeGenericMeshPipeline.cpp b/Engine/Plugins/Interchange/Runtime/Source/Pipelines/Private/InterchangeGenericMeshPipeline.cpp
--- a/Engine/Plugins/Interchange/Runtime/Source/Pipelines/Private/InterchangeGenericMeshPipeline.cpp
+++ b/Engine/Plugins/Interchange/Runtime/Source/Pipelines/Private/InterchangeGenericMeshPipeline.cpp
@@ -34,7 +34,12 @@ void UInterchangeGenericMeshPipeline::AdjustSettingsForContext(EInterchangePipel
 		return;
 	}
 
-	const bool bIsReimport = IsReimportContext();
+	bool bAutoDetectConvertStaticMeshToSkeletalMesh = false;
+	bool bContainStaticMesh = false;
+	bool bContainSkeletalMesh = false;
+	bool bContainStaticMeshAnimationNode = false;
+	bool bIgnoreStaticMesh = false;
+	GetMeshesInformationFromTranslatedData(InBaseNodeContainer, bAutoDetectConvertStaticMeshToSkeletalMesh, bContainStaticMesh, bContainSkeletalMesh, bContainStaticMeshAnimationNode, bIgnoreStaticMesh);
 
 	//Avoid creating physics asset when importing a LOD or the alternate skinning
 	if (ImportType == EInterchangePipelineContext::AssetCustomLODImport
@@ -76,21 +81,6 @@ void UInterchangeGenericMeshPipeline::AdjustSettingsForContext(EInterchangePipel
 	TArray<FString> HideSubCategories;
 	if (ImportType == EInterchangePipelineContext::AssetReimport)
 	{
-		bool bContainSkeletalMesh = false;
-		bool bContainStaticMesh = false;
-
-		if (InBaseNodeContainer)
-		{
-			InBaseNodeContainer->BreakableIterateNodesOfType<UInterchangeMeshNode>([&bContainSkeletalMesh, &bContainStaticMesh](const FString& NodeUid, UInterchangeMeshNode* MeshNode)
-				{
-					if (!MeshNode->IsMorphTarget())
-					{
-						MeshNode->IsSkinnedMesh() ? bContainSkeletalMesh = true : bContainStaticMesh = true;
-					}
-					return bContainSkeletalMesh && bContainStaticMesh;
-				});
-		}
-
 		CommonMeshesProperties->bAutoDetectMeshType = false;
 
 		HideSubCategories.Add(TEXT("Build"));
@@ -114,7 +104,7 @@ void UInterchangeGenericMeshPipeline::AdjustSettingsForContext(EInterchangePipel
 		{
 			HideCategories.Add(SkeletalMeshesCategory);
 			HideCategories.Add(CommonSkeletalMeshesAndAnimationCategory);
-			
+			bImportSkeletalMeshes = false;
 			if (!bContainStaticMesh)
 			{
 				CommonMeshesProperties->ForceAllMeshAsType = EInterchangeForceMeshType::IFMT_StaticMesh;
@@ -201,16 +191,18 @@ bool UInterchangeGenericMeshPipeline::GetPropertyPossibleValues(const FName Prop
 
 #endif
 
-UInterchangePipelineMeshesUtilities* UInterchangeGenericMeshPipeline::CreateMeshPipelineUtilities(UInterchangeBaseNodeContainer* InBaseNodeContainer
-	, const UInterchangeGenericMeshPipeline* Pipeline
-	, const bool bAutoDetectType)
+void UInterchangeGenericMeshPipeline::GetMeshesInformationFromTranslatedData(const UInterchangeBaseNodeContainer* InBaseNodeContainer
+	, bool& bAutoDetectConvertStaticMeshToSkeletalMesh
+	, bool& bContainStaticMesh
+	, bool& bContainSkeletalMesh
+	, bool& bContainStaticMeshAnimationNode
+	, bool& bIgnoreStaticMeshes) const
 {
-	UInterchangePipelineMeshesUtilities* CreatedPipelineMeshesUtilities = UInterchangePipelineMeshesUtilities::CreateInterchangePipelineMeshesUtilities(InBaseNodeContainer);
-
-	bool bAutoDetectConvertStaticMeshToSkeletalMesh = false;
-	bool bContainStaticMesh = false;
-	bool bContainSkeletalMesh = false;
-	bool bContainStaticMeshAnimationNode = false;
+	bAutoDetectConvertStaticMeshToSkeletalMesh = false;
+	bContainStaticMesh = false;
+	bContainSkeletalMesh = false;
+	bContainStaticMeshAnimationNode = false;
+	bIgnoreStaticMeshes = false;
 	{
 		TArray<FString> StaticMeshNodeUids;
 		InBaseNodeContainer->IterateNodesOfType<UInterchangeMeshNode>([&bContainSkeletalMesh, &StaticMeshNodeUids](const FString& NodeUid, UInterchangeMeshNode* MeshNode)
@@ -221,7 +213,7 @@ UInterchangePipelineMeshesUtilities* UInterchangeGenericMeshPipeline::CreateMesh
 				}
 			});
 		bContainStaticMesh = !StaticMeshNodeUids.IsEmpty();
-	
+
 		TMap<const UInterchangeSceneNode*, bool> CacheProcessSceneNodes;
 		InBaseNodeContainer->BreakableIterateNodesOfType<UInterchangeTransformAnimationTrackNode>([&InBaseNodeContainer, &bContainStaticMeshAnimationNode, &StaticMeshNodeUids, &CacheProcessSceneNodes](const FString& NodeUid, UInterchangeTransformAnimationTrackNode* AnimationNode)
 			{
@@ -240,14 +232,31 @@ UInterchangePipelineMeshesUtilities* UInterchangeGenericMeshPipeline::CreateMesh
 			});
 	}
 
-	if (bAutoDetectType && Pipeline->CommonMeshesProperties->ForceAllMeshAsType == EInterchangeForceMeshType::IFMT_None)
+	if (CommonMeshesProperties->bAutoDetectMeshType && CommonMeshesProperties->ForceAllMeshAsType == EInterchangeForceMeshType::IFMT_None)
 	{
 		if (!bContainSkeletalMesh && bContainStaticMesh)
 		{
 			//Auto detect some static mesh transform animations, we need to force the skeletal mesh type and recompute
 			bAutoDetectConvertStaticMeshToSkeletalMesh = bContainStaticMeshAnimationNode;
 		}
+		else if (bContainSkeletalMesh)
+		{
+			bIgnoreStaticMeshes = true;
+		}
 	}
+}
+
+UInterchangePipelineMeshesUtilities* UInterchangeGenericMeshPipeline::CreateMeshPipelineUtilities(UInterchangeBaseNodeContainer* InBaseNodeContainer
+	, const UInterchangeGenericMeshPipeline* Pipeline)
+{
+	UInterchangePipelineMeshesUtilities* CreatedPipelineMeshesUtilities = UInterchangePipelineMeshesUtilities::CreateInterchangePipelineMeshesUtilities(InBaseNodeContainer);
+
+	bool bAutoDetectConvertStaticMeshToSkeletalMesh = false;
+	bool bContainStaticMesh = false;
+	bool bContainSkeletalMesh = false;
+	bool bContainStaticMeshAnimationNode = false;
+	bool bIgnoreStaticMeshes = false;
+	Pipeline->GetMeshesInformationFromTranslatedData(InBaseNodeContainer, bAutoDetectConvertStaticMeshToSkeletalMesh, bContainStaticMesh, bContainSkeletalMesh, bContainStaticMeshAnimationNode, bIgnoreStaticMeshes);
 
 	//Set the context option to use when querying the pipeline mesh utilities
 	FInterchangePipelineMeshesUtilitiesContext DataContext;
@@ -262,6 +271,7 @@ UInterchangePipelineMeshesUtilities* UInterchangeGenericMeshPipeline::CreateMesh
 	DataContext.bConvertStaticsWithMorphTargetsToSkeletals = Pipeline->CommonSkeletalMeshesAndAnimationsProperties->bConvertStaticsWithMorphTargetsToSkeletals;
 	DataContext.bImportMeshesInBoneHierarchy = Pipeline->CommonSkeletalMeshesAndAnimationsProperties->bImportMeshesInBoneHierarchy;
 	DataContext.bQueryGeometryOnlyIfNoInstance = Pipeline->CommonMeshesProperties->bBakeMeshes || Pipeline->CommonMeshesProperties->bBakePivotMeshes;
+	DataContext.bIgnoreStaticMeshes = bIgnoreStaticMeshes;
 	CreatedPipelineMeshesUtilities->SetContext(DataContext);
 	return CreatedPipelineMeshesUtilities;
 }
@@ -280,7 +290,7 @@ void UInterchangeGenericMeshPipeline::ExecutePipeline(UInterchangeBaseNodeContai
 	{
 		SourceDatas.Add(SourceData);
 	}
-	PipelineMeshesUtilities = CreateMeshPipelineUtilities(BaseNodeContainer, this, CommonMeshesProperties->bAutoDetectMeshType);
+	PipelineMeshesUtilities = CreateMeshPipelineUtilities(BaseNodeContainer, this);
 
 	//Create skeletalmesh factory nodes
 	ExecutePreImportPipelineSkeletalMesh();
diff --git a/Engine/Plugins/Interchange/Runtime/Source/Pipelines/Private/InterchangePipelineMeshesUtilities.cpp b/Engine/Plugins/Interchange/Runtime/Source/Pipelines/Private/InterchangePipelineMeshesUtilities.cpp
--- a/Engine/Plugins/Interchange/Runtime/Source/Pipelines/Private/InterchangePipelineMeshesUtilities.cpp
+++ b/Engine/Plugins/Interchange/Runtime/Source/Pipelines/Private/InterchangePipelineMeshesUtilities.cpp
@@ -175,6 +175,10 @@ static bool IsSceneNodeASocket(const UInterchangeSceneNode* SceneNode)
 
 bool FInterchangePipelineMeshesUtilitiesContext::IsStaticMeshInstance(const FInterchangeMeshInstance& MeshInstance, UInterchangeBaseNodeContainer* BaseNodeContainer)
 {
+	if (bIgnoreStaticMeshes)
+	{
+		return false;
+	}
 	return !IsSkeletalMeshInstance(MeshInstance, BaseNodeContainer);
 }
 
@@ -215,6 +219,10 @@ bool FInterchangePipelineMeshesUtilitiesContext::IsSkeletalMeshInstance(const FI
 
 bool FInterchangePipelineMeshesUtilitiesContext::IsStaticMeshGeometry(const FInterchangeMeshGeometry& MeshGeometry)
 {
+	if (bIgnoreStaticMeshes)
+	{
+		return false;
+	}
 	if (bQueryGeometryOnlyIfNoInstance && MeshGeometry.ReferencingMeshInstanceUids.Num() > 0)
 	{
 		return false;
diff --git a/Engine/Plugins/Interchange/Runtime/Source/Pipelines/Public/InterchangeGenericMeshPipeline.h b/Engine/Plugins/Interchange/Runtime/Source/Pipelines/Public/InterchangeGenericMeshPipeline.h
--- a/Engine/Plugins/Interchange/Runtime/Source/Pipelines/Public/InterchangeGenericMeshPipeline.h
+++ b/Engine/Plugins/Interchange/Runtime/Source/Pipelines/Public/InterchangeGenericMeshPipeline.h
@@ -229,8 +229,7 @@ public:
 #endif
 
 	static UInterchangePipelineMeshesUtilities* CreateMeshPipelineUtilities(UInterchangeBaseNodeContainer* InBaseNodeContainer
-		, const UInterchangeGenericMeshPipeline* Pipeline
-		, const bool bAutoDetectType);
+		, const UInterchangeGenericMeshPipeline* Pipeline);
 
 protected:
 	virtual void ExecutePipeline(UInterchangeBaseNodeContainer* InBaseNodeContainer, const TArray<UInterchangeSourceData*>& InSourceDatas, const FString& ContentBasePath) override;
@@ -348,6 +347,16 @@ private:
 	/* Static mesh API END                                                */
 	/************************************************************************/
 
+	/**
+	 * Fill all reference parameter from the translated node found in the container
+	 */
+	void GetMeshesInformationFromTranslatedData(const UInterchangeBaseNodeContainer* InBaseNodeContainer
+		, bool& bAutoDetectConvertStaticMeshToSkeletalMesh
+		, bool& bContainStaticMesh
+		, bool& bContainSkeletalMesh
+		, bool& bContainStaticMeshAnimationNode
+		, bool& bIgnoreStaticMeshes) const;
+
 private:
 
 	UInterchangeBaseNodeContainer* BaseNodeContainer = nullptr;
diff --git a/Engine/Plugins/Interchange/Runtime/Source/Pipelines/Public/InterchangePipelineMeshesUtilities.h b/Engine/Plugins/Interchange/Runtime/Source/Pipelines/Public/InterchangePipelineMeshesUtilities.h
--- a/Engine/Plugins/Interchange/Runtime/Source/Pipelines/Public/InterchangePipelineMeshesUtilities.h
+++ b/Engine/Plugins/Interchange/Runtime/Source/Pipelines/Public/InterchangePipelineMeshesUtilities.h
@@ -159,6 +159,12 @@ struct FInterchangePipelineMeshesUtilitiesContext
 	UPROPERTY(EditAnywhere, BlueprintReadWrite, Category = "Interchange | Pipeline | MeshesContext")
 	bool bQueryGeometryOnlyIfNoInstance = true;
 
+	/**
+	 * If enabled, all static meshes will be ignored. The mesh utility will not return any static meshes.
+	 */
+	UPROPERTY(EditAnywhere, BlueprintReadWrite, Category = "Interchange | Pipeline | MeshesContext")
+	bool bIgnoreStaticMeshes = false;
+
 	bool IsStaticMeshInstance(const FInterchangeMeshInstance& MeshInstance, UInterchangeBaseNodeContainer* BaseNodeContainer);
 	bool IsSkeletalMeshInstance(const FInterchangeMeshInstance& MeshInstance, UInterchangeBaseNodeContainer* BaseNodeContainer);
 	bool IsSkeletalMeshInstance(const FInterchangeMeshInstance& MeshInstance, UInterchangeBaseNodeContainer* BaseNodeContainer, bool& bOutIsStaticMeshNestedInSkeleton);
