Interchange: fix a crash from a missing pointer verification

#rnx

diff --git a/Engine/Plugins/Interchange/Runtime/Source/Pipelines/Private/InterchangeGenericMeshPipeline.cpp b/Engine/Plugins/Interchange/Runtime/Source/Pipelines/Private/InterchangeGenericMeshPipeline.cpp
--- a/Engine/Plugins/Interchange/Runtime/Source/Pipelines/Private/InterchangeGenericMeshPipeline.cpp
+++ b/Engine/Plugins/Interchange/Runtime/Source/Pipelines/Private/InterchangeGenericMeshPipeline.cpp
@@ -198,6 +198,11 @@ void UInterchangeGenericMeshPipeline::GetMeshesInformationFromTranslatedData(con
 	, bool& bContainStaticMeshAnimationNode
 	, bool& bIgnoreStaticMeshes) const
 {
+	//Its valid to call GetMeshesInformationFromTranslatedData with a null container
+	if (!InBaseNodeContainer)
+	{
+		return;
+	}
 	bAutoDetectConvertStaticMeshToSkeletalMesh = false;
 	bContainStaticMesh = false;
 	bContainSkeletalMesh = false;
