Fixed issue with simple collision box when re-importing staticmesh

#rnx

diff --git a/Engine/Plugins/Interchange/Runtime/Source/Import/Private/Mesh/InterchangeStaticMeshFactory.cpp b/Engine/Plugins/Interchange/Runtime/Source/Import/Private/Mesh/InterchangeStaticMeshFactory.cpp
--- a/Engine/Plugins/Interchange/Runtime/Source/Import/Private/Mesh/InterchangeStaticMeshFactory.cpp
+++ b/Engine/Plugins/Interchange/Runtime/Source/Import/Private/Mesh/InterchangeStaticMeshFactory.cpp
@@ -199,7 +199,6 @@ UInterchangeFactoryBase::FImportAssetResult UInterchangeStaticMeshFactory::Begin
 			StaticMesh->ReleaseResourcesFence.Wait();
 			
 			StaticMesh->SetRenderData(nullptr);
-			StaticMesh->SetBodySetup(nullptr);
 		}
 	}
 	
@@ -213,7 +212,10 @@ UInterchangeFactoryBase::FImportAssetResult UInterchangeStaticMeshFactory::Begin
 	}
 
 	// create the BodySetup on the game thread
-	StaticMesh->CreateBodySetup();
+	if (!ExistingAsset)
+	{
+		StaticMesh->CreateBodySetup();
+	}
 	
 #if WITH_EDITOR
 	if (!ImportAssetObjectData.bIsAppGame)
