Interchange: fix the LOD prefix workflow to not add lod group with only one LOD.

#rnx

diff --git a/Engine/Plugins/Interchange/Runtime/Source/Pipelines/Private/InterchangePipelineMeshesUtilities.cpp b/Engine/Plugins/Interchange/Runtime/Source/Pipelines/Private/InterchangePipelineMeshesUtilities.cpp
--- a/Engine/Plugins/Interchange/Runtime/Source/Pipelines/Private/InterchangePipelineMeshesUtilities.cpp
+++ b/Engine/Plugins/Interchange/Runtime/Source/Pipelines/Private/InterchangePipelineMeshesUtilities.cpp
@@ -17,9 +17,10 @@ namespace UE::Private::InterchangeMeshPipeline
 {
 	void FindNamedLodGroup(UInterchangeBaseNodeContainer* BaseNodeContainer, TMap<FString, TArray<FString>>& SceneMeshNodeUidsPerLodParentUidMap)
 	{
+		TMap<FString, TArray<FString>> TmpSceneMeshNodeUidsPerLodParentUidMap;
 		const FString LodPrefix = TEXT("LOD");
 		BaseNodeContainer->IterateNodes(
-			[&BaseNodeContainer, &SceneMeshNodeUidsPerLodParentUidMap, &LodPrefix](const FString& NodeUid, UInterchangeBaseNode* Node)
+			[&BaseNodeContainer, &TmpSceneMeshNodeUidsPerLodParentUidMap, &LodPrefix](const FString& NodeUid, UInterchangeBaseNode* Node)
 			{
 				if (Node->GetNodeContainerType() == EInterchangeNodeContainerType::TranslatedScene)
 				{
@@ -71,7 +72,7 @@ namespace UE::Private::InterchangeMeshPipeline
 							int32 LodNumber = FPlatformString::Atoi(*LODXNumber);
 							FString MatchName = ParentUniqueID;
 							
-							TArray<FString>& LodChildUids = SceneMeshNodeUidsPerLodParentUidMap.FindOrAdd(MatchName);
+							TArray<FString>& LodChildUids = TmpSceneMeshNodeUidsPerLodParentUidMap.FindOrAdd(MatchName);
 							//Add LOD at the correct index
 							if (LodNumber >= LodChildUids.Num())
 							{
@@ -86,7 +87,7 @@ namespace UE::Private::InterchangeMeshPipeline
 		);
 
 		//Remove all empty entry, we use empty entry to set all lod in the correct order
-		for (TPair<FString, TArray<FString>>& LodPrefixNodePair : SceneMeshNodeUidsPerLodParentUidMap)
+		for (TPair<FString, TArray<FString>>& LodPrefixNodePair : TmpSceneMeshNodeUidsPerLodParentUidMap)
 		{
 			TArray<FString>& LodChildUids = LodPrefixNodePair.Value;
 			for (int32 ChildLodIndex = LodChildUids.Num() - 1; ChildLodIndex >= 0; ChildLodIndex--)
@@ -98,6 +99,11 @@ namespace UE::Private::InterchangeMeshPipeline
 			}
 			//Shrink the array to the correct size
 			LodChildUids.Shrink();
+
+			if (LodChildUids.Num() > 1)
+			{
+				SceneMeshNodeUidsPerLodParentUidMap.Add(LodPrefixNodePair);
+			}
 		}
 	}
 
