Interchange Level Sequence:
- Add Static/Skeletal Mesh property tracks

diff --git a/Engine/Plugins/Interchange/Runtime/Source/Import/Private/Animation/InterchangeLevelSequenceHelper.cpp b/Engine/Plugins/Interchange/Runtime/Source/Import/Private/Animation/InterchangeLevelSequenceHelper.cpp
--- a/Engine/Plugins/Interchange/Runtime/Source/Import/Private/Animation/InterchangeLevelSequenceHelper.cpp
+++ b/Engine/Plugins/Interchange/Runtime/Source/Import/Private/Animation/InterchangeLevelSequenceHelper.cpp
@@ -9,9 +9,11 @@
 #include "Components/LocalLightComponent.h"
 #include "Components/PointLightComponent.h"
 #include "Components/RectLightComponent.h"
+#include "Components/SkeletalMeshComponent.h"
 #include "Components/SkyLightComponent.h"
 #include "Components/SpotLightComponent.h"
 #include "Components/SplineComponent.h"
+#include "Components/StaticMeshComponent.h"
 
 #include "MovieScene.h"
 #include "MovieSceneSection.h"
@@ -52,6 +54,7 @@ namespace UE::Interchange::Private
 			{EInterchangePropertyTracks::DrawDebug, {UMovieSceneBoolTrack::StaticClass()->GetName(),  GET_MEMBER_NAME_STRING_CHECKED(USplineComponent, bDrawDebug) , TEXT("Draw Debug")}},
 			{EInterchangePropertyTracks::EmissiveLightSource, {UMovieSceneBoolTrack::StaticClass()->GetName(),  GET_MEMBER_NAME_STRING_CHECKED(UPrimitiveComponent, bEmissiveLightSource) , TEXT("Emissive Light Source")}},
 			{EInterchangePropertyTracks::ExcludeFromLightAttachmentGroup, {UMovieSceneBoolTrack::StaticClass()->GetName(),  GET_MEMBER_NAME_STRING_CHECKED(UPrimitiveComponent, bExcludeFromLightAttachmentGroup) , TEXT("Exclude from Light Attachment Group")}},
+			{EInterchangePropertyTracks::HiddenInGame, {UMovieSceneBoolTrack::StaticClass()->GetName(),  GET_MEMBER_NAME_STRING_CHECKED(USceneComponent, bHiddenInGame) , TEXT("Hidden in Game")}},
 			{EInterchangePropertyTracks::HiddenInSceneCapture, {UMovieSceneBoolTrack::StaticClass()->GetName(),  GET_MEMBER_NAME_STRING_CHECKED(UPrimitiveComponent, bHiddenInSceneCapture) , TEXT("Hidden In Scene Capture")}},
 			{EInterchangePropertyTracks::Holdout, {UMovieSceneBoolTrack::StaticClass()->GetName(),  GET_MEMBER_NAME_STRING_CHECKED(UPrimitiveComponent, bHoldout) , TEXT("Holdout")}},
 			{EInterchangePropertyTracks::LightAttachmentsAsGroup, {UMovieSceneBoolTrack::StaticClass()->GetName(),  GET_MEMBER_NAME_STRING_CHECKED(UPrimitiveComponent, bLightAttachmentsAsGroup) , TEXT("Light Attachments as Group")}},
@@ -330,6 +333,30 @@ namespace UE::Interchange::Private
 			{EInterchangePropertyTracks::CameraUpdateOrthoPlanes, {UMovieSceneBoolTrack::StaticClass()->GetName(),  GET_MEMBER_NAME_STRING_CHECKED(UCameraComponent, bUpdateOrthoPlanes) , TEXT("Update Ortho Planes")}},
 			{EInterchangePropertyTracks::CameraUseCameraHeightAsViewTarget, {UMovieSceneBoolTrack::StaticClass()->GetName(),  GET_MEMBER_NAME_STRING_CHECKED(UCameraComponent, bUseCameraHeightAsViewTarget) , TEXT("Use Camera Height as View Target")}},
 			{EInterchangePropertyTracks::CameraUseFieldOfViewForLOD, {UMovieSceneBoolTrack::StaticClass()->GetName(),  GET_MEMBER_NAME_STRING_CHECKED(UCameraComponent, bUseFieldOfViewForLOD) , TEXT("Use Field Of View for LOD")}},
+
+			// Common Mesh
+			{EInterchangePropertyTracks::MeshOverlayMaterialMaxDrawDistance, {UMovieSceneFloatTrack::StaticClass()->GetName(),  GET_MEMBER_NAME_STRING_CHECKED(UMeshComponent, OverlayMaterialMaxDrawDistance) , TEXT("Overlay Material Max Draw Distance")}},
+
+			// Skinned Mesh
+			{EInterchangePropertyTracks::SkinnedMeshCapsuleIndirectShadowMinVisibility, {UMovieSceneFloatTrack::StaticClass()->GetName(),  GET_MEMBER_NAME_STRING_CHECKED(USkinnedMeshComponent, CapsuleIndirectShadowMinVisibility) , TEXT("Capsule Indirect Shadow Min Visibility")}},
+			{EInterchangePropertyTracks::SkinnedMeshCastCapsuleDirectShadow, {UMovieSceneBoolTrack::StaticClass()->GetName(),  GET_MEMBER_NAME_STRING_CHECKED(USkinnedMeshComponent, bCastCapsuleDirectShadow) , TEXT("Capsule Direct Shadow")}},
+			{EInterchangePropertyTracks::SkinnedMeshCastCapsuleIndirectShadow, {UMovieSceneBoolTrack::StaticClass()->GetName(),  GET_MEMBER_NAME_STRING_CHECKED(USkinnedMeshComponent, bCastCapsuleIndirectShadow) , TEXT("Capsule Indirect Shadow")}},
+			{EInterchangePropertyTracks::SkinnedMeshRenderStatic, {UMovieSceneBoolTrack::StaticClass()->GetName(),  GET_MEMBER_NAME_STRING_CHECKED(USkinnedMeshComponent, bRenderStatic) , TEXT("Render Static")}},
+			{EInterchangePropertyTracks::SkinnedMeshVisibilityBasedAnimTickOption, {UMovieSceneEnumTrack::StaticClass()->GetName(),  GET_MEMBER_NAME_STRING_CHECKED(USkinnedMeshComponent, VisibilityBasedAnimTickOption) , TEXT("Visibility Based Anim Tick Option"), StaticEnum<EVisibilityBasedAnimTickOption>()}},
+
+			// Skeletal Mesh
+			{EInterchangePropertyTracks::SkeletalMeshAllowClothActors, {UMovieSceneBoolTrack::StaticClass()->GetName(),  GET_MEMBER_NAME_STRING_CHECKED(USkeletalMeshComponent, bAllowClothActors) , TEXT("Allow Cloth Actors")}},
+			{EInterchangePropertyTracks::SkeletalMeshAnimationMode, {UMovieSceneEnumTrack::StaticClass()->GetName(),  GET_MEMBER_NAME_STRING_CHECKED(USkeletalMeshComponent, USkeletalMeshComponent::GetAnimationModePropertyNameChecked()) , TEXT("Animation Mode"), StaticEnum<EAnimationMode::Type>()}},
+			{EInterchangePropertyTracks::SkeletalMeshClothBlendWeight, {UMovieSceneFloatTrack::StaticClass()->GetName(),  GET_MEMBER_NAME_STRING_CHECKED(USkeletalMeshComponent, ClothBlendWeight) , TEXT("Cloth Blend Weight")}},
+			{EInterchangePropertyTracks::SkeletalMeshClothMaxDistanceScale, {UMovieSceneFloatTrack::StaticClass()->GetName(),  GET_MEMBER_NAME_STRING_CHECKED(USkeletalMeshComponent, ClothMaxDistanceScale) , TEXT("Cloth Max Distance Scale")}},
+			
+			// Static Mesh
+			{EInterchangePropertyTracks::StaticMeshDistanceFieldSelfShadowBias, {UMovieSceneFloatTrack::StaticClass()->GetName(),  GET_MEMBER_NAME_STRING_CHECKED(UStaticMeshComponent, DistanceFieldSelfShadowBias) , TEXT("Distance Field Self Shadow Bias")}},
+			{EInterchangePropertyTracks::StaticMeshEvaluateWorldPositionOffset, {UMovieSceneBoolTrack::StaticClass()->GetName(),  GET_MEMBER_NAME_STRING_CHECKED(UStaticMeshComponent, bEvaluateWorldPositionOffset) , TEXT("Evaluate World Position Offset")}},
+			{EInterchangePropertyTracks::StaticMeshEvaluateWorldPositionOffsetInRayTracing, {UMovieSceneBoolTrack::StaticClass()->GetName(),  GET_MEMBER_NAME_STRING_CHECKED(UStaticMeshComponent, bEvaluateWorldPositionOffsetInRayTracing) , TEXT("Evaluate World Position Offset in Ray Tracing")}},
+			{EInterchangePropertyTracks::StaticMeshForcedLodModel, {UMovieSceneIntegerTrack::StaticClass()->GetName(),  GET_MEMBER_NAME_STRING_CHECKED(UStaticMeshComponent, ForcedLodModel) , TEXT("Forced Lod Model")}},
+			{EInterchangePropertyTracks::StaticMeshReverseCulling, {UMovieSceneBoolTrack::StaticClass()->GetName(),  GET_MEMBER_NAME_STRING_CHECKED(UStaticMeshComponent, bReverseCulling) , TEXT("Reverse Culling")}},
+			{EInterchangePropertyTracks::StaticMeshWorldPositionOffsetDisableDistance, {UMovieSceneIntegerTrack::StaticClass()->GetName(),  GET_MEMBER_NAME_STRING_CHECKED(UStaticMeshComponent, WorldPositionOffsetDisableDistance) , TEXT("World Position Offset Disable Distance")}},			
 		}
 	{}
 
diff --git a/Engine/Plugins/Interchange/Runtime/Source/Nodes/Public/InterchangeAnimationDefinitions.h b/Engine/Plugins/Interchange/Runtime/Source/Nodes/Public/InterchangeAnimationDefinitions.h
--- a/Engine/Plugins/Interchange/Runtime/Source/Nodes/Public/InterchangeAnimationDefinitions.h
+++ b/Engine/Plugins/Interchange/Runtime/Source/Nodes/Public/InterchangeAnimationDefinitions.h
@@ -34,6 +34,7 @@ enum class EInterchangePropertyTracks : int32
 	DrawDebug,
 	EmissiveLightSource,
 	ExcludeFromLightAttachmentGroup,
+	HiddenInGame, // This visibility property is set on the Component
 	HiddenInSceneCapture,
 	Holdout,
 	LightAttachmentsAsGroup,
@@ -48,7 +49,7 @@ enum class EInterchangePropertyTracks : int32
 	TranslucencySortDistanceOffset,
 	VisibleInRayTracing,
 	VisibleInSceneCaptureOnly,
-	Visibility,
+	Visibility, // This visibility property is set on the Actor
 
 	/** Light Properties */
 	LightAffectGlobalIllumination,
@@ -313,6 +314,30 @@ enum class EInterchangePropertyTracks : int32
 	CameraUseCameraHeightAsViewTarget,
 	CameraUseFieldOfViewForLOD,
 
+	/** Common Mesh Properties */
+	MeshOverlayMaterialMaxDrawDistance, 
+
+	/** Skinned Mesh Properties */
+	SkinnedMeshCapsuleIndirectShadowMinVisibility,
+	SkinnedMeshCastCapsuleDirectShadow,
+	SkinnedMeshCastCapsuleIndirectShadow,
+	SkinnedMeshRenderStatic,
+	SkinnedMeshVisibilityBasedAnimTickOption,
+
+	/** Skeletal Mesh Properties */
+	SkeletalMeshAllowClothActors,
+	SkeletalMeshAnimationMode,
+	SkeletalMeshClothBlendWeight,
+	SkeletalMeshClothMaxDistanceScale,
+
+	/** Static Mesh Properties */
+	StaticMeshDistanceFieldSelfShadowBias,
+	StaticMeshEvaluateWorldPositionOffset,
+	StaticMeshEvaluateWorldPositionOffsetInRayTracing,
+	StaticMeshForcedLodModel,
+	StaticMeshReverseCulling,
+	StaticMeshWorldPositionOffsetDisableDistance,
+
 	None = -1 UMETA(hidden),
 };
 
