Interchange: fix skeletalmesh import vertex color space.

#rnx

diff --git a/Engine/Plugins/Interchange/Runtime/Source/Import/Private/Mesh/InterchangeSkeletalMeshFactory.cpp b/Engine/Plugins/Interchange/Runtime/Source/Import/Private/Mesh/InterchangeSkeletalMeshFactory.cpp
--- a/Engine/Plugins/Interchange/Runtime/Source/Import/Private/Mesh/InterchangeSkeletalMeshFactory.cpp
+++ b/Engine/Plugins/Interchange/Runtime/Source/Import/Private/Mesh/InterchangeSkeletalMeshFactory.cpp
@@ -604,6 +604,23 @@ namespace UE
 					}
 				}
 
+				//The color data is linearized twice by the translator, we need to convert to sRGB to have proper linear
+				//TODO: change the translator to put linear instead of linear of linear
+				//      remove the ToFColor in StaticMeshBuilder
+				//      version the meshdescription properly to force the ToFColor when loading old static meshdescription
+				{
+					TVertexInstanceAttributesRef<FVector4f> VertexColor = SkeletalMeshAttributes.GetVertexInstanceColors();
+					for (FVertexInstanceID VertexInstanceID : LodMeshDescription.VertexInstances().GetElementIDs())
+					{
+						const FColor LinearUint8Color = FLinearColor(VertexColor[VertexInstanceID]).ToFColor(true);
+						constexpr float ColorScale = (1.0f / 255.0f);
+						VertexColor[VertexInstanceID] = FVector4f(static_cast<float>(LinearUint8Color.R) * ColorScale
+							, static_cast<float>(LinearUint8Color.G) * ColorScale
+							, static_cast<float>(LinearUint8Color.B) * ColorScale
+							, static_cast<float>(LinearUint8Color.A) * ColorScale);
+					}
+				}
+
 				DestinationImportData = FSkeletalMeshImportData::CreateFromMeshDescription(LodMeshDescription);
 				DestinationImportData.RefBonesBinary = RefBonesBinary;
 
