Interchange: Make sure the bImportAnimation is set to true if we re-import an animsequence.

#rnx

diff --git a/Engine/Plugins/Interchange/Runtime/Source/Pipelines/Private/InterchangeGenericAnimationPipeline.cpp b/Engine/Plugins/Interchange/Runtime/Source/Pipelines/Private/InterchangeGenericAnimationPipeline.cpp
--- a/Engine/Plugins/Interchange/Runtime/Source/Pipelines/Private/InterchangeGenericAnimationPipeline.cpp
+++ b/Engine/Plugins/Interchange/Runtime/Source/Pipelines/Private/InterchangeGenericAnimationPipeline.cpp
@@ -64,6 +64,7 @@ void UInterchangeGenericAnimationPipeline::AdjustSettingsForContext(EInterchange
 			//Set the skeleton to the current asset skeleton and re-import only the animation
 			CommonSkeletalMeshesAndAnimationsProperties->Skeleton = AnimSequence->GetSkeleton();
 			CommonSkeletalMeshesAndAnimationsProperties->bImportOnlyAnimations = true;
+			bImportAnimations = true;
 		}
 		else
 		{
@@ -71,6 +72,11 @@ void UInterchangeGenericAnimationPipeline::AdjustSettingsForContext(EInterchange
 		}
 	}
 
+	if (CommonSkeletalMeshesAndAnimationsProperties->bImportOnlyAnimations)
+	{
+		bImportAnimations = true;
+	}
+
 	if (UInterchangePipelineBase* OuterMostPipeline = GetMostPipelineOuter())
 	{
 		for (const FString& HideCategoryName : HideCategories)
@@ -91,6 +97,11 @@ void UInterchangeGenericAnimationPipeline::ExecutePipeline(UInterchangeBaseNodeC
 
 	BaseNodeContainer = InBaseNodeContainer;
 
+	if (CommonSkeletalMeshesAndAnimationsProperties->bImportOnlyAnimations)
+	{
+		bImportAnimations = true;
+	}
+
 	if (!bImportAnimations)
 	{
 		//Nothing to import
diff --git a/Engine/Plugins/Interchange/Runtime/Source/Pipelines/Private/InterchangeGenericAssetsPipeline.cpp b/Engine/Plugins/Interchange/Runtime/Source/Pipelines/Private/InterchangeGenericAssetsPipeline.cpp
--- a/Engine/Plugins/Interchange/Runtime/Source/Pipelines/Private/InterchangeGenericAssetsPipeline.cpp
+++ b/Engine/Plugins/Interchange/Runtime/Source/Pipelines/Private/InterchangeGenericAssetsPipeline.cpp
@@ -1090,6 +1090,10 @@ void UInterchangeGenericAssetsPipeline::ExecutePipeline(UInterchangeBaseNodeCont
 	//When we import only animation we need to prevent material and physic asset to be created
 	if (CommonSkeletalMeshesAndAnimationsProperties->bImportOnlyAnimations)
 	{
+		//We need to turn on ImportAnimation
+		AnimationPipeline->bImportAnimations = true;
+		
+		//Turn off mesh, material, texture, physics
 		MaterialPipeline->bImportMaterials = false;
 		MeshPipeline->bImportStaticMeshes = false;
 		MeshPipeline->bCreatePhysicsAsset = false;
