Interchange: Auto detect animation only asset import.

#rnx

diff --git a/Engine/Plugins/Interchange/Runtime/Source/Pipelines/Private/InterchangeGenericAnimationPipeline.cpp b/Engine/Plugins/Interchange/Runtime/Source/Pipelines/Private/InterchangeGenericAnimationPipeline.cpp
--- a/Engine/Plugins/Interchange/Runtime/Source/Pipelines/Private/InterchangeGenericAnimationPipeline.cpp
+++ b/Engine/Plugins/Interchange/Runtime/Source/Pipelines/Private/InterchangeGenericAnimationPipeline.cpp
@@ -3,6 +3,9 @@
 
 #include "Animation/AnimationSettings.h"
 #include "CoreMinimal.h"
+#include "Engine/SkeletalMesh.h"
+#include "Engine/StaticMesh.h"
+#include "InterchangeGenericMeshPipeline.h"
 #include "InterchangeLevelSequenceFactoryNode.h"
 #include "InterchangeAnimationTrackSetNode.h"
 #include "InterchangeAnimSequenceFactoryNode.h"
@@ -21,6 +24,49 @@
 
 #include UE_INLINE_GENERATED_CPP_BY_NAME(InterchangeGenericAnimationPipeline)
 
+namespace UE::Interchange::Private
+{
+	bool IsTranslatedDataContainOnlyJointAnimation(const UInterchangeBaseNodeContainer* InBaseNodeContainer)
+	{
+		//Its valid to call GetMeshesInformationFromTranslatedData with a null container
+		if (!InBaseNodeContainer)
+		{
+			return false;
+		}
+		bool bContainOnlyJointAnimation = false;
+		bool bContainJointAnimation = false;
+		InBaseNodeContainer->BreakableIterateNodesOfType<UInterchangeSkeletalAnimationTrackNode>([&bContainJointAnimation](const FString& NodeUid, UInterchangeSkeletalAnimationTrackNode* AnimationNode)
+			{
+				bContainJointAnimation = true;
+				return bContainJointAnimation;
+			});
+		if (bContainJointAnimation)
+		{
+			//if we have bone animation and no skinned mesh, we want to import animation only.
+			bool bContainSkinnedMeshNode = false;
+			InBaseNodeContainer->BreakableIterateNodesOfType<UInterchangeMeshNode>([&bContainSkinnedMeshNode](const FString& NodeUid, UInterchangeMeshNode* MeshNode)
+				{
+					if (!MeshNode->IsMorphTarget())
+					{
+						if (MeshNode->IsSkinnedMesh())
+						{
+							bContainSkinnedMeshNode = true;
+						}
+					}
+					return bContainSkinnedMeshNode;
+				});
+			bContainOnlyJointAnimation = !bContainSkinnedMeshNode;
+		}
+
+		return bContainOnlyJointAnimation;
+	}
+}
+
+FString UInterchangeGenericAnimationPipeline::GetPipelineCategory(UClass* AssetClass)
+{
+	return TEXT("Animations");
+}
+
 #if WITH_EDITOR
 bool UInterchangeGenericAnimationPipeline::CanEditChange(const FProperty* InProperty) const
 {
@@ -56,7 +102,26 @@ void UInterchangeGenericAnimationPipeline::AdjustSettingsForContext(EInterchange
 		CommonSkeletalMeshesAndAnimationsProperties->bImportOnlyAnimations = false;
 	}
 	
+	const FString CommonMeshesCategory =  UInterchangeGenericCommonMeshesProperties::GetPipelineCategory(nullptr);
+	const FString StaticMeshesCategory = UInterchangeGenericMeshPipeline::GetPipelineCategory(UStaticMesh::StaticClass());
+	const FString SkeletalMeshesCategory = UInterchangeGenericMeshPipeline::GetPipelineCategory(USkeletalMesh::StaticClass());
+	const FString AnimationCategory = UInterchangeGenericAnimationPipeline::GetPipelineCategory(nullptr);
+
 	TArray<FString> HideCategories;
+	if(ImportType == EInterchangePipelineContext::AssetImport)
+	{
+		if(UE::Interchange::Private::IsTranslatedDataContainOnlyJointAnimation(InBaseNodeContainer))
+		{
+			bImportAnimations = true;
+			CommonSkeletalMeshesAndAnimationsProperties->bImportOnlyAnimations = true;
+
+			HideCategories.Add(StaticMeshesCategory);
+			HideCategories.Add(SkeletalMeshesCategory);
+			HideCategories.Add(CommonMeshesCategory);
+		}
+	}
+
+	
 	if (ImportType == EInterchangePipelineContext::AssetReimport)
 	{
 		if (UAnimSequence* AnimSequence = Cast<UAnimSequence>(ReimportAsset))
@@ -68,7 +133,7 @@ void UInterchangeGenericAnimationPipeline::AdjustSettingsForContext(EInterchange
 		}
 		else
 		{
-			HideCategories.Add(TEXT("Animations"));
+			HideCategories.Add(AnimationCategory);
 		}
 	}
 
diff --git a/Engine/Plugins/Interchange/Runtime/Source/Pipelines/Private/InterchangeGenericAssetsPipeline.cpp b/Engine/Plugins/Interchange/Runtime/Source/Pipelines/Private/InterchangeGenericAssetsPipeline.cpp
--- a/Engine/Plugins/Interchange/Runtime/Source/Pipelines/Private/InterchangeGenericAssetsPipeline.cpp
+++ b/Engine/Plugins/Interchange/Runtime/Source/Pipelines/Private/InterchangeGenericAssetsPipeline.cpp
@@ -3,6 +3,7 @@
 #include "InterchangeGenericAssetsPipeline.h"
 
 #include "Animation/Skeleton.h"
+#include "AssetRegistry/ARFilter.h"
 #include "AssetRegistry/AssetData.h"
 #include "AssetRegistry/AssetRegistryModule.h"
 #include "CoreMinimal.h"
@@ -73,9 +74,15 @@ UInterchangeGenericAssetsPipeline::UInterchangeGenericAssetsPipeline()
 
 void UInterchangeGenericAssetsPipeline::PreDialogCleanup(const FName PipelineStackName)
 {
-	check(CommonSkeletalMeshesAndAnimationsProperties)
-	//We always clean the pipeline skeleton when showing the dialog
+	check(CommonSkeletalMeshesAndAnimationsProperties);
+	//Set the pipeline skeleton using what we found in AdjustSettingsForContext
 	CommonSkeletalMeshesAndAnimationsProperties->Skeleton = nullptr;
+	if (ContentPathExistingSkeleton.IsValid())
+	{
+		CommonSkeletalMeshesAndAnimationsProperties->Skeleton = Cast<USkeleton>(ContentPathExistingSkeleton.TryLoad());
+	}
+
+	CommonSkeletalMeshesAndAnimationsProperties->bImportOnlyAnimations = bImportOnlyAnimationAdjusted;
 
 	if (MaterialPipeline)
 	{
@@ -92,7 +99,10 @@ void UInterchangeGenericAssetsPipeline::PreDialogCleanup(const FName PipelineSta
 		AnimationPipeline->PreDialogCleanup(PipelineStackName);
 	}
 	
-	SaveSettings(PipelineStackName);
+	if (IsStandAlonePipeline())
+	{
+		SaveSettings(PipelineStackName);
+	}
 }
 
 bool UInterchangeGenericAssetsPipeline::IsSettingsAreValid(TOptional<FText>& OutInvalidReason) const
@@ -129,6 +139,26 @@ void UInterchangeGenericAssetsPipeline::AdjustSettingsForContext(EInterchangePip
 {
 	Super::AdjustSettingsForContext(ImportType, ReimportAsset, InBaseNodeContainer);
 
+	if (ImportType == EInterchangePipelineContext::AssetImport)
+	{
+		if (!CommonSkeletalMeshesAndAnimationsProperties->Skeleton.IsValid())
+		{
+			//Search if there is a valid skeleton in this path
+			FARFilter Filter;
+			Filter.PackagePaths.Add(*ContentImportPath);
+			Filter.ClassPaths.Add(USkeleton::StaticClass()->GetClassPathName());
+
+			IAssetRegistry& AssetRegistry = FModuleManager::LoadModuleChecked<FAssetRegistryModule>("AssetRegistry").Get();
+			TArray<FAssetData> SkeletonAssets;
+			AssetRegistry.GetAssets(Filter, SkeletonAssets);
+			if (SkeletonAssets.Num() > 0)
+			{
+				CommonSkeletalMeshesAndAnimationsProperties->Skeleton = CastChecked<USkeleton>(SkeletonAssets[0].GetAsset());
+				ContentPathExistingSkeleton = CommonSkeletalMeshesAndAnimationsProperties->Skeleton.Get();
+			}
+		}
+	}
+
 	if (MaterialPipeline)
 	{
 		MaterialPipeline->AdjustSettingsForContext(ImportType, ReimportAsset, InBaseNodeContainer);
@@ -143,6 +173,9 @@ void UInterchangeGenericAssetsPipeline::AdjustSettingsForContext(EInterchangePip
 	{
 		AnimationPipeline->AdjustSettingsForContext(ImportType, ReimportAsset, InBaseNodeContainer);
 	}
+
+	//Store the adjusted settings
+	bImportOnlyAnimationAdjusted = CommonSkeletalMeshesAndAnimationsProperties->bImportOnlyAnimations;
 }
 
 #if WITH_EDITOR
@@ -166,32 +199,36 @@ void UInterchangeGenericAssetsPipeline::FilterPropertiesFromTranslatedData(UInte
 		UInterchangePipelineMeshesUtilities* PipelineMeshesUtilities = UInterchangeGenericMeshPipeline::CreateMeshPipelineUtilities(InBaseNodeContainer, MeshPipeline);
 
 		TArray<FString> SkeletalMeshes;
-		PipelineMeshesUtilities->GetAllSkinnedMeshInstance(SkeletalMeshes);
-		if(SkeletalMeshes.Num() == 0)
-		{
-			PipelineMeshesUtilities->GetAllSkinnedMeshGeometry(SkeletalMeshes);
-		}
 		TArray<FString> StaticMeshes;
-		PipelineMeshesUtilities->GetAllStaticMeshInstance(StaticMeshes);
-		if(StaticMeshes.Num() == 0)
-		{
-			PipelineMeshesUtilities->GetAllStaticMeshGeometry(StaticMeshes);
-		}
-
 		int32 RawStaticMesh = 0;
 		int32 RawSkeletalMesh = 0;
 		int32 RawMorphTargetShape = 0;
-		InBaseNodeContainer->IterateNodesOfType<UInterchangeMeshNode>([&RawStaticMesh, &RawSkeletalMesh, &RawMorphTargetShape](const FString& NodeUid, UInterchangeMeshNode* MeshNode)
+		if (!CommonSkeletalMeshesAndAnimationsProperties->bImportOnlyAnimations)
+		{
+			PipelineMeshesUtilities->GetAllSkinnedMeshInstance(SkeletalMeshes);
+			if (SkeletalMeshes.Num() == 0)
 			{
-				if (MeshNode->IsMorphTarget())
-				{
-					RawMorphTargetShape++;
-				}
-				else
+				PipelineMeshesUtilities->GetAllSkinnedMeshGeometry(SkeletalMeshes);
+			}
+			
+			PipelineMeshesUtilities->GetAllStaticMeshInstance(StaticMeshes);
+			if (StaticMeshes.Num() == 0)
+			{
+				PipelineMeshesUtilities->GetAllStaticMeshGeometry(StaticMeshes);
+			}
+
+			InBaseNodeContainer->IterateNodesOfType<UInterchangeMeshNode>([&RawStaticMesh, &RawSkeletalMesh, &RawMorphTargetShape](const FString& NodeUid, UInterchangeMeshNode* MeshNode)
 				{
-					MeshNode->IsSkinnedMesh() ? RawSkeletalMesh++ : RawStaticMesh++;
-				}
-			});
+					if (MeshNode->IsMorphTarget())
+					{
+						RawMorphTargetShape++;
+					}
+					else
+					{
+						MeshNode->IsSkinnedMesh() ? RawSkeletalMesh++ : RawStaticMesh++;
+					}
+				});
+		}
 
 		int32 RawAnimationNode = 0;
 		InBaseNodeContainer->IterateNodesOfType<UInterchangeAnimationTrackBaseNode>([&RawAnimationNode](const FString& NodeUid, UInterchangeAnimationTrackBaseNode* AnimationNode)
@@ -284,30 +321,47 @@ void UInterchangeGenericAssetsPipeline::FilterPropertiesFromTranslatedData(UInte
 			bHideCommonSkeletalMeshesAndAnimations = true;
 		}
 
+
+		if (CommonSkeletalMeshesAndAnimationsProperties->bImportOnlyAnimations)
+		{
+			if (MaterialPipeline)
+			{
+				HideFullCategory(UInterchangeGenericMaterialPipeline::GetPipelineCategory(nullptr), MaterialPipeline);
+				if (MaterialPipeline->TexturePipeline)
+				{
+					HideFullCategory(UInterchangeGenericTexturePipeline::GetPipelineCategory(nullptr), MaterialPipeline->TexturePipeline);
+				}
+			}
+		}
+
 		//Hide the categories
 		if (bHideStaticMeshes)
 		{
-			HideFullCategory(TEXT("Static Meshes"), MeshPipeline);
+			HideFullCategory(UInterchangeGenericMeshPipeline::GetPipelineCategory(UStaticMesh::StaticClass()), MeshPipeline);
 		}
 		if (bHideSkeletalMeshes)
 		{
-			HideFullCategory(TEXT("Skeletal Meshes"), MeshPipeline);
+			HideFullCategory(UInterchangeGenericMeshPipeline::GetPipelineCategory(USkeletalMesh::StaticClass()), MeshPipeline);
 		}
 		if(bHideCommonMeshes)
 		{
-			HideFullCategory(TEXT("Common Meshes"), CommonMeshesProperties);
+			HideFullCategory(UInterchangeGenericCommonMeshesProperties::GetPipelineCategory(nullptr), CommonMeshesProperties);
 		}
 		if (bHideCommonSkeletalMeshesAndAnimations)
 		{
-			HideFullCategory(TEXT("Common Skeletal Meshes and Animations"), CommonSkeletalMeshesAndAnimationsProperties);
+			//We cannot hide common sk and anim if import only animation is true.
+			if (!CommonSkeletalMeshesAndAnimationsProperties->bImportOnlyAnimations)
+			{
+				HideFullCategory(UInterchangeGenericCommonSkeletalMeshesAndAnimationsProperties::GetPipelineCategory(nullptr), CommonSkeletalMeshesAndAnimationsProperties);
+			}
 		}
 		if (bHideCommonSkeletalMeshesAndAnimations_StaticMesh)
 		{
-			HideFullCategory(TEXT("Static Meshes"), CommonSkeletalMeshesAndAnimationsProperties);
+			HideFullCategory(UInterchangeGenericMeshPipeline::GetPipelineCategory(UStaticMesh::StaticClass()), CommonSkeletalMeshesAndAnimationsProperties);
 		}
 		if (bHideAnimations)
 		{
-			HideFullCategory(TEXT("Animations"), AnimationPipeline);
+			HideFullCategory(UInterchangeGenericAnimationPipeline::GetPipelineCategory(nullptr), AnimationPipeline);
 		}
 	}
 }
diff --git a/Engine/Plugins/Interchange/Runtime/Source/Pipelines/Private/InterchangeGenericMaterialPipeline.cpp b/Engine/Plugins/Interchange/Runtime/Source/Pipelines/Private/InterchangeGenericMaterialPipeline.cpp
--- a/Engine/Plugins/Interchange/Runtime/Source/Pipelines/Private/InterchangeGenericMaterialPipeline.cpp
+++ b/Engine/Plugins/Interchange/Runtime/Source/Pipelines/Private/InterchangeGenericMaterialPipeline.cpp
@@ -582,6 +582,11 @@ UInterchangeGenericMaterialPipeline::UInterchangeGenericMaterialPipeline()
 	TexturePipeline = CreateDefaultSubobject<UInterchangeGenericTexturePipeline>("TexturePipeline");
 }
 
+FString UInterchangeGenericMaterialPipeline::GetPipelineCategory(UClass* AssetClass)
+{
+	return TEXT("Materials");
+}
+
 void UInterchangeGenericMaterialPipeline::PreDialogCleanup(const FName PipelineStackName)
 {
 	if (TexturePipeline)
@@ -589,7 +594,11 @@ void UInterchangeGenericMaterialPipeline::PreDialogCleanup(const FName PipelineS
 		TexturePipeline->PreDialogCleanup(PipelineStackName);
 	}
 
-	SaveSettings(PipelineStackName);
+	//Save only pipeline if we are a stand alone pipeline (not a sub object of another pipeline)
+	if (IsStandAlonePipeline())
+	{
+		SaveSettings(PipelineStackName);
+	}
 }
 
 bool UInterchangeGenericMaterialPipeline::IsSettingsAreValid(TOptional<FText>& OutInvalidReason) const
@@ -619,7 +628,7 @@ void UInterchangeGenericMaterialPipeline::AdjustSettingsForContext(EInterchangeP
 		|| ImportType == EInterchangePipelineContext::AssetAlternateSkinningReimport)
 	{
 		bImportMaterials = false;
-		HideCategories.Add(TEXT("Materials"));
+		HideCategories.Add(UInterchangeGenericMaterialPipeline::GetPipelineCategory(nullptr));
 		SearchLocation = EInterchangeMaterialSearchLocation::DoNotSearch;
 	}
 
diff --git a/Engine/Plugins/Interchange/Runtime/Source/Pipelines/Private/InterchangeGenericMeshPipeline.cpp b/Engine/Plugins/Interchange/Runtime/Source/Pipelines/Private/InterchangeGenericMeshPipeline.cpp
--- a/Engine/Plugins/Interchange/Runtime/Source/Pipelines/Private/InterchangeGenericMeshPipeline.cpp
+++ b/Engine/Plugins/Interchange/Runtime/Source/Pipelines/Private/InterchangeGenericMeshPipeline.cpp
@@ -21,6 +21,22 @@
 
 #include UE_INLINE_GENERATED_CPP_BY_NAME(InterchangeGenericMeshPipeline)
 
+FString UInterchangeGenericMeshPipeline::GetPipelineCategory(UClass* AssetClass)
+{
+	if (ensure(AssetClass))
+	{
+		if (AssetClass->IsChildOf(UStaticMesh::StaticClass()))
+		{
+			return TEXT("Static Meshes");
+		}
+		else if (AssetClass->IsChildOf(USkeletalMesh::StaticClass()))
+		{
+			return TEXT("Skeletal Meshes");
+		}
+	}
+	return TEXT("Static Meshes");
+}
+
 void UInterchangeGenericMeshPipeline::AdjustSettingsForContext(EInterchangePipelineContext ImportType, TObjectPtr<UObject> ReimportAsset, const UInterchangeBaseNodeContainer* InBaseNodeContainer)
 {
 	Super::AdjustSettingsForContext(ImportType, ReimportAsset, InBaseNodeContainer);
@@ -72,10 +88,10 @@ void UInterchangeGenericMeshPipeline::AdjustSettingsForContext(EInterchangePipel
 			CommonSkeletalMeshesAndAnimationsProperties->bImportOnlyAnimations = false;
 		}
 	}
-	const FString CommonMeshesCategory = TEXT("Common Meshes");
-	const FString StaticMeshesCategory = TEXT("Static Meshes");
-	const FString SkeletalMeshesCategory = TEXT("Skeletal Meshes");
-	const FString CommonSkeletalMeshesAndAnimationCategory = TEXT("Common Skeletal Meshes and Animations");
+	const FString CommonMeshesCategory = UInterchangeGenericCommonMeshesProperties::GetPipelineCategory(nullptr);
+	const FString StaticMeshesCategory = UInterchangeGenericMeshPipeline::GetPipelineCategory(UStaticMesh::StaticClass());
+	const FString SkeletalMeshesCategory = UInterchangeGenericMeshPipeline::GetPipelineCategory(USkeletalMesh::StaticClass());
+	const FString CommonSkeletalMeshesAndAnimationCategory = UInterchangeGenericCommonSkeletalMeshesAndAnimationsProperties::GetPipelineCategory(nullptr);
 
 	TArray<FString> HideCategories;
 	TArray<FString> HideSubCategories;
diff --git a/Engine/Plugins/Interchange/Runtime/Source/Pipelines/Private/InterchangeGenericTexturePipeline.cpp b/Engine/Plugins/Interchange/Runtime/Source/Pipelines/Private/InterchangeGenericTexturePipeline.cpp
--- a/Engine/Plugins/Interchange/Runtime/Source/Pipelines/Private/InterchangeGenericTexturePipeline.cpp
+++ b/Engine/Plugins/Interchange/Runtime/Source/Pipelines/Private/InterchangeGenericTexturePipeline.cpp
@@ -110,6 +110,11 @@ namespace UE::Interchange::Private
 	}
 }
 
+FString UInterchangeGenericTexturePipeline::GetPipelineCategory(UClass* AssetClass)
+{
+	return TEXT("Textures");
+}
+
 void UInterchangeGenericTexturePipeline::AdjustSettingsForContext(EInterchangePipelineContext ImportType, TObjectPtr<UObject> ReimportAsset, const UInterchangeBaseNodeContainer* InBaseNodeContainer)
 {
 	Super::AdjustSettingsForContext(ImportType, ReimportAsset, InBaseNodeContainer);
@@ -123,7 +128,7 @@ void UInterchangeGenericTexturePipeline::AdjustSettingsForContext(EInterchangePi
 		|| ImportType == EInterchangePipelineContext::AssetAlternateSkinningReimport)
 	{
 		bImportTextures = false;
-		HideCategories.Add(TEXT("Textures"));
+		HideCategories.Add(UInterchangeGenericTexturePipeline::GetPipelineCategory(nullptr));
 	}
 	if (UInterchangePipelineBase* OuterMostPipeline = GetMostPipelineOuter())
 	{
@@ -149,7 +154,7 @@ void UInterchangeGenericTexturePipeline::FilterPropertiesFromTranslatedData(UInt
 		//Filter out all Textures properties
 		if (UInterchangePipelineBase* OuterMostPipeline = GetMostPipelineOuter())
 		{
-			HidePropertiesOfCategory(OuterMostPipeline, this, TEXT("Textures"));
+			HidePropertiesOfCategory(OuterMostPipeline, this, UInterchangeGenericTexturePipeline::GetPipelineCategory(nullptr));
 		}
 	}
 }
diff --git a/Engine/Plugins/Interchange/Runtime/Source/Pipelines/Public/InterchangeGenericAnimationPipeline.h b/Engine/Plugins/Interchange/Runtime/Source/Pipelines/Public/InterchangeGenericAnimationPipeline.h
--- a/Engine/Plugins/Interchange/Runtime/Source/Pipelines/Public/InterchangeGenericAnimationPipeline.h
+++ b/Engine/Plugins/Interchange/Runtime/Source/Pipelines/Public/InterchangeGenericAnimationPipeline.h
@@ -41,6 +41,8 @@ public:
 	virtual bool CanEditChange(const FProperty* InProperty) const override;
 #endif
 
+	static FString GetPipelineCategory(UClass* AssetClass);
+
 	//Common SkeletalMeshes And Animations Properties Settings Pointer
 	UPROPERTY(Transient)
 	TWeakObjectPtr<UInterchangeGenericCommonSkeletalMeshesAndAnimationsProperties> CommonSkeletalMeshesAndAnimationsProperties;
diff --git a/Engine/Plugins/Interchange/Runtime/Source/Pipelines/Public/InterchangeGenericAssetsPipeline.h b/Engine/Plugins/Interchange/Runtime/Source/Pipelines/Public/InterchangeGenericAssetsPipeline.h
--- a/Engine/Plugins/Interchange/Runtime/Source/Pipelines/Public/InterchangeGenericAssetsPipeline.h
+++ b/Engine/Plugins/Interchange/Runtime/Source/Pipelines/Public/InterchangeGenericAssetsPipeline.h
@@ -130,6 +130,14 @@ public:
 		TArray<TSharedPtr<FSkeletonJoint>> Children;
 	};
 
+	//We need to store the adjusted content path existing skeleton to restore it in PreDialogCleanup
+	UPROPERTY(meta = (AlwaysResetToDefault = "True"))
+	FSoftObjectPath ContentPathExistingSkeleton;
+
+	//We need to store the adjusted import only animation boolean to restore it in PreDialogCleanup
+	UPROPERTY(meta = (AlwaysResetToDefault = "True"))
+	bool bImportOnlyAnimationAdjusted = false;
+
 protected:
 
 	virtual void ExecutePipeline(UInterchangeBaseNodeContainer* InBaseNodeContainer, const TArray<UInterchangeSourceData*>& InSourceDatas, const FString& ContentBasePath) override;
diff --git a/Engine/Plugins/Interchange/Runtime/Source/Pipelines/Public/InterchangeGenericAssetsPipelineSharedSettings.h b/Engine/Plugins/Interchange/Runtime/Source/Pipelines/Public/InterchangeGenericAssetsPipelineSharedSettings.h
--- a/Engine/Plugins/Interchange/Runtime/Source/Pipelines/Public/InterchangeGenericAssetsPipelineSharedSettings.h
+++ b/Engine/Plugins/Interchange/Runtime/Source/Pipelines/Public/InterchangeGenericAssetsPipelineSharedSettings.h
@@ -51,6 +51,12 @@ class INTERCHANGEPIPELINES_API UInterchangeGenericCommonMeshesProperties : publi
 {
 	GENERATED_BODY()
 public:
+
+	static FString GetPipelineCategory(UClass* AssetClass)
+	{
+		return TEXT("Common Meshes");
+	}
+
 	//////	COMMON_MESHES_CATEGORY Properties //////
 
 	/**
@@ -144,6 +150,12 @@ class INTERCHANGEPIPELINES_API UInterchangeGenericCommonSkeletalMeshesAndAnimati
 {
 	GENERATED_BODY()
 public:
+
+	static FString GetPipelineCategory(UClass* AssetClass)
+	{
+		return TEXT("Common Skeletal Meshes and Animations");
+	}
+
 	/** If enabled, only animations are imported from the source. You must also set a valid skeleton. */
 	UPROPERTY(EditAnywhere, BlueprintReadWrite, Category = "Common Skeletal Meshes and Animations")
 	bool bImportOnlyAnimations = false;
@@ -181,6 +193,10 @@ public:
 		{
 			return true;
 		}
+		else if (PropertyChangedEvent.Property->GetName() == GET_MEMBER_NAME_CHECKED(UInterchangeGenericCommonSkeletalMeshesAndAnimationsProperties, bImportOnlyAnimations))
+		{
+			return true;
+		}
 		return Super::IsPropertyChangeNeedRefresh(PropertyChangedEvent);
 	}
 #endif //WITH_EDITOR
diff --git a/Engine/Plugins/Interchange/Runtime/Source/Pipelines/Public/InterchangeGenericMaterialPipeline.h b/Engine/Plugins/Interchange/Runtime/Source/Pipelines/Public/InterchangeGenericMaterialPipeline.h
--- a/Engine/Plugins/Interchange/Runtime/Source/Pipelines/Public/InterchangeGenericMaterialPipeline.h
+++ b/Engine/Plugins/Interchange/Runtime/Source/Pipelines/Public/InterchangeGenericMaterialPipeline.h
@@ -61,6 +61,8 @@ class INTERCHANGEPIPELINES_API UInterchangeGenericMaterialPipeline : public UInt
 public:
 	UInterchangeGenericMaterialPipeline();
 
+	static FString GetPipelineCategory(UClass* AssetClass);
+
 	/** The name of the pipeline that will be display in the import dialog. */
 	UPROPERTY(EditAnywhere, BlueprintReadWrite, Category = "Materials", meta = (StandAlonePipelineProperty = "True", PipelineInternalEditionData = "True"))
 	FString PipelineDisplayName;
diff --git a/Engine/Plugins/Interchange/Runtime/Source/Pipelines/Public/InterchangeGenericMeshPipeline.h b/Engine/Plugins/Interchange/Runtime/Source/Pipelines/Public/InterchangeGenericMeshPipeline.h
--- a/Engine/Plugins/Interchange/Runtime/Source/Pipelines/Public/InterchangeGenericMeshPipeline.h
+++ b/Engine/Plugins/Interchange/Runtime/Source/Pipelines/Public/InterchangeGenericMeshPipeline.h
@@ -30,6 +30,8 @@ class INTERCHANGEPIPELINES_API UInterchangeGenericMeshPipeline : public UInterch
 	GENERATED_BODY()
 
 public:
+	//IInterchangeGenericPipelineCategoryInterface
+	static FString GetPipelineCategory(UClass* AssetClass);
 
 	//Common Meshes Properties Settings Pointer
 	UPROPERTY(Transient)
diff --git a/Engine/Plugins/Interchange/Runtime/Source/Pipelines/Public/InterchangeGenericTexturePipeline.h b/Engine/Plugins/Interchange/Runtime/Source/Pipelines/Public/InterchangeGenericTexturePipeline.h
--- a/Engine/Plugins/Interchange/Runtime/Source/Pipelines/Public/InterchangeGenericTexturePipeline.h
+++ b/Engine/Plugins/Interchange/Runtime/Source/Pipelines/Public/InterchangeGenericTexturePipeline.h
@@ -19,6 +19,8 @@ class INTERCHANGEPIPELINES_API UInterchangeGenericTexturePipeline : public UInte
 	GENERATED_BODY()
 
 public:
+	static FString GetPipelineCategory(UClass* AssetClass);
+
 	/** The name of the pipeline that will be display in the import dialog. */
 	UPROPERTY(EditAnywhere, BlueprintReadWrite, Category = "Textures", meta = (StandAlonePipelineProperty = "True", PipelineInternalEditionData = "True"))
 	FString PipelineDisplayName;
diff --git a/Engine/Source/Runtime/Interchange/Core/Public/InterchangePipelineBase.h b/Engine/Source/Runtime/Interchange/Core/Public/InterchangePipelineBase.h
--- a/Engine/Source/Runtime/Interchange/Core/Public/InterchangePipelineBase.h
+++ b/Engine/Source/Runtime/Interchange/Core/Public/InterchangePipelineBase.h
@@ -462,6 +462,12 @@ public:
 	 */
 	FString DestinationName;
 
+	/*
+	 * The content path where asset should be created.
+	 */
+	UPROPERTY(Transient)
+	FString ContentImportPath;
+
 protected:
 
 	INTERCHANGECORE_API virtual FString GetPipelineDisplayName() const;
diff --git a/Engine/Source/Runtime/Interchange/Engine/Private/InterchangeManager.cpp b/Engine/Source/Runtime/Interchange/Engine/Private/InterchangeManager.cpp
--- a/Engine/Source/Runtime/Interchange/Engine/Private/InterchangeManager.cpp
+++ b/Engine/Source/Runtime/Interchange/Engine/Private/InterchangeManager.cpp
@@ -1737,7 +1737,7 @@ UInterchangeManager::ImportInternal(const FString& ContentPath, const UInterchan
 	
 	UInterchangeBaseNodeContainer* BaseNodeContainer = AsyncHelper->BaseNodeContainers[SourceIndex].Get();
 
-	auto AdjustPipelineSettingForContext = [bIsReimport, bImportScene, &TaskData, BaseNodeContainer](UInterchangePipelineBase* Pipeline)
+	auto AdjustPipelineSettingForContext = [bIsReimport, bImportScene, &TaskData, BaseNodeContainer, &ContentBasePath](UInterchangePipelineBase* Pipeline)
 		{
 			EInterchangePipelineContext Context;
 			if (bIsReimport)
@@ -1748,7 +1748,7 @@ UInterchangeManager::ImportInternal(const FString& ContentPath, const UInterchan
 			{
 				Context = bImportScene ? EInterchangePipelineContext::SceneImport : EInterchangePipelineContext::AssetImport;
 			}
-
+			Pipeline->ContentImportPath = ContentBasePath;
 			Pipeline->AdjustSettingsForContext(Context, TaskData.ReimportObject, BaseNodeContainer);
 			Pipeline->DestinationName = TaskData.DestinationName;
 		};
