Interchange: downgrade warning fbx missing texture to Display log level

#rnx

diff --git a/Engine/Plugins/Interchange/Runtime/Source/Messages/Private/Fbx/InterchangeFbxMessages.cpp b/Engine/Plugins/Interchange/Runtime/Source/Messages/Private/Fbx/InterchangeFbxMessages.cpp
--- a/Engine/Plugins/Interchange/Runtime/Source/Messages/Private/Fbx/InterchangeFbxMessages.cpp
+++ b/Engine/Plugins/Interchange/Runtime/Source/Messages/Private/Fbx/InterchangeFbxMessages.cpp
@@ -42,7 +42,7 @@ FText UInterchangeResultMeshWarning_TooManyUVs::GetText() const
 	return FText::Format(LOCTEXT("TooManyUVs", "Reached the maximum number of UV Channels for mesh '{MeshName}' - discarding {ExcessUVs} {ExcessUVs}|plural(one=channel,other=channels)."), Args);
 }
 
-FText UInterchangeResultTextureWarning_TextureFileDoNotExist::GetText() const
+FText UInterchangeResultTextureDisplay_TextureFileDoNotExist::GetText() const
 {
 	FFormatNamedArguments Args
 	{
diff --git a/Engine/Plugins/Interchange/Runtime/Source/Messages/Public/Fbx/InterchangeFbxMessages.h b/Engine/Plugins/Interchange/Runtime/Source/Messages/Public/Fbx/InterchangeFbxMessages.h
--- a/Engine/Plugins/Interchange/Runtime/Source/Messages/Public/Fbx/InterchangeFbxMessages.h
+++ b/Engine/Plugins/Interchange/Runtime/Source/Messages/Public/Fbx/InterchangeFbxMessages.h
@@ -22,6 +22,20 @@ public:
 	FString MeshName;
 };
 
+/**
+ * Base class for FBX parser warnings
+ */
+UCLASS()
+class INTERCHANGEMESSAGES_API UInterchangeResultTextureDisplay : public UInterchangeResultDisplay_Generic
+{
+	GENERATED_BODY()
+
+public:
+
+	UPROPERTY()
+	FString TextureName;
+};
+
 /**
  * Base class for FBX parser warnings
  */
@@ -102,16 +116,13 @@ public:
  * A generic class for FBX parser warnings, with no additional metadata, and where the text is specified by the user
  */
 UCLASS()
-class INTERCHANGEMESSAGES_API UInterchangeResultTextureWarning_TextureFileDoNotExist : public UInterchangeResultTextureWarning
+class INTERCHANGEMESSAGES_API UInterchangeResultTextureDisplay_TextureFileDoNotExist : public UInterchangeResultTextureDisplay
 {
 	GENERATED_BODY()
 
 public:
 	virtual FText GetText() const override;
 
-	UPROPERTY()
-	FText Text;
-
 	UPROPERTY()
 	FString MaterialName;
 };
\ No newline at end of file
diff --git a/Engine/Plugins/Interchange/Runtime/Source/Parsers/Fbx/Private/FbxMaterial.cpp b/Engine/Plugins/Interchange/Runtime/Source/Parsers/Fbx/Private/FbxMaterial.cpp
--- a/Engine/Plugins/Interchange/Runtime/Source/Parsers/Fbx/Private/FbxMaterial.cpp
+++ b/Engine/Plugins/Interchange/Runtime/Source/Parsers/Fbx/Private/FbxMaterial.cpp
@@ -93,7 +93,7 @@ namespace UE
 					const UInterchangeBaseNode* ShaderGraphNode = NodeContainer.GetNode(ShaderUniqueID);
 					if(!GIsAutomationTesting)
 					{
-						UInterchangeResultTextureWarning_TextureFileDoNotExist* Message = Parser.AddMessage<UInterchangeResultTextureWarning_TextureFileDoNotExist>();
+						UInterchangeResultTextureDisplay_TextureFileDoNotExist* Message = Parser.AddMessage<UInterchangeResultTextureDisplay_TextureFileDoNotExist>();
 						Message->TextureName = TextureFilename;
 						Message->MaterialName = ShaderGraphNode ? ShaderGraphNode->GetDisplayLabel() : TEXT("Unknown");
 					}
@@ -209,7 +209,7 @@ namespace UE
 				{
 					if (!GIsAutomationTesting)
 					{
-						UInterchangeResultTextureWarning_TextureFileDoNotExist* Message = Parser.AddMessage<UInterchangeResultTextureWarning_TextureFileDoNotExist>();
+						UInterchangeResultTextureDisplay_TextureFileDoNotExist* Message = Parser.AddMessage<UInterchangeResultTextureDisplay_TextureFileDoNotExist>();
 						Message->TextureName = FbxTexture ? UTF8_TO_TCHAR(FbxTexture->GetFileName()) : TEXT("Undefined");
 						Message->MaterialName = ShaderGraphNode->GetDisplayLabel();
 					}
@@ -472,7 +472,7 @@ namespace UE
 					{
 						if (!GIsAutomationTesting)
 						{
-							UInterchangeResultTextureWarning_TextureFileDoNotExist* Message = Parser.AddMessage<UInterchangeResultTextureWarning_TextureFileDoNotExist>();
+							UInterchangeResultTextureDisplay_TextureFileDoNotExist* Message = Parser.AddMessage<UInterchangeResultTextureDisplay_TextureFileDoNotExist>();
 							Message->TextureName = TextureFilename;
 							Message->MaterialName.Empty();
 						}
