Interchange: fix crash cause by missing remap of the bone when restoring the cloth. Also add missing sections restore data (Disable and generate UpTo)

#rnx

diff --git a/Engine/Plugins/Interchange/Runtime/Source/Import/InterchangeImport.Build.cs b/Engine/Plugins/Interchange/Runtime/Source/Import/InterchangeImport.Build.cs
--- a/Engine/Plugins/Interchange/Runtime/Source/Import/InterchangeImport.Build.cs
+++ b/Engine/Plugins/Interchange/Runtime/Source/Import/InterchangeImport.Build.cs
@@ -32,6 +32,7 @@ namespace UnrealBuildTool.Rules
 				{
 					"AssetRegistry",
 					"CinematicCamera",
+					"ClothingSystemRuntimeCommon",
 					"GLTFCore",
 					"IESFile",
 					"ImageCore",
diff --git a/Engine/Plugins/Interchange/Runtime/Source/Import/Private/Mesh/InterchangeSkeletalMeshFactory.cpp b/Engine/Plugins/Interchange/Runtime/Source/Import/Private/Mesh/InterchangeSkeletalMeshFactory.cpp
--- a/Engine/Plugins/Interchange/Runtime/Source/Import/Private/Mesh/InterchangeSkeletalMeshFactory.cpp
+++ b/Engine/Plugins/Interchange/Runtime/Source/Import/Private/Mesh/InterchangeSkeletalMeshFactory.cpp
@@ -5,6 +5,7 @@
 #include "AssetRegistry/AssetRegistryModule.h"
 #include "Async/Async.h"
 #include "Async/ParallelFor.h"
+#include "ClothingAsset.h"
 #include "Components.h"
 #include "CoreGlobals.h"
 #include "Engine/SkeletalMesh.h"
@@ -2144,6 +2145,13 @@ void UInterchangeSkeletalMeshFactory::FinalizeObject_GameThread(const FSetupObje
 		//Restore the clothing
 		if (ImportAssetObjectData.ExistingClothingBindings.Num() > 0)
 		{
+			for (ClothingAssetUtils::FClothingAssetMeshBinding& ExistingClothMeshBinding : ImportAssetObjectData.ExistingClothingBindings)
+			{
+				if (UClothingAssetCommon* ClothAssetCommon = ExistingClothMeshBinding.Asset)
+				{
+					ClothAssetCommon->RefreshBoneMapping(SkeletalMesh);
+				}
+			}
 			FSkeletalMeshModel* ImportedResource = SkeletalMesh->GetImportedModel();
 			for (int32 LodIndex = 0; LodIndex < ImportedResource->LODModels.Num(); ++LodIndex)
 			{
diff --git a/Engine/Source/Developer/SkeletalMeshUtilitiesCommon/Private/LODUtilities.cpp b/Engine/Source/Developer/SkeletalMeshUtilitiesCommon/Private/LODUtilities.cpp
--- a/Engine/Source/Developer/SkeletalMeshUtilitiesCommon/Private/LODUtilities.cpp
+++ b/Engine/Source/Developer/SkeletalMeshUtilitiesCommon/Private/LODUtilities.cpp
@@ -4154,6 +4154,8 @@ void FLODUtilities::MatchImportedMaterials(FLODUtilities::FSkeletalMeshMatchImpo
 						SourceSectionUserData.bVisibleInRayTracing = ExistingSection.bVisibleInRayTracing;
 						SourceSectionUserData.bRecomputeTangent = ExistingSection.bRecomputeTangent;
 						SourceSectionUserData.RecomputeTangentsVertexMaskChannel = ExistingSection.RecomputeTangentsVertexMaskChannel;
+						SourceSectionUserData.bDisabled = ExistingSection.bDisabled;
+						SourceSectionUserData.GenerateUpToLodIndex = ExistingSection.GenerateUpToLodIndex;
 					};
 
 				if (ExistingImportedMaterialSlotName != NAME_None)
