[Interchange] IK bone are not imported correctly for animations fix.

#rnx

diff --git a/Engine/Plugins/Interchange/Runtime/Source/Parsers/Fbx/Private/FbxAnimation.cpp b/Engine/Plugins/Interchange/Runtime/Source/Parsers/Fbx/Private/FbxAnimation.cpp
--- a/Engine/Plugins/Interchange/Runtime/Source/Parsers/Fbx/Private/FbxAnimation.cpp
+++ b/Engine/Plugins/Interchange/Runtime/Source/Parsers/Fbx/Private/FbxAnimation.cpp
@@ -457,10 +457,13 @@ namespace UE::Interchange::Private
 		FbxScene* SDKScene = nullptr;
 		FbxNode* Node = nullptr;
 		bool IsNodeAnimated = false;
+		bool IsEmptyAnimation = false;	//Report true if Node->Lcl transforms return empty (indicating identity Transform setting for the Animation),
+										//		while the Node (outside of Animation) has a non-identity transform
+										//(We only report true for AddSkeletalTransformAnimation for IsNodeAnimated.)
 		FbxAnimStack* CurrentAnimStack = nullptr;
 	};
 
-	void GetFbxTransformCurves(FGetFbxTransformCurvesParameters& Parameters, const int32& AnimationIndex)
+	void GetFbxTransformCurves(FGetFbxTransformCurvesParameters& Parameters, const int32& AnimationIndex, UInterchangeSceneNode* SceneNode)
 	{
 		if (!ensure(Parameters.SDKScene) || !ensure(Parameters.Node))
 		{
@@ -521,6 +524,27 @@ namespace UE::Interchange::Private
 				}
 			}
 		}
+
+		if (!Parameters.IsNodeAnimated)
+		{
+			//At this point this flag means that the Local Transform curves have no values.
+			// This could also mean that the node is animated to Identity in case the Node's original transform is not Identity:
+			FTransform LocalTransform;
+			if (SceneNode->GetCustomBindPoseLocalTransform(LocalTransform))
+			{
+				if (!LocalTransform.Equals(FTransform::Identity))
+				{
+					Parameters.IsEmptyAnimation = true;
+				}
+			}
+			else if (SceneNode->GetCustomLocalTransform(LocalTransform))
+			{
+				if (!LocalTransform.Equals(FTransform::Identity))
+				{
+					Parameters.IsEmptyAnimation = true;
+				}
+			}
+		}
 	}
 
 	bool FAnimationPayloadContext::FetchPayloadToFile(FFbxParser& Parser, const FString& PayloadFilepath)
@@ -829,9 +853,9 @@ namespace UE::Interchange::Private
 		, const int32& AnimationIndex)
 	{
 		FGetFbxTransformCurvesParameters Parameters(SDKScene, Node);
-		GetFbxTransformCurves(Parameters, AnimationIndex);
+		GetFbxTransformCurves(Parameters, AnimationIndex, SceneNode);
 
-		if (Parameters.IsNodeAnimated)
+		if (Parameters.IsNodeAnimated || Parameters.IsEmptyAnimation)
 		{
 			FString PayLoadKey = Parser.GetFbxHelper()->GetFbxNodeHierarchyName(Node) + TEXT("_") + FString::FromInt(AnimationIndex) + TEXT("_SkeletalAnimationPayloadKey");
 			if (ensure(!PayloadContexts.Contains(PayLoadKey)))
