Interchange: fix import of animation containing NAN in some transform.

#rnx

diff --git a/Engine/Plugins/Interchange/Runtime/Source/Parsers/Fbx/Private/FbxAnimation.cpp b/Engine/Plugins/Interchange/Runtime/Source/Parsers/Fbx/Private/FbxAnimation.cpp
--- a/Engine/Plugins/Interchange/Runtime/Source/Parsers/Fbx/Private/FbxAnimation.cpp
+++ b/Engine/Plugins/Interchange/Runtime/Source/Parsers/Fbx/Private/FbxAnimation.cpp
@@ -414,17 +414,16 @@ namespace UE::Interchange::Private
 				if (GlobalTransform.ContainsNaN())
 				{
 					LogNanError();
-					continue;
+					GlobalTransform.SetIdentity();
 				}
 
 				FbxAMatrix ParentTransform = ParentNode->EvaluateGlobalTransform(CurTime);
-				if (GlobalTransform.ContainsNaN())
+				FTransform ParentGlobalTransform = UE::Interchange::Private::FFbxConvert::ConvertTransform<FTransform, FVector, FQuat>(ParentTransform);
+				if (ParentGlobalTransform.ContainsNaN())
 				{
 					LogNanError();
-					continue;
+					ParentGlobalTransform.SetIdentity();
 				}
-				FTransform ParentGlobalTransform = UE::Interchange::Private::FFbxConvert::ConvertTransform<FTransform, FVector, FQuat>(ParentTransform);
-
 				LocalTransform = GlobalTransform.GetRelativeTransform(ParentGlobalTransform);
 			}
 			else
