Interchange:
- Fix option Flip normal green channel

diff --git a/Engine/Plugins/Interchange/Runtime/Source/Pipelines/Private/InterchangeGenericTexturePipeline.cpp b/Engine/Plugins/Interchange/Runtime/Source/Pipelines/Private/InterchangeGenericTexturePipeline.cpp
--- a/Engine/Plugins/Interchange/Runtime/Source/Pipelines/Private/InterchangeGenericTexturePipeline.cpp
+++ b/Engine/Plugins/Interchange/Runtime/Source/Pipelines/Private/InterchangeGenericTexturePipeline.cpp
@@ -402,6 +402,12 @@ void UInterchangeGenericTexturePipeline::PostImportTextureAssetImport(UObject* C
 	// import iteration, _before_ the iteration to call to PostEditChange which is what starts the texture build via UpdateResource,
 	// so altering properties here should be safe!)
 	check(!FTextureCompilingManager::Get().IsCompilingTexture(Texture));
+
+	if(bFlipNormalMapGreenChannel && Texture->IsNormalMap())
+	{
+		Texture->bFlipGreenChannel = true;
+	}
+
 	FTextureSource& Source = Texture->Source;
 
 	bool bRunNormapMapDetection = !bIsAReimport && bDetectNormalMapTexture && !Texture->IsNormalMap();
diff --git a/Engine/Plugins/Interchange/Runtime/Source/Pipelines/Public/InterchangeGenericTexturePipeline.h b/Engine/Plugins/Interchange/Runtime/Source/Pipelines/Public/InterchangeGenericTexturePipeline.h
--- a/Engine/Plugins/Interchange/Runtime/Source/Pipelines/Public/InterchangeGenericTexturePipeline.h
+++ b/Engine/Plugins/Interchange/Runtime/Source/Pipelines/Public/InterchangeGenericTexturePipeline.h
@@ -41,7 +41,7 @@ public:
 	UPROPERTY(EditAnywhere, BlueprintReadWrite, Category = "Textures", Meta=(EditCondition="bImportTextures"))
 	bool bDetectNormalMapTexture = true;
 
-	/** If enabled, the texture's green channel will be inverted for normal maps. This setting is only used if the Detect Normal Map Texture setting is also enabled. */
+	/** If enabled, the texture's green channel will be inverted for normal maps. This setting is only used if the Detect Normal Map Texture setting is also enabled or if the texture has been imported as a Normal map. */
 	UPROPERTY(EditAnywhere, BlueprintReadWrite, Category = "Textures", Meta=(EditCondition="bImportTextures"))
 	bool bFlipNormalMapGreenChannel = false;
 
