Interchange: fix the bind pose that is wrong in some cases.

#rnx

diff --git a/Engine/Plugins/Interchange/Runtime/Source/Parsers/Fbx/Private/FbxAPI.cpp b/Engine/Plugins/Interchange/Runtime/Source/Parsers/Fbx/Private/FbxAPI.cpp
--- a/Engine/Plugins/Interchange/Runtime/Source/Parsers/Fbx/Private/FbxAPI.cpp
+++ b/Engine/Plugins/Interchange/Runtime/Source/Parsers/Fbx/Private/FbxAPI.cpp
@@ -332,20 +332,6 @@ namespace UE
 
 			void FFbxParser::CleanupFbxData()
 			{
-				//////////////////////////////////////////////////////////////////////////
-				// Make sure there is a valid bind pose
-
-				//Find root bones
-				const int32 Default_NbPoses = SDKScene->GetFbxManager()->GetBindPoseCount(SDKScene);
-				// If there are no BindPoses, the following will generate them.
-				SDKScene->GetFbxManager()->CreateMissingBindPoses(SDKScene);
-				//if we created missing bind poses, update the number of bind poses
-				const int32 NbPoses = SDKScene->GetFbxManager()->GetBindPoseCount(SDKScene);
-				if (NbPoses != Default_NbPoses && !GIsAutomationTesting)
-				{
-					UInterchangeResultWarning_Generic* Message = AddMessage<UInterchangeResultWarning_Generic>();
-					Message->Text = LOCTEXT("MissingBindPose", "Missing bind pose - the FBX SDK has created one.");
-				}
 				auto MakeFbxObjectNameUnique = [bKeepFbxNamespaceClosure = bKeepFbxNamespace](FbxObject* Object, TMap<FString, int32>& Names)
 					{
 						FString ObjectName = UTF8_TO_TCHAR(Object->GetName());
diff --git a/Engine/Plugins/Interchange/Runtime/Source/Parsers/Fbx/Private/FbxConvert.h b/Engine/Plugins/Interchange/Runtime/Source/Parsers/Fbx/Private/FbxConvert.h
--- a/Engine/Plugins/Interchange/Runtime/Source/Parsers/Fbx/Private/FbxConvert.h
+++ b/Engine/Plugins/Interchange/Runtime/Source/Parsers/Fbx/Private/FbxConvert.h
@@ -56,6 +56,7 @@ namespace UE
 						}
 					}
 
+					VerifyFiniteMatrix<MatrixType>(UEMatrix);
 					return UEMatrix;
 				}
 
@@ -97,7 +98,7 @@ namespace UE
 					UnrealQuat.Y = -Quaternion[1];
 					UnrealQuat.Z = Quaternion[2];
 					UnrealQuat.W = -Quaternion[3];
-
+					VerifyFiniteQuat<QuatType>(UnrealQuat);
 					return UnrealQuat;
 				}
 
@@ -116,6 +117,7 @@ namespace UE
 					Out[0] = Vector[0];
 					Out[1] = Vector[1];
 					Out[2] = Vector[2];
+					VerifyFiniteVector<VectorType>(Out);
 					return Out;
 				}
 				
@@ -131,14 +133,18 @@ namespace UE
 				template<typename VectorType>
 				static VectorType ConvertPos(const FbxVector4& Vector)
 				{
-					return VectorType(Vector[0], -Vector[1], Vector[2]);
+					VectorType Pos(Vector[0], -Vector[1], Vector[2]);
+					VerifyFiniteVector<VectorType>(Pos);
+					return Pos;
 				}
 
 				/* Return a FVector float or double */
 				template<typename VectorType>
 				static VectorType ConvertDir(const FbxVector4& Vector)
 				{
-					return VectorType(Vector[0], -Vector[1], Vector[2]);
+					VectorType Dir(Vector[0], -Vector[1], Vector[2]);
+					VerifyFiniteVector<VectorType>(Dir);
+					return Dir;
 				}
 
 				static FLinearColor ConvertColor(const FbxDouble3& Color);
@@ -186,6 +192,42 @@ namespace UE
 
 				/** Scene Conversion Private Implementation End */
 				//////////////////////////////////////////////////////////////////////////
+
+				template<typename ScalarType>
+				static void VerifyFiniteScalar(ScalarType& Value)
+				{
+					if (!FMath::IsFinite(Value))
+					{
+						Value = 0.0f;
+					}
+				}
+
+				template<typename VectorType>
+				static void VerifyFiniteVector(VectorType& Value)
+				{
+					if (Value.ContainsNaN())
+					{
+						Value.Set(0.0, 0.0, 0.0);
+					}
+				}
+
+				template<typename QuatType>
+				static void VerifyFiniteQuat(QuatType& Value)
+				{
+					if (Value.ContainsNaN())
+					{
+						Value = QuatType::Identity;
+					}
+				}
+
+				template<typename MatrixType>
+				static void VerifyFiniteMatrix(MatrixType& Value)
+				{
+					if (Value.ContainsNaN())
+					{
+						Value.SetIdentity();
+					}
+				}
 			};
 		}//ns Private
 	}//ns Interchange
diff --git a/Engine/Plugins/Interchange/Runtime/Source/Parsers/Fbx/Private/FbxMesh.cpp b/Engine/Plugins/Interchange/Runtime/Source/Parsers/Fbx/Private/FbxMesh.cpp
--- a/Engine/Plugins/Interchange/Runtime/Source/Parsers/Fbx/Private/FbxMesh.cpp
+++ b/Engine/Plugins/Interchange/Runtime/Source/Parsers/Fbx/Private/FbxMesh.cpp
@@ -1980,8 +1980,9 @@ void FFbxMesh::AddAllMeshes(FbxScene* SDKScene, FbxGeometryConverter* SDKGeometr
 	} // for GeometryCount
 }
 
-bool FFbxMesh::GetGlobalJointBindPoseTransform(FbxScene* SDKScene, FbxNode* Joint, FbxAMatrix& GlobalBindPoseJointMatrix)
+bool FFbxMesh::GetGlobalJointBindPoseTransform(FFbxParser* Parser, FbxScene* SDKScene, FbxNode* Joint, FbxAMatrix& GlobalBindPoseJointMatrix, bool& bBadBindPoseMessageDisplay)
 {
+	FbxManager* SDKManager = SDKScene->GetFbxManager();
 	//First look for Cluster and then look in the bind pose if no cluster was found
 
 	//Search all skeletalmesh(FbxGeometry with valid deformer) using this joint and see if there is a valid FbxCluster
@@ -2037,80 +2038,129 @@ bool FFbxMesh::GetGlobalJointBindPoseTransform(FbxScene* SDKScene, FbxNode* Join
 		return false;
 	};
 
-	const int32 PoseCount = SDKScene->GetPoseCount();
-	for (int32 PoseIndex = 0; PoseIndex < PoseCount; PoseIndex++)
-	{
-		FbxPose* CurrentPose = SDKScene->GetPose(PoseIndex);
-
-		// current pose is bind pose, 
-		if (CurrentPose && CurrentPose->IsBindPose())
+	auto RetrievePoseFromBindPose = [&]()
 		{
-			FString PoseName = CurrentPose->GetName();
-			// all error report status
-			FbxStatus Status;
-
-			FbxArray<FbxNode*> pMissingAncestors, pMissingDeformers, pMissingDeformersAncestors, pWrongMatrices;
-
-			if (CurrentPose->IsValidBindPoseVerbose(Joint, pMissingAncestors, pMissingDeformers, pMissingDeformersAncestors, pWrongMatrices, 0.0001, &Status))
-			{
-				if (AcquireBindPoseMatrix(CurrentPose, GlobalBindPoseJointMatrix, Joint))
-				{
-					return true;
-				}
-			}
-			else
+			const int32 PoseCount = SDKScene->GetPoseCount();
+			for (int32 PoseIndex = 0; PoseIndex < PoseCount; PoseIndex++)
 			{
-				// first try to fix up
-				// add missing ancestors
-				for (int i = 0; i < pMissingAncestors.GetCount(); i++)
+				FbxPose* CurrentPose = SDKScene->GetPose(PoseIndex);
+
+				// current pose is bind pose, 
+				if (CurrentPose && CurrentPose->IsBindPose())
 				{
-					FbxAMatrix mat = pMissingAncestors.GetAt(i)->EvaluateGlobalTransform(FBXSDK_TIME_ZERO);
-					CurrentPose->Add(pMissingAncestors.GetAt(i), mat);
-				}
+					FString PoseName = CurrentPose->GetName();
+					// all error report status
+					FbxStatus Status;
 
-				pMissingAncestors.Clear();
-				pMissingDeformers.Clear();
-				pMissingDeformersAncestors.Clear();
-				pWrongMatrices.Clear();
+					FbxArray<FbxNode*> pMissingAncestors, pMissingDeformers, pMissingDeformersAncestors, pWrongMatrices;
 
-				// check it again
-				if (CurrentPose->IsValidBindPose(Joint))
-				{
-					if (AcquireBindPoseMatrix(CurrentPose, GlobalBindPoseJointMatrix, Joint))
+					if (CurrentPose->IsValidBindPoseVerbose(Joint, pMissingAncestors, pMissingDeformers, pMissingDeformersAncestors, pWrongMatrices, 0.0001, &Status))
 					{
-						return true;
+						if (AcquireBindPoseMatrix(CurrentPose, GlobalBindPoseJointMatrix, Joint))
+						{
+							return true;
+						}
 					}
-				}
-				else
-				{
-					// first try to find parent who is null group and see if you can try test it again
-					FbxNode* ParentNode = Joint->GetParent();
-					while (ParentNode)
+					else
 					{
-						FbxNodeAttribute* Attr = ParentNode->GetNodeAttribute();
-						if (Attr && Attr->GetAttributeType() == FbxNodeAttribute::eNull)
+						// first try to fix up
+						// add missing ancestors
+						for (int i = 0; i < pMissingAncestors.GetCount(); i++)
 						{
-							// found it 
-							break;
+							FbxAMatrix mat = pMissingAncestors.GetAt(i)->EvaluateGlobalTransform(FBXSDK_TIME_ZERO);
+							CurrentPose->Add(pMissingAncestors.GetAt(i), mat);
 						}
 
-						// find next parent
-						ParentNode = ParentNode->GetParent();
-					}
+						pMissingAncestors.Clear();
+						pMissingDeformers.Clear();
+						pMissingDeformersAncestors.Clear();
+						pWrongMatrices.Clear();
 
-					if (ParentNode && CurrentPose->IsValidBindPose(ParentNode))
-					{
-						if (AcquireBindPoseMatrix(CurrentPose, GlobalBindPoseJointMatrix, Joint))
+						// check it again
+						if (CurrentPose->IsValidBindPose(Joint))
 						{
-							return true;
+							if (AcquireBindPoseMatrix(CurrentPose, GlobalBindPoseJointMatrix, Joint))
+							{
+								return true;
+							}
+						}
+						else
+						{
+							// first try to find parent who is null group and see if you can try test it again
+							FbxNode* ParentNode = Joint->GetParent();
+							while (ParentNode)
+							{
+								FbxNodeAttribute* Attr = ParentNode->GetNodeAttribute();
+								if (Attr && Attr->GetAttributeType() == FbxNodeAttribute::eNull)
+								{
+									// found it 
+									break;
+								}
+
+								// find next parent
+								ParentNode = ParentNode->GetParent();
+							}
+
+							if (ParentNode && CurrentPose->IsValidBindPose(ParentNode))
+							{
+								if (AcquireBindPoseMatrix(CurrentPose, GlobalBindPoseJointMatrix, Joint))
+								{
+									return true;
+								}
+							}
 						}
 					}
 				}
 			}
+
+			return false;
+		};
+
+	bool bRerieveBindPoseResult = RetrievePoseFromBindPose();
+	// get bind pose
+	if (!bRerieveBindPoseResult)
+	{
+		bool bOriginalBadBindPoseMessageDisplay = bBadBindPoseMessageDisplay;
+		if (!GIsAutomationTesting && !bOriginalBadBindPoseMessageDisplay)
+		{
+			UInterchangeResultWarning_Generic* Message = Parser->AddMessage<UInterchangeResultWarning_Generic>();
+			Message->Text = LOCTEXT("MissingBindPose", "Missing bind pose. Try to recreate bind pose.");
+			bBadBindPoseMessageDisplay = true;
 		}
-	}
+		// if failed, delete bind pose, and retry.
+		const int32 PoseCount = SDKScene->GetPoseCount();
+		for (int32 PoseIndex = PoseCount - 1; PoseIndex >= 0; --PoseIndex)
+		{
+			FbxPose* CurrentPose = SDKScene->GetPose(PoseIndex);
 
-	return false;
+			// current pose is bind pose, 
+			if (CurrentPose && CurrentPose->IsBindPose())
+			{
+				SDKScene->RemovePose(PoseIndex);
+				CurrentPose->Destroy();
+			}
+		}
+
+		SDKManager->CreateMissingBindPoses(SDKScene);
+		bRerieveBindPoseResult = RetrievePoseFromBindPose();
+		if (!bRerieveBindPoseResult)
+		{
+			if (!GIsAutomationTesting && !bOriginalBadBindPoseMessageDisplay)
+			{
+				UInterchangeResultWarning_Generic* Message = Parser->AddMessage<UInterchangeResultWarning_Generic>();
+				Message->Text = LOCTEXT("MissingBindPose_RecreateFailed", "Missing bind pose. Recreating bind pose failed.");
+			}
+		}
+		else
+		{
+			if (!GIsAutomationTesting && !bOriginalBadBindPoseMessageDisplay)
+			{
+				UInterchangeResultWarning_Generic* Message = Parser->AddMessage<UInterchangeResultWarning_Generic>();
+				Message->Text = LOCTEXT("MissingBindPose_RecreateSuccess", "Missing bind pose. Recreating bind pose succeeded.");
+			}
+		}
+	}
+	return bRerieveBindPoseResult;
 }
 
 bool FFbxMesh::ExtractSkinnedMeshNodeJoints(FbxScene* SDKScene, UInterchangeBaseNodeContainer& NodeContainer, FbxMesh* Mesh, UInterchangeMeshNode* MeshNode)
diff --git a/Engine/Plugins/Interchange/Runtime/Source/Parsers/Fbx/Private/FbxMesh.h b/Engine/Plugins/Interchange/Runtime/Source/Parsers/Fbx/Private/FbxMesh.h
--- a/Engine/Plugins/Interchange/Runtime/Source/Parsers/Fbx/Private/FbxMesh.h
+++ b/Engine/Plugins/Interchange/Runtime/Source/Parsers/Fbx/Private/FbxMesh.h
@@ -196,7 +196,7 @@ namespace UE
 				{}
 
 				void AddAllMeshes(FbxScene* SDKScene, FbxGeometryConverter* SDKGeometryConverter, UInterchangeBaseNodeContainer& NodeContainer, TMap<FString, TSharedPtr<FPayloadContextBase>>& PayloadContexts);
-				static bool GetGlobalJointBindPoseTransform(FbxScene* SDKScene, FbxNode* Joint, FbxAMatrix& GlobalBindPoseJointMatrix);
+				static bool GetGlobalJointBindPoseTransform(FFbxParser* Parser, FbxScene* SDKScene, FbxNode* Joint, FbxAMatrix& GlobalBindPoseJointMatrix, bool& bBadBindPoseMessageDisplay);
 
 				TArray<FMorphTargetAnimationBuildingData>& GetMorphTargetAnimationsBuildingData() { return MorphTargetAnimationsBuildingData; }
 			protected:
diff --git a/Engine/Plugins/Interchange/Runtime/Source/Parsers/Fbx/Private/FbxScene.cpp b/Engine/Plugins/Interchange/Runtime/Source/Parsers/Fbx/Private/FbxScene.cpp
--- a/Engine/Plugins/Interchange/Runtime/Source/Parsers/Fbx/Private/FbxScene.cpp
+++ b/Engine/Plugins/Interchange/Runtime/Source/Parsers/Fbx/Private/FbxScene.cpp
@@ -160,7 +160,8 @@ namespace UE
 				, FbxScene* SDKScene
 				, UInterchangeBaseNodeContainer& NodeContainer
 				, TMap<FString, TSharedPtr<FPayloadContextBase, ESPMode::ThreadSafe>>& PayloadContexts
-				, TArray<FbxNode*>& ForceJointNodes)
+				, TArray<FbxNode*>& ForceJointNodes
+				, bool& bBadBindPoseMessageDisplay)
 			{
 				constexpr bool bResetCache = false;
 				FString NodeName = Parser.GetFbxHelper()->GetFbxObjectName(Node);
@@ -211,20 +212,20 @@ namespace UE
 					}
 				}
 
-				auto ApplySkeletonAttribute = [this, &SDKScene, &UnrealNode, &Node, &NodeContainer, &bResetCache, &GetConvertedTransform]()
+				auto ApplySkeletonAttribute = [this, &SDKScene, &UnrealNode, &Node, &NodeContainer, &bResetCache, &GetConvertedTransform, &bBadBindPoseMessageDisplay]()
 				{
 					//Add the joint specialized type
 					UnrealNode->AddSpecializedType(FSceneNodeStaticData::GetJointSpecializeTypeString());
 					//Get the bind pose transform for this joint
 					FbxAMatrix GlobalBindPoseJointMatrix;
-					if (FFbxMesh::GetGlobalJointBindPoseTransform(SDKScene, Node, GlobalBindPoseJointMatrix))
+					if (FFbxMesh::GetGlobalJointBindPoseTransform(&Parser, SDKScene, Node, GlobalBindPoseJointMatrix, bBadBindPoseMessageDisplay))
 					{
 						FTransform GlobalBindPoseJointTransform = GetConvertedTransform(GlobalBindPoseJointMatrix);
 						//We grab the fbx parent node to compute the local transform
 						if (FbxNode* ParentNode = Node->GetParent())
 						{
 							FbxAMatrix GlobalFbxParentMatrix = ParentNode->EvaluateGlobalTransform();
-							FFbxMesh::GetGlobalJointBindPoseTransform(SDKScene, ParentNode, GlobalFbxParentMatrix);
+							FFbxMesh::GetGlobalJointBindPoseTransform(&Parser, SDKScene, ParentNode, GlobalFbxParentMatrix, bBadBindPoseMessageDisplay);
 							FbxAMatrix	LocalFbxMatrix = GlobalFbxParentMatrix.Inverse() * GlobalBindPoseJointMatrix;
 							FTransform LocalBindPoseJointTransform = GetConvertedTransform(LocalFbxMatrix);
 							UnrealNode->SetCustomBindPoseLocalTransform(&NodeContainer, LocalBindPoseJointTransform, bResetCache);
@@ -474,7 +475,7 @@ namespace UE
 				for (int32 ChildIndex = 0; ChildIndex < ChildCount; ++ChildIndex)
 				{
 					FbxNode* ChildNode = Node->GetChild(ChildIndex);
-					AddHierarchyRecursively(UnrealNode, ChildNode, SDKScene, NodeContainer, PayloadContexts, ForceJointNodes);
+					AddHierarchyRecursively(UnrealNode, ChildNode, SDKScene, NodeContainer, PayloadContexts, ForceJointNodes, bBadBindPoseMessageDisplay);
 				}
 			}
 
@@ -504,7 +505,8 @@ namespace UE
 				TArray<FbxNode*> ForceJointNodes;
 				FindForceJointNode(SDKScene, ForceJointNodes);
 
-				AddHierarchyRecursively(nullptr, RootNode, SDKScene, NodeContainer, PayloadContexts, ForceJointNodes);
+				bool bBadBindPoseMessageDisplay = false;
+				AddHierarchyRecursively(nullptr, RootNode, SDKScene, NodeContainer, PayloadContexts, ForceJointNodes, bBadBindPoseMessageDisplay);
 
 				int32 NodeCount = SDKScene->GetNodeCount();
 				for (int32 NodeIndex = 0; NodeIndex < NodeCount; ++NodeIndex)
@@ -515,7 +517,7 @@ namespace UE
 						{
 							if (Node->GetParent() == nullptr)
 							{
-								AddHierarchyRecursively(nullptr, Node, SDKScene, NodeContainer, PayloadContexts, ForceJointNodes);
+								AddHierarchyRecursively(nullptr, Node, SDKScene, NodeContainer, PayloadContexts, ForceJointNodes, bBadBindPoseMessageDisplay);
 							}
 						}
 					}
diff --git a/Engine/Plugins/Interchange/Runtime/Source/Parsers/Fbx/Private/FbxScene.h b/Engine/Plugins/Interchange/Runtime/Source/Parsers/Fbx/Private/FbxScene.h
--- a/Engine/Plugins/Interchange/Runtime/Source/Parsers/Fbx/Private/FbxScene.h
+++ b/Engine/Plugins/Interchange/Runtime/Source/Parsers/Fbx/Private/FbxScene.h
@@ -44,7 +44,8 @@ namespace UE
 					, FbxScene* SDKScene
 					, UInterchangeBaseNodeContainer& NodeContainer
 					, TMap<FString, TSharedPtr<FPayloadContextBase, ESPMode::ThreadSafe>>& PayloadContexts
-					, TArray<FbxNode*>& ForceJointNodes);
+					, TArray<FbxNode*>& ForceJointNodes
+					, bool& bBadBindPoseMessageDisplay);
 
 				void AddAnimationRecursively(FbxNode* Node
 					, FbxScene* SDKScene
