Undo changelist 32861263,
Fixing BindPose vs Local Transform application for Animations in AnimSequenceFactory.

#rnx

diff --git a/Engine/Plugins/Interchange/Runtime/Source/Import/Private/Animation/InterchangeAnimSequenceFactory.cpp b/Engine/Plugins/Interchange/Runtime/Source/Import/Private/Animation/InterchangeAnimSequenceFactory.cpp
--- a/Engine/Plugins/Interchange/Runtime/Source/Import/Private/Animation/InterchangeAnimSequenceFactory.cpp
+++ b/Engine/Plugins/Interchange/Runtime/Source/Import/Private/Animation/InterchangeAnimSequenceFactory.cpp
@@ -418,6 +418,7 @@ namespace UE::Interchange::Private
 
 		TArray<FString> SkeletonNodes;
 		GetSkeletonSceneNodeFlatListRecursive(NodeContainer, SkeletonRootUid, SkeletonNodes);
+		TArray<FString> NonAnimatedSkeletonNodes = SkeletonNodes;
 
 		TMap<FString, FInterchangeAnimationPayLoadKey> PayloadKeys;
 		AnimSequenceFactoryNode->GetSceneNodeAnimationPayloadKeys(PayloadKeys);
@@ -463,6 +464,8 @@ namespace UE::Interchange::Private
 
 			for (TTuple< const UInterchangeSceneNode*, UE::Interchange::FAnimationPayloadData>& AnimationPayload : BoneTrackData.PreProcessedAnimationPayloads)
 			{
+				NonAnimatedSkeletonNodes.Remove(AnimationPayload.Key->GetUniqueID());
+
 				const FName BoneName = FName(*(AnimationPayload.Key->GetDisplayLabel()));
 				UE::Interchange::FAnimationPayloadData& AnimationTransformPayload = AnimationPayload.Value;
 				
@@ -577,6 +580,42 @@ namespace UE::Interchange::Private
 				Controller.AddBoneCurve(BoneName, bShouldTransact);
 				Controller.SetBoneTrackKeys(BoneName, RawTrack.PosKeys, RawTrack.RotKeys, RawTrack.ScaleKeys, bShouldTransact);
 			}
+
+			//For joint with no animation, verify if the bind pose equal the local time 0 pose.
+			// If not, Add one animation track with only one transform key at time 0 with the Local transform.
+			for (const FString& NonAnimatedSkeletonNodeUID : NonAnimatedSkeletonNodes)
+			{
+				if (const UInterchangeSceneNode* SkeletonNode = Cast<UInterchangeSceneNode>(NodeContainer->GetNode(NonAnimatedSkeletonNodeUID)))
+				{
+					//check if BindPose exists and if it does does it equal to LocalTransform
+					FTransform LocalBindPoseTransform;
+					FTransform LocalTransform;
+					if (SkeletonNode->GetCustomBindPoseLocalTransform(LocalBindPoseTransform)
+						&& SkeletonNode->GetCustomLocalTransform(LocalTransform)
+						&& !LocalBindPoseTransform.Equals(LocalTransform))
+					{
+						//If we bake the mesh and the current non animated node is the root joint, get the global transform instead of the local
+						if (bBakeMeshes && SkeletonNode->GetUniqueID().Equals(SkeletonRootUid))
+						{
+							if (const UInterchangeSceneNode* RootJointNode = Cast<UInterchangeSceneNode>(NodeContainer->GetNode(SkeletonRootUid)))
+							{
+								RootJointNode->GetCustomGlobalTransform(NodeContainer, GlobalOffsetTransform, LocalTransform);
+							}
+						}
+						FTransform3f AnimKeyTransform = static_cast<FTransform3f>(LocalTransform);
+						const FName BoneName = FName(*(SkeletonNode->GetDisplayLabel()));
+						//Add only one transform key at time 0 since this node is not animated
+						FRawAnimSequenceTrack RawTrack;
+						TArray<float> TimeKeys;
+						RawTrack.ScaleKeys.Add(AnimKeyTransform.GetScale3D());
+						RawTrack.PosKeys.Add(AnimKeyTransform.GetLocation());
+						RawTrack.RotKeys.Add(AnimKeyTransform.GetRotation());
+						TimeKeys.Add(0);
+						Controller.AddBoneCurve(BoneName, bShouldTransact);
+						Controller.SetBoneTrackKeys(BoneName, RawTrack.PosKeys, RawTrack.RotKeys, RawTrack.ScaleKeys, bShouldTransact);
+					}
+				}
+			}
 		}
 
 		bool bDeleteExistingMorphTargetCurves = false;
diff --git a/Engine/Plugins/Interchange/Runtime/Source/Parsers/Fbx/Private/FbxAnimation.cpp b/Engine/Plugins/Interchange/Runtime/Source/Parsers/Fbx/Private/FbxAnimation.cpp
--- a/Engine/Plugins/Interchange/Runtime/Source/Parsers/Fbx/Private/FbxAnimation.cpp
+++ b/Engine/Plugins/Interchange/Runtime/Source/Parsers/Fbx/Private/FbxAnimation.cpp
@@ -456,13 +456,10 @@ namespace UE::Interchange::Private
 		FbxScene* SDKScene = nullptr;
 		FbxNode* Node = nullptr;
 		bool IsNodeAnimated = false;
-		bool IsEmptyAnimation = false;	//Report true if Node->Lcl transforms return empty (indicating identity Transform setting for the Animation),
-										//		while the Node (outside of Animation) has a non-identity transform
-										//(We only report true for AddSkeletalTransformAnimation for IsNodeAnimated.)
 		FbxAnimStack* CurrentAnimStack = nullptr;
 	};
 
-	void GetFbxTransformCurves(FGetFbxTransformCurvesParameters& Parameters, const int32& AnimationIndex, UInterchangeSceneNode* SceneNode)
+	void GetFbxTransformCurves(FGetFbxTransformCurvesParameters& Parameters, const int32& AnimationIndex)
 	{
 		if (!ensure(Parameters.SDKScene) || !ensure(Parameters.Node))
 		{
@@ -523,27 +520,6 @@ namespace UE::Interchange::Private
 				}
 			}
 		}
-
-		if (!Parameters.IsNodeAnimated)
-		{
-			//At this point this flag means that the Local Transform curves have no values.
-			// This could also mean that the node is animated to Identity in case the Node's original transform is not Identity:
-			FTransform LocalTransform;
-			if (SceneNode->GetCustomBindPoseLocalTransform(LocalTransform))
-			{
-				if (!LocalTransform.Equals(FTransform::Identity))
-				{
-					Parameters.IsEmptyAnimation = true;
-				}
-			}
-			else if (SceneNode->GetCustomLocalTransform(LocalTransform))
-			{
-				if (!LocalTransform.Equals(FTransform::Identity))
-				{
-					Parameters.IsEmptyAnimation = true;
-				}
-			}
-		}
 	}
 
 	bool FAnimationPayloadContext::FetchPayloadToFile(FFbxParser& Parser, const FString& PayloadFilepath)
@@ -852,9 +828,9 @@ namespace UE::Interchange::Private
 		, const int32& AnimationIndex)
 	{
 		FGetFbxTransformCurvesParameters Parameters(SDKScene, Node);
-		GetFbxTransformCurves(Parameters, AnimationIndex, SceneNode);
+		GetFbxTransformCurves(Parameters, AnimationIndex);
 
-		if (Parameters.IsNodeAnimated || Parameters.IsEmptyAnimation)
+		if (Parameters.IsNodeAnimated)
 		{
 			FString PayLoadKey = Parser.GetFbxHelper()->GetFbxNodeHierarchyName(Node) + TEXT("_") + FString::FromInt(AnimationIndex) + TEXT("_SkeletalAnimationPayloadKey");
 			if (ensure(!PayloadContexts.Contains(PayLoadKey)))
