MaterialX:
- Add support of Opacity for Standard Surface Shader, connect it to the Surface Coveerage of the ThinTranslucency instead of OpacityMask for a better result

Interchange Factory/FactoryNode/MaterialPipeline:
- Add support of SurfaceCoverage from the ThinTranslucent Material Expression

#virtualized

diff --git a/Engine/Plugins/Interchange/Runtime/Source/FactoryNodes/Private/InterchangeMaterialFactoryNode.cpp b/Engine/Plugins/Interchange/Runtime/Source/FactoryNodes/Private/InterchangeMaterialFactoryNode.cpp
--- a/Engine/Plugins/Interchange/Runtime/Source/FactoryNodes/Private/InterchangeMaterialFactoryNode.cpp
+++ b/Engine/Plugins/Interchange/Runtime/Source/FactoryNodes/Private/InterchangeMaterialFactoryNode.cpp
@@ -285,6 +285,21 @@ bool UInterchangeMaterialFactoryNode::ConnectOutputToTransmissionColor(const FSt
 	return UInterchangeShaderPortsAPI::ConnectOuputToInputByName(this, UE::Interchange::Materials::ThinTranslucent::Parameters::TransmissionColor.ToString(), ExpressionNodeUid, OutputName);
 }
 
+bool UInterchangeMaterialFactoryNode::GetSurfaceCoverageConnection(FString& ExpressionNodeUid, FString& OutputName) const
+{
+	return UInterchangeShaderPortsAPI::GetInputConnection(this, UE::Interchange::Materials::ThinTranslucent::Parameters::SurfaceCoverage.ToString(), ExpressionNodeUid, OutputName);
+}
+
+bool UInterchangeMaterialFactoryNode::ConnectToSurfaceCoverage(const FString& ExpressionUid)
+{
+	return UInterchangeShaderPortsAPI::ConnectDefaultOuputToInput(this, UE::Interchange::Materials::ThinTranslucent::Parameters::SurfaceCoverage.ToString(), ExpressionUid);
+}
+
+bool UInterchangeMaterialFactoryNode::ConnectOutputToSurfaceCoverage(const FString& ExpressionNodeUid, const FString& OutputName)
+{
+	return UInterchangeShaderPortsAPI::ConnectOuputToInputByName(this, UE::Interchange::Materials::ThinTranslucent::Parameters::SurfaceCoverage.ToString(), ExpressionNodeUid, OutputName);
+}
+
 bool UInterchangeMaterialFactoryNode::GetFuzzColorConnection(FString& ExpressionNodeUid, FString& OutputName) const
 {
 	return UInterchangeShaderPortsAPI::GetInputConnection(this, UE::Interchange::Materials::Sheen::Parameters::SheenColor.ToString(), ExpressionNodeUid, OutputName);
diff --git a/Engine/Plugins/Interchange/Runtime/Source/FactoryNodes/Public/InterchangeMaterialFactoryNode.h b/Engine/Plugins/Interchange/Runtime/Source/FactoryNodes/Public/InterchangeMaterialFactoryNode.h
--- a/Engine/Plugins/Interchange/Runtime/Source/FactoryNodes/Public/InterchangeMaterialFactoryNode.h
+++ b/Engine/Plugins/Interchange/Runtime/Source/FactoryNodes/Public/InterchangeMaterialFactoryNode.h
@@ -189,6 +189,15 @@ public:
 
 	UFUNCTION(BlueprintCallable, Category = "Interchange | Node | MaterialFactory")
 	bool ConnectOutputToTransmissionColor(const FString& ExpressionNodeUid, const FString& OutputName);
+	
+	UFUNCTION(BlueprintCallable, Category = "Interchange | Node | MaterialFactory")
+	bool GetSurfaceCoverageConnection(FString& ExpressionNodeUid, FString& OutputName) const;
+
+	UFUNCTION(BlueprintCallable, Category = "Interchange | Node | MaterialFactory")
+	bool ConnectToSurfaceCoverage(const FString& ExpressionUid);
+
+	UFUNCTION(BlueprintCallable, Category = "Interchange | Node | MaterialFactory")
+	bool ConnectOutputToSurfaceCoverage(const FString& ExpressionNodeUid, const FString& OutputName);
 
 	UFUNCTION(BlueprintCallable, Category = "Interchange | Node | MaterialFactory")
 	bool GetFuzzColorConnection(FString& ExpressionNodeUid, FString& OutputName) const;
diff --git a/Engine/Plugins/Interchange/Runtime/Source/Import/Private/Material/InterchangeMaterialFactory.cpp b/Engine/Plugins/Interchange/Runtime/Source/Import/Private/Material/InterchangeMaterialFactory.cpp
--- a/Engine/Plugins/Interchange/Runtime/Source/Import/Private/Material/InterchangeMaterialFactory.cpp
+++ b/Engine/Plugins/Interchange/Runtime/Source/Import/Private/Material/InterchangeMaterialFactory.cpp
@@ -1641,20 +1641,38 @@ void UInterchangeMaterialFactory::SetupMaterial(UMaterial* Material, const FImpo
 
 	// Thin Translucent
 	{
-		FString TransmissionColorUid;
-		FString OutputName;
+		FString TransmissionColorUid, SurfaceCoverageUid;
+		FString TransmissionColorOutputName, SurfaceCoverageOutputName;
 
-		if (MaterialFactoryNode->GetTransmissionColorConnection(TransmissionColorUid, OutputName))
+		const bool bHasTransmissionColor = MaterialFactoryNode->GetTransmissionColorConnection(TransmissionColorUid, TransmissionColorOutputName);
+		const bool bHasSurfaceCoverage = MaterialFactoryNode->GetSurfaceCoverageConnection(SurfaceCoverageUid, SurfaceCoverageOutputName);
+		if(bHasTransmissionColor || bHasSurfaceCoverage)
 		{
 			const UInterchangeMaterialExpressionFactoryNode* TransmissionColorNode = Cast<UInterchangeMaterialExpressionFactoryNode>(Arguments.NodeContainer->GetNode(TransmissionColorUid));
+			const UInterchangeMaterialExpressionFactoryNode* SurfaceCoverageNode = Cast<UInterchangeMaterialExpressionFactoryNode>(Arguments.NodeContainer->GetNode(SurfaceCoverageUid));
+
+			UMaterialExpression* ThinTranslucentMaterialOutput = nullptr;
 
-			if (TransmissionColorNode)
+			if(TransmissionColorNode)
 			{
-				if (UMaterialExpression* TransmissionColorExpression = Builder.CreateExpressionsForNode(*TransmissionColorNode))
+				if(UMaterialExpression* TransmissionColorExpression = Builder.CreateExpressionsForNode(*TransmissionColorNode))
 				{
-					UMaterialExpression* ThinTranslucentMaterialOutput = CreateMaterialExpression(Material, nullptr, UMaterialExpressionThinTranslucentMaterialOutput::StaticClass());
-					
-					TransmissionColorExpression->ConnectExpression(ThinTranslucentMaterialOutput->GetInput(0), GetOutputIndex(*TransmissionColorExpression, OutputName));
+					ThinTranslucentMaterialOutput = CreateMaterialExpression(Material, nullptr, UMaterialExpressionThinTranslucentMaterialOutput::StaticClass());
+
+					TransmissionColorExpression->ConnectExpression(ThinTranslucentMaterialOutput->GetInput(0), GetOutputIndex(*TransmissionColorExpression, TransmissionColorOutputName));
+				}
+			}
+
+			if(SurfaceCoverageNode)
+			{
+				if(UMaterialExpression* SurfaceCoverageExpression = Builder.CreateExpressionsForNode(*SurfaceCoverageNode))
+				{
+					if(!ThinTranslucentMaterialOutput)
+					{
+						ThinTranslucentMaterialOutput = CreateMaterialExpression(Material, nullptr, UMaterialExpressionThinTranslucentMaterialOutput::StaticClass());
+					}
+
+					SurfaceCoverageExpression->ConnectExpression(ThinTranslucentMaterialOutput->GetInput(1), GetOutputIndex(*SurfaceCoverageExpression, SurfaceCoverageOutputName));
 				}
 			}
 		}
diff --git a/Engine/Plugins/Interchange/Runtime/Source/Import/Private/MaterialX/MaterialXUtils/MaterialXStandardSurfaceShader.cpp b/Engine/Plugins/Interchange/Runtime/Source/Import/Private/MaterialX/MaterialXUtils/MaterialXStandardSurfaceShader.cpp
--- a/Engine/Plugins/Interchange/Runtime/Source/Import/Private/MaterialX/MaterialXUtils/MaterialXStandardSurfaceShader.cpp
+++ b/Engine/Plugins/Interchange/Runtime/Source/Import/Private/MaterialX/MaterialXUtils/MaterialXStandardSurfaceShader.cpp
@@ -147,6 +147,9 @@ void FMaterialXStandardSurfaceShader::ConnectToStandardSurface(UInterchangeFunct
 	//Transmission Extra Roughness
 	ConnectNodeOutputToInput(Input::TransmissionExtraRoughness, StandardSurfaceShaderNode, StandardSurface::Parameters::TransmissionExtraRoughness.ToString(), DefaultValue::Float::TransmissionExtraRoughness);
 
+	//Opacity
+	ConnectNodeOutputToInput(Input::Opacity, StandardSurfaceShaderNode, StandardSurface::Parameters::Opacity.ToString(), DefaultValue::Color3::Opacity);
+
 	// Outputs
 	UInterchangeShaderPortsAPI::ConnectOuputToInputByName(ShaderGraphNode, PBRMR::Parameters::BaseColor.ToString(), StandardSurfaceShaderNode->GetUniqueID(), TEXT("Base Color"));
 	UInterchangeShaderPortsAPI::ConnectOuputToInputByName(ShaderGraphNode, PBRMR::Parameters::Metallic.ToString(), StandardSurfaceShaderNode->GetUniqueID(), PBRMR::Parameters::Metallic.ToString());
@@ -163,6 +166,12 @@ void FMaterialXStandardSurfaceShader::ConnectToStandardSurface(UInterchangeFunct
 		UInterchangeShaderPortsAPI::ConnectOuputToInputByName(ShaderGraphNode, PBRMR::Parameters::Opacity.ToString(), StandardSurfaceShaderNode->GetUniqueID(), PBRMR::Parameters::Opacity.ToString());
 		UInterchangeShaderPortsAPI::ConnectOuputToInputByName(ShaderGraphNode, ThinTranslucent::Parameters::TransmissionColor.ToString(), StandardSurfaceShaderNode->GetUniqueID(), ThinTranslucent::Parameters::TransmissionColor.ToString());
 		UInterchangeShaderPortsAPI::ConnectOuputToInputByName(ShaderGraphNode, Common::Parameters::Refraction.ToString(), StandardSurfaceShaderNode->GetUniqueID(), Common::Parameters::Refraction.ToString());
+
+		// If we have have Transmission and Opacity let's use the Surface Coverage instead
+		if(UInterchangeShaderPortsAPI::HasInput(StandardSurfaceShaderNode, StandardSurface::Parameters::Opacity))
+		{
+			UInterchangeShaderPortsAPI::ConnectOuputToInputByName(ShaderGraphNode, ThinTranslucent::Parameters::SurfaceCoverage.ToString(), StandardSurfaceShaderNode->GetUniqueID(), ThinTranslucent::Parameters::SurfaceCoverage.ToString());
+		}
 	}
 	else if(UInterchangeShaderPortsAPI::HasInput(StandardSurfaceShaderNode, StandardSurface::Parameters::Sheen))
 	{
@@ -180,6 +189,10 @@ void FMaterialXStandardSurfaceShader::ConnectToStandardSurface(UInterchangeFunct
 		UInterchangeShaderPortsAPI::ConnectOuputToInputByName(ShaderGraphNode, PBRMR::Parameters::Opacity.ToString(), StandardSurfaceShaderNode->GetUniqueID(), PBRMR::Parameters::Opacity.ToString());
 		UInterchangeShaderPortsAPI::ConnectOuputToInputByName(ShaderGraphNode, Subsurface::Parameters::SubsurfaceColor.ToString(), StandardSurfaceShaderNode->GetUniqueID(), Subsurface::Parameters::SubsurfaceColor.ToString());
 	}
+	else if(UInterchangeShaderPortsAPI::HasInput(StandardSurfaceShaderNode, StandardSurface::Parameters::Opacity))
+	{
+		UInterchangeShaderPortsAPI::ConnectOuputToInputByName(ShaderGraphNode, ThinTranslucent::Parameters::SurfaceCoverage.ToString(), StandardSurfaceShaderNode->GetUniqueID(), ThinTranslucent::Parameters::SurfaceCoverage.ToString());
+	}
 }
 
 void FMaterialXStandardSurfaceShader::ConnectToSubstrateStandardSurface(UInterchangeFunctionCallShaderNode* StandardSurfaceShaderNode)
diff --git a/Engine/Plugins/Interchange/Runtime/Source/Nodes/Public/InterchangeMaterialDefinitions.h b/Engine/Plugins/Interchange/Runtime/Source/Nodes/Public/InterchangeMaterialDefinitions.h
--- a/Engine/Plugins/Interchange/Runtime/Source/Nodes/Public/InterchangeMaterialDefinitions.h
+++ b/Engine/Plugins/Interchange/Runtime/Source/Nodes/Public/InterchangeMaterialDefinitions.h
@@ -450,6 +450,7 @@ namespace Interchange
 			namespace Parameters
 			{
 				const FName TransmissionColor = TEXT("TransmissionColor"); // Type: vector3
+				const FName SurfaceCoverage = TEXT("SurfaceCoverage"); // Type: float
 			}
 		}
 
diff --git a/Engine/Plugins/Interchange/Runtime/Source/Pipelines/Private/InterchangeGenericMaterialPipeline.cpp b/Engine/Plugins/Interchange/Runtime/Source/Pipelines/Private/InterchangeGenericMaterialPipeline.cpp
--- a/Engine/Plugins/Interchange/Runtime/Source/Pipelines/Private/InterchangeGenericMaterialPipeline.cpp
+++ b/Engine/Plugins/Interchange/Runtime/Source/Pipelines/Private/InterchangeGenericMaterialPipeline.cpp
@@ -1066,8 +1066,9 @@ bool UInterchangeGenericMaterialPipeline::HasThinTranslucency(const UInterchange
 	using namespace UE::Interchange::Materials::ThinTranslucent;
 
 	const bool bHasTransmissionColorInput = UInterchangeShaderPortsAPI::HasInput(ShaderGraphNode, Parameters::TransmissionColor);
+	const bool bHasSurfaceCoverageInput = UInterchangeShaderPortsAPI::HasInput(ShaderGraphNode, Parameters::SurfaceCoverage);
 
-	return bHasTransmissionColorInput;
+	return bHasTransmissionColorInput || bHasSurfaceCoverageInput;
 }
 
 bool UInterchangeGenericMaterialPipeline::IsMetalRoughModel(const UInterchangeShaderGraphNode* ShaderGraphNode) const
@@ -1521,6 +1522,24 @@ bool UInterchangeGenericMaterialPipeline::HandleThinTranslucent(const UInterchan
 		}
 	}
 
+	// Surface Coverage
+	{
+		const bool bHasInput = UInterchangeShaderPortsAPI::HasInput(ShaderGraphNode, Parameters::SurfaceCoverage);
+
+		if(bHasInput)
+		{
+			TTuple<UInterchangeMaterialExpressionFactoryNode*, FString> ExpressionFactoryNode =
+				CreateMaterialExpressionForInput(MaterialFactoryNode, ShaderGraphNode, Parameters::SurfaceCoverage.ToString(), MaterialFactoryNode->GetUniqueID());
+
+			if(ExpressionFactoryNode.Get<0>())
+			{
+				MaterialFactoryNode->ConnectOutputToSurfaceCoverage(ExpressionFactoryNode.Get<0>()->GetUniqueID(), ExpressionFactoryNode.Get<1>());
+			}
+
+			bShadingModelHandled = true;
+		}
+	}
+
 	if (bShadingModelHandled)
 	{
 		MaterialFactoryNode->SetCustomBlendMode(EBlendMode::BLEND_Translucent);
