FBX Material:
- Set the value of shininess to 1 to be consistent with the master Material of Phong and legacy.

diff --git a/Engine/Plugins/Interchange/Runtime/Source/Parsers/Fbx/Private/FbxMaterial.cpp b/Engine/Plugins/Interchange/Runtime/Source/Parsers/Fbx/Private/FbxMaterial.cpp
--- a/Engine/Plugins/Interchange/Runtime/Source/Parsers/Fbx/Private/FbxMaterial.cpp
+++ b/Engine/Plugins/Interchange/Runtime/Source/Parsers/Fbx/Private/FbxMaterial.cpp
@@ -245,10 +245,9 @@ namespace UE
 						
 						UInterchangeShaderPortsAPI::ConnectDefaultOuputToInput(ShaderGraphNode, InputName, MultiplyNode->GetUniqueID());
 
-						// Scale texture output from [0-1] to [0-1000]
 						UInterchangeShaderNode* WeightNode = UInterchangeShaderNode::Create(&NodeContainer, TEXT("ShininessMapWeight"), MultiplyNode->GetUniqueID());
 						WeightNode->SetCustomShaderType(ScalarParameter::Name.ToString());
-						WeightNode->AddFloatAttribute(UInterchangeShaderPortsAPI::MakeInputParameterKey(ScalarParameter::Attributes::DefaultValue.ToString()), 1000.f);
+						WeightNode->AddFloatAttribute(UInterchangeShaderPortsAPI::MakeInputParameterKey(ScalarParameter::Attributes::DefaultValue.ToString()), 1.f);
 
 						UInterchangeShaderPortsAPI::ConnectDefaultOuputToInput(MultiplyNode, Multiply::Inputs::B.ToString(), WeightNode->GetUniqueID());
 
