[Interchange] Removing BindPose on Joints without Animations workarounds for glTFs.

#rnx

diff --git a/Engine/Plugins/Interchange/Runtime/Source/Import/Private/Gltf/InterchangeGltfAnimation.cpp b/Engine/Plugins/Interchange/Runtime/Source/Import/Private/Gltf/InterchangeGltfAnimation.cpp
--- a/Engine/Plugins/Interchange/Runtime/Source/Import/Private/Gltf/InterchangeGltfAnimation.cpp
+++ b/Engine/Plugins/Interchange/Runtime/Source/Import/Private/Gltf/InterchangeGltfAnimation.cpp
@@ -19,8 +19,6 @@
 
 namespace UE::Interchange::Gltf::Private
 {
-	const FString BIND_POSE_FIX = TEXT("BIND_POSE_FIX<->");
-
 	/*
 	* According to gltf specification Seconds acquired from Samplers.input:
 	* "The values represent time in seconds with time[0] >= 0.0, and strictly increasing values, Index.e., time[n + 1] > time[n]."
@@ -548,18 +546,6 @@ namespace UE::Interchange::Gltf::Private
 
 		if (!ParsePayLoadKey(GltfAsset, PayLoadKey, AnimationIndex, ChannelIndices))
 		{
-			if (PayLoadKey.Contains(BIND_POSE_FIX))
-			{
-				FString CutPayloadKey = PayLoadKey.Replace(*BIND_POSE_FIX, TEXT(""));
-				int32 NodeToSetIndex = INDEX_NONE;
-				LexFromString(NodeToSetIndex, *CutPayloadKey);
-
-				if (GltfAsset.Nodes.IsValidIndex(NodeToSetIndex))
-				{
-					PayloadData.Transforms.Init(GltfAsset.Nodes[NodeToSetIndex].Transform, BakeKeyCount);
-					return true;
-				}
-			}
 			return false;
 		}
 		const GLTF::FAnimation& GltfAnimation = GltfAsset.Animations[AnimationIndex];
@@ -1035,42 +1021,6 @@ namespace UE::Interchange::Gltf::Private
 				NodeContainer.AddNode(TrackSetNode);
 			}
 
-			void FixSkeletalAnimations(TMap<int32, TSet<int32>>& SkeletonRootToAnimatedJointNodeIndicesMap)
-			{
-				for (const TTuple<int32, TSet<int32>>& AnimatedJointNodeIndices : SkeletonRootToAnimatedJointNodeIndicesMap)
-				{
-					int32 SkeletonRootIndex = AnimatedJointNodeIndices.Key;
-
-					if (!GLTFNodes.IsValidIndex(SkeletonRootIndex))
-					{
-						continue;
-					}
-
-					const FString* SkeletonUid = GLTFNodeToInterchangeUidMap.Find(&GLTFNodes[SkeletonRootIndex]);
-					if (!ensure(SkeletonUid))
-					{
-						continue;
-					}
-
-					FAnimationTimes AnimationTimes;
-					UInterchangeSkeletalAnimationTrackNode* TrackNode = AcquireTrackNode(*SkeletonUid, AnimationTimes);
-
-					TSet<int32> JointsWithBindPose;
-					AcquireJointsWithBindPose(SkeletonRootIndex, JointsWithBindPose);
-
-					TSet<int32> NodesToAddToAnimation = JointsWithBindPose.Difference(AnimatedJointNodeIndices.Value);
-					for (int32 NodeToAddIndex : NodesToAddToAnimation)
-					{
-						const FString* NodeToAddUidPtr = GLTFNodeToInterchangeUidMap.Find(&GLTFNodes[NodeToAddIndex]);
-						if (!ensure(NodeToAddUidPtr))
-						{
-							continue;
-						}
-						TrackNode->SetAnimationPayloadKeyForSceneNodeUid(*NodeToAddUidPtr, BIND_POSE_FIX + LexToString(NodeToAddIndex), EInterchangeAnimationPayLoadType::BAKED);
-					}
-				}
-			}
-
 			void Process()
 			{
 				TMap<int32, TSet<int32>> RigidAnimation;							// (AnimatedNodeIndex, [Channels])
@@ -1120,8 +1070,6 @@ namespace UE::Interchange::Gltf::Private
 				ProcessRiggedAnimations(RiggedAnimations);
 				ProcessMorphTargetAnimations(MorphTargetAnimations);
 				ProcessRigidAnimation(RigidAnimation);
-
-				FixSkeletalAnimations(SkeletonRootToAnimatedJointNodeIndicesMap);
 			}
 		};
 	}
