Fixed potentially uninitialized FVectors in Interchange plugin.

diff --git a/Engine/Plugins/Interchange/Runtime/Source/Import/Private/Gltf/InterchangeGltfAnimation.cpp b/Engine/Plugins/Interchange/Runtime/Source/Import/Private/Gltf/InterchangeGltfAnimation.cpp
--- a/Engine/Plugins/Interchange/Runtime/Source/Import/Private/Gltf/InterchangeGltfAnimation.cpp
+++ b/Engine/Plugins/Interchange/Runtime/Source/Import/Private/Gltf/InterchangeGltfAnimation.cpp
@@ -86,7 +86,7 @@ namespace UE::Interchange::Gltf::Private
 		double TimeDelta = Seconds[kp1Index] - Seconds[kIndex];
 		double Time = (CurrentRequestedTimeStamp - Seconds[kIndex]) / TimeDelta;
 
-		T Value;
+		T Value = {};
 		switch (Interpolation)
 		{
 		case GLTF::FAnimation::EInterpolation::Linear:
diff --git a/Engine/Plugins/Interchange/Runtime/Source/Import/Private/Mesh/InterchangeStaticMeshFactory.cpp b/Engine/Plugins/Interchange/Runtime/Source/Import/Private/Mesh/InterchangeStaticMeshFactory.cpp
--- a/Engine/Plugins/Interchange/Runtime/Source/Import/Private/Mesh/InterchangeStaticMeshFactory.cpp
+++ b/Engine/Plugins/Interchange/Runtime/Source/Import/Private/Mesh/InterchangeStaticMeshFactory.cpp
@@ -1312,8 +1312,8 @@ bool UInterchangeStaticMeshFactory::AddCapsuleGeomFromVertices(const FImportAsse
 		return false;
 	}
 
-	FVector AxisStart;
-	FVector AxisEnd;
+	FVector AxisStart = FVector::ZeroVector;
+	FVector AxisEnd = FVector::ZeroVector;
 	float MaxDistSqr = 0.f;
 
 	for (int32 IndexA = 0; IndexA < VertexPositions.GetNumElements() - 1; IndexA++)
