Interchange Level Sequence:
- Add support of Object Property Tracks, currently the work needs to be completed to notify correctly the UObjects when the property is defined as private (for example StaticMesh in UStaticMeshComponent)
- Add Static/Skeletal Mesh object property tracks

diff --git a/Engine/Plugins/Interchange/Runtime/Source/Import/Private/Animation/InterchangeLevelSequenceFactory.cpp b/Engine/Plugins/Interchange/Runtime/Source/Import/Private/Animation/InterchangeLevelSequenceFactory.cpp
--- a/Engine/Plugins/Interchange/Runtime/Source/Import/Private/Animation/InterchangeLevelSequenceFactory.cpp
+++ b/Engine/Plugins/Interchange/Runtime/Source/Import/Private/Animation/InterchangeLevelSequenceFactory.cpp
@@ -20,6 +20,7 @@
 #include "Channels/MovieSceneDoubleChannel.h"
 #include "Channels/MovieSceneIntegerChannel.h"
 #include "Channels/MovieSceneByteChannel.h"
+#include "Channels/MovieSceneObjectPathChannel.h"
 #include "LevelSequence.h"
 #include "MovieScene.h"
 #include "MovieSceneSection.h"
@@ -148,7 +149,7 @@ namespace UE::Interchange::Private
 		{
 			const FFrameRate& FrameRate = MovieScene->GetTickResolution();
 
-			TMovieSceneChannelData<ValueType> Data = Channel.GetData();
+			auto Data = Channel.GetData();
 
 			Data.Reset();
 
@@ -166,7 +167,14 @@ namespace UE::Interchange::Private
 					this->MaxFrameNumber = FrameNumber;
 				}
 
-				Data.AddKey(FrameNumber, Values[KeyIndex]);
+				if constexpr(std::is_same_v<ChannelType, FMovieSceneObjectPathChannel>)
+				{
+					Data.AddKey(FrameNumber, FSoftObjectPath{ Values[KeyIndex] }.TryLoad());
+				}
+				else
+				{
+					Data.AddKey(FrameNumber, Values[KeyIndex]);
+				}
 			}
 		}
 
@@ -482,6 +490,7 @@ namespace UE::Interchange::Private
 		const FName IntegerChannelTypeName = FMovieSceneIntegerChannel::StaticStruct()->GetFName();
 		const FName BoolChannelTypeName = FMovieSceneBoolChannel::StaticStruct()->GetFName();
 		const FName EnumChannelTypeName = FMovieSceneByteChannel::StaticStruct()->GetFName();
+		const FName ObjectPathChannelTypeName = FMovieSceneObjectPathChannel::StaticStruct()->GetFName();
 
 		auto CopyToChannel = [this](auto Channel, const FRichCurve& Curve)
 		{
@@ -538,33 +547,39 @@ namespace UE::Interchange::Private
 			const bool bIsIntegerChannel = (ChannelTypeName == IntegerChannelTypeName);
 			const bool bIsDoubleChannel = (ChannelTypeName == DoubleChannelTypeName);
 			const bool bIsFloatChannel = (ChannelTypeName == FloatChannelTypeName);
+			const bool bIsObjectPathChannel = (ChannelTypeName == ObjectPathChannelTypeName);
 
 			if(!bIsBoolChannel &&
 			   !bIsEnumChannel &&
 			   !bIsIntegerChannel &&
 			   !bIsDoubleChannel &&
-			   !bIsFloatChannel)
+			   !bIsFloatChannel &&
+			   !bIsObjectPathChannel)
 			{
 				continue;
 			}
 
 			TArrayView<FMovieSceneChannel* const> Channels = ChannelEntry.GetChannels();
-			int32 NumChannels = (bIsBoolChannel || bIsEnumChannel || bIsIntegerChannel) ? PayloadData->StepCurves.Num() : PayloadData->Curves.Num();
+			int32 NumChannels = (bIsBoolChannel || bIsEnumChannel || bIsIntegerChannel || bIsObjectPathChannel) ? PayloadData->StepCurves.Num() : PayloadData->Curves.Num();
 			NumChannels = FMath::Min(NumChannels, Channels.Num());
 			for(int32 Index = 0; Index < NumChannels; ++Index)
 			{
 				FMovieSceneChannelHandle Channel = ChannelProxy.MakeHandle(ChannelTypeName, Index);
 				if(bIsBoolChannel)
 				{
-					UpdateStepChannel(*(Channel.Cast<FMovieSceneBoolChannel>().Get()), PayloadData->StepCurves[Index].KeyTimes, PayloadData->StepCurves[0].BooleanKeyValues.GetValue());
+					UpdateStepChannel(*(Channel.Cast<FMovieSceneBoolChannel>().Get()), PayloadData->StepCurves[Index].KeyTimes, PayloadData->StepCurves[Index].BooleanKeyValues.GetValue());
 				}
 				else if(bIsEnumChannel)
 				{
-					UpdateStepChannel(*(Channel.Cast<FMovieSceneByteChannel>().Get()), PayloadData->StepCurves[Index].KeyTimes, PayloadData->StepCurves[0].ByteKeyValues.GetValue());
+					UpdateStepChannel(*(Channel.Cast<FMovieSceneByteChannel>().Get()), PayloadData->StepCurves[Index].KeyTimes, PayloadData->StepCurves[Index].ByteKeyValues.GetValue());
 				}
 				else if(bIsIntegerChannel)
 				{
-					UpdateStepChannel(*(Channel.Cast<FMovieSceneIntegerChannel>().Get()), PayloadData->StepCurves[Index].KeyTimes, PayloadData->StepCurves[0].IntegerKeyValues.GetValue());
+					UpdateStepChannel(*(Channel.Cast<FMovieSceneIntegerChannel>().Get()), PayloadData->StepCurves[Index].KeyTimes, PayloadData->StepCurves[Index].IntegerKeyValues.GetValue());
+				}
+				else if(bIsObjectPathChannel)
+				{
+					UpdateStepChannel(*(Channel.Cast<FMovieSceneObjectPathChannel>().Get()), PayloadData->StepCurves[Index].KeyTimes, PayloadData->StepCurves[Index].StringKeyValues.GetValue());
 				}
 				else if(bIsFloatChannel)
 				{
diff --git a/Engine/Plugins/Interchange/Runtime/Source/Import/Private/Animation/InterchangeLevelSequenceHelper.cpp b/Engine/Plugins/Interchange/Runtime/Source/Import/Private/Animation/InterchangeLevelSequenceHelper.cpp
--- a/Engine/Plugins/Interchange/Runtime/Source/Import/Private/Animation/InterchangeLevelSequenceHelper.cpp
+++ b/Engine/Plugins/Interchange/Runtime/Source/Import/Private/Animation/InterchangeLevelSequenceHelper.cpp
@@ -15,15 +15,18 @@
 #include "Components/SplineComponent.h"
 #include "Components/StaticMeshComponent.h"
 
+#include "Engine/SkeletalMesh.h"
+#include "Engine/StaticMesh.h"
 #include "MovieScene.h"
 #include "MovieSceneSection.h"
 #include "Tracks/MovieSceneBoolTrack.h"
 #include "Tracks/MovieSceneByteTrack.h"
 #include "Tracks/MovieSceneColorTrack.h"
-#include "Tracks/MovieSceneVectorTrack.h"
 #include "Tracks/MovieSceneEnumTrack.h"
 #include "Tracks/MovieSceneFloatTrack.h"
 #include "Tracks/MovieSceneIntegerTrack.h"
+#include "Tracks/MovieSceneObjectPropertyTrack.h"
+#include "Tracks/MovieSceneVectorTrack.h"
 #include "Tracks/MovieSceneVisibilityTrack.h"
 
 namespace UE::Interchange::Private
@@ -345,18 +348,22 @@ namespace UE::Interchange::Private
 			{EInterchangePropertyTracks::SkinnedMeshVisibilityBasedAnimTickOption, {UMovieSceneEnumTrack::StaticClass()->GetName(),  GET_MEMBER_NAME_STRING_CHECKED(USkinnedMeshComponent, VisibilityBasedAnimTickOption) , TEXT("Visibility Based Anim Tick Option"), StaticEnum<EVisibilityBasedAnimTickOption>()}},
 
 			// Skeletal Mesh
+#if WITH_EDITORONLY_DATA
+			{EInterchangePropertyTracks::SkeletalMesh, {UMovieSceneObjectPropertyTrack::StaticClass()->GetName(),  USkeletalMeshComponent::GetSkeletalMeshAssetPropertyNameChecked().ToString(), TEXT("Skeletal Mesh Asset"), USkeletalMesh::StaticClass()}},
+#endif
 			{EInterchangePropertyTracks::SkeletalMeshAllowClothActors, {UMovieSceneBoolTrack::StaticClass()->GetName(),  GET_MEMBER_NAME_STRING_CHECKED(USkeletalMeshComponent, bAllowClothActors) , TEXT("Allow Cloth Actors")}},
 			{EInterchangePropertyTracks::SkeletalMeshAnimationMode, {UMovieSceneEnumTrack::StaticClass()->GetName(),  GET_MEMBER_NAME_STRING_CHECKED(USkeletalMeshComponent, USkeletalMeshComponent::GetAnimationModePropertyNameChecked()) , TEXT("Animation Mode"), StaticEnum<EAnimationMode::Type>()}},
 			{EInterchangePropertyTracks::SkeletalMeshClothBlendWeight, {UMovieSceneFloatTrack::StaticClass()->GetName(),  GET_MEMBER_NAME_STRING_CHECKED(USkeletalMeshComponent, ClothBlendWeight) , TEXT("Cloth Blend Weight")}},
 			{EInterchangePropertyTracks::SkeletalMeshClothMaxDistanceScale, {UMovieSceneFloatTrack::StaticClass()->GetName(),  GET_MEMBER_NAME_STRING_CHECKED(USkeletalMeshComponent, ClothMaxDistanceScale) , TEXT("Cloth Max Distance Scale")}},
 			
-			// Static Mesh
+			// Static Mesh				
+			{EInterchangePropertyTracks::StaticMesh, {UMovieSceneObjectPropertyTrack::StaticClass()->GetName(), UStaticMeshComponent::GetMemberNameChecked_StaticMesh().ToString() , TEXT("Static Mesh"), UStaticMesh::StaticClass()}},
 			{EInterchangePropertyTracks::StaticMeshDistanceFieldSelfShadowBias, {UMovieSceneFloatTrack::StaticClass()->GetName(),  GET_MEMBER_NAME_STRING_CHECKED(UStaticMeshComponent, DistanceFieldSelfShadowBias) , TEXT("Distance Field Self Shadow Bias")}},
 			{EInterchangePropertyTracks::StaticMeshEvaluateWorldPositionOffset, {UMovieSceneBoolTrack::StaticClass()->GetName(),  GET_MEMBER_NAME_STRING_CHECKED(UStaticMeshComponent, bEvaluateWorldPositionOffset) , TEXT("Evaluate World Position Offset")}},
 			{EInterchangePropertyTracks::StaticMeshEvaluateWorldPositionOffsetInRayTracing, {UMovieSceneBoolTrack::StaticClass()->GetName(),  GET_MEMBER_NAME_STRING_CHECKED(UStaticMeshComponent, bEvaluateWorldPositionOffsetInRayTracing) , TEXT("Evaluate World Position Offset in Ray Tracing")}},
 			{EInterchangePropertyTracks::StaticMeshForcedLodModel, {UMovieSceneIntegerTrack::StaticClass()->GetName(),  GET_MEMBER_NAME_STRING_CHECKED(UStaticMeshComponent, ForcedLodModel) , TEXT("Forced Lod Model")}},
 			{EInterchangePropertyTracks::StaticMeshReverseCulling, {UMovieSceneBoolTrack::StaticClass()->GetName(),  GET_MEMBER_NAME_STRING_CHECKED(UStaticMeshComponent, bReverseCulling) , TEXT("Reverse Culling")}},
-			{EInterchangePropertyTracks::StaticMeshWorldPositionOffsetDisableDistance, {UMovieSceneIntegerTrack::StaticClass()->GetName(),  GET_MEMBER_NAME_STRING_CHECKED(UStaticMeshComponent, WorldPositionOffsetDisableDistance) , TEXT("World Position Offset Disable Distance")}},			
+			{EInterchangePropertyTracks::StaticMeshWorldPositionOffsetDisableDistance, {UMovieSceneIntegerTrack::StaticClass()->GetName(),  GET_MEMBER_NAME_STRING_CHECKED(UStaticMeshComponent, WorldPositionOffsetDisableDistance) , TEXT("World Position Offset Disable Distance")}},
 		}
 	{}
 
@@ -395,23 +402,29 @@ namespace UE::Interchange::Private
 			PropertyTrack->RemoveAllAnimationData();
 		}
 
-		if(InterchangePropertyTrack->EnumClass)
+		if(UEnum* const* EnumClass = InterchangePropertyTrack->VariantProperty.TryGet<UEnum*>())
 		{
 			if(UMovieSceneByteTrack* ByteTrack = Cast<UMovieSceneByteTrack>(PropertyTrack))
 			{
-				ByteTrack->SetEnum(InterchangePropertyTrack->EnumClass);
+				ByteTrack->SetEnum(*EnumClass);
 			}
 		}
-
-		if(InterchangePropertyTrack->NumChannelsUsed.IsSet())
+		else if(const int32* NumChannelsUsed = InterchangePropertyTrack->VariantProperty.TryGet<int32>())
 		{
 			if(UMovieSceneDoubleVectorTrack* DoubleVectorTrack = Cast<UMovieSceneDoubleVectorTrack>(PropertyTrack))
 			{
-				DoubleVectorTrack->SetNumChannelsUsed(*InterchangePropertyTrack->NumChannelsUsed);
+				DoubleVectorTrack->SetNumChannelsUsed(*NumChannelsUsed);
 			}
 			else if(UMovieSceneFloatVectorTrack* FloatVectorTrack = Cast<UMovieSceneFloatVectorTrack>(PropertyTrack))
 			{
-				FloatVectorTrack->SetNumChannelsUsed(*InterchangePropertyTrack->NumChannelsUsed);
+				FloatVectorTrack->SetNumChannelsUsed(*NumChannelsUsed);
+			}
+		}
+		else if(UClass* const* ObjectPropertyClass = InterchangePropertyTrack->VariantProperty.TryGet<UClass*>())
+		{
+			if(UMovieSceneObjectPropertyTrack* ObjectPropertyTrack = Cast<UMovieSceneObjectPropertyTrack>(PropertyTrack))
+			{
+				ObjectPropertyTrack->PropertyClass = *ObjectPropertyClass;
 			}
 		}
 
diff --git a/Engine/Plugins/Interchange/Runtime/Source/Import/Private/Animation/InterchangeLevelSequenceHelper.h b/Engine/Plugins/Interchange/Runtime/Source/Import/Private/Animation/InterchangeLevelSequenceHelper.h
--- a/Engine/Plugins/Interchange/Runtime/Source/Import/Private/Animation/InterchangeLevelSequenceHelper.h
+++ b/Engine/Plugins/Interchange/Runtime/Source/Import/Private/Animation/InterchangeLevelSequenceHelper.h
@@ -67,21 +67,27 @@ namespace UE::Interchange::Private
 				: ClassType{ MoveTemp(ClassType) }
 				, Path{ MoveTemp(Path) }
 				, Name{ MoveTemp(Name) }
-				, EnumClass{EnumClass}
+				, VariantProperty{ TInPlaceType<UEnum*>{}, EnumClass}
 			{}
 
 			FInterchangeProperty(FString&& ClassType, FString&& Path, FString&& Name, int32 NumChannelsUsed)
-				: ClassType{ std::move(ClassType) }
+				: ClassType{ MoveTemp(ClassType) }
+				, Path{ MoveTemp(Path) }
+				, Name{ MoveTemp(Name) }
+				, VariantProperty{ TInPlaceType<int32>{}, NumChannelsUsed }
+			{}
+
+			FInterchangeProperty(FString && ClassType, FString && Path, FString && Name, UClass * ObjectPropertyClass)
+				: ClassType{ MoveTemp(ClassType) }
 				, Path{ MoveTemp(Path) }
 				, Name{ MoveTemp(Name) }
-				, NumChannelsUsed{ NumChannelsUsed }
+				, VariantProperty{ TInPlaceType<UClass*>{}, ObjectPropertyClass }
 			{}
 
 			FString ClassType; // Float, Double, Byte, etc. Basically the class name of the UMovieSceneTrack
 			FString Path;
 			FName Name;
-			UEnum* EnumClass = nullptr; // Only used for Enum property tracks
-			TOptional<int32> NumChannelsUsed; // Only used for Vector property tracks (Vector tracks can only have 2-4 channels)
+			TVariant<UEnum*, int32, UClass*> VariantProperty; // These are mutually exclusive, we can only have one of a kind (Either an Enum, or Number of channels used for a Vector, or an Object Path)
 		};
 
 		TMap<EInterchangePropertyTracks, FInterchangeProperty> PropertyTracks;
diff --git a/Engine/Plugins/Interchange/Runtime/Source/Nodes/Public/InterchangeAnimationDefinitions.h b/Engine/Plugins/Interchange/Runtime/Source/Nodes/Public/InterchangeAnimationDefinitions.h
--- a/Engine/Plugins/Interchange/Runtime/Source/Nodes/Public/InterchangeAnimationDefinitions.h
+++ b/Engine/Plugins/Interchange/Runtime/Source/Nodes/Public/InterchangeAnimationDefinitions.h
@@ -325,6 +325,7 @@ enum class EInterchangePropertyTracks : int32
 	SkinnedMeshVisibilityBasedAnimTickOption,
 
 	/** Skeletal Mesh Properties */
+	SkeletalMesh,
 	SkeletalMeshAllowClothActors,
 	SkeletalMeshAnimationMode,
 	SkeletalMeshClothBlendWeight,
@@ -338,6 +339,8 @@ enum class EInterchangePropertyTracks : int32
 	StaticMeshReverseCulling,
 	StaticMeshWorldPositionOffsetDisableDistance,
 
+	StaticMesh,
+
 	None = -1 UMETA(hidden),
 };
 
