Interchange: Fix crash with UInterchangeFbxTranslatorSettings being garbage collect.

#rnx

diff --git a/Engine/Plugins/Interchange/Runtime/Source/Import/Private/Fbx/InterchangeFbxTranslator.cpp b/Engine/Plugins/Interchange/Runtime/Source/Import/Private/Fbx/InterchangeFbxTranslator.cpp
--- a/Engine/Plugins/Interchange/Runtime/Source/Import/Private/Fbx/InterchangeFbxTranslator.cpp
+++ b/Engine/Plugins/Interchange/Runtime/Source/Import/Private/Fbx/InterchangeFbxTranslator.cpp
@@ -303,9 +303,9 @@ void UInterchangeFbxTranslator::SetSettings(const UInterchangeTranslatorSettings
 		CacheFbxTranslatorSettings->ClearInternalFlags(EInternalObjectFlags::Async);
 		CacheFbxTranslatorSettings = nullptr;
 	}
-	if (InterchangeTranslatorSettings)
+	if (const UInterchangeFbxTranslatorSettings* InterchangeFbxTranslatorSettings = Cast<UInterchangeFbxTranslatorSettings>(InterchangeTranslatorSettings))
 	{
-		CacheFbxTranslatorSettings = DuplicateObject<UInterchangeFbxTranslatorSettings>(Cast<UInterchangeFbxTranslatorSettings>(InterchangeTranslatorSettings), GetTransientPackage());
+		CacheFbxTranslatorSettings = DuplicateObject<UInterchangeFbxTranslatorSettings>(InterchangeFbxTranslatorSettings, GetTransientPackage());
 		CacheFbxTranslatorSettings->ClearInternalFlags(EInternalObjectFlags::Async);
 		CacheFbxTranslatorSettings->SetFlags(RF_Standalone);
 	}
diff --git a/Engine/Source/Runtime/Interchange/Engine/Private/InterchangeAssetImportData.cpp b/Engine/Source/Runtime/Interchange/Engine/Private/InterchangeAssetImportData.cpp
--- a/Engine/Source/Runtime/Interchange/Engine/Private/InterchangeAssetImportData.cpp
+++ b/Engine/Source/Runtime/Interchange/Engine/Private/InterchangeAssetImportData.cpp
@@ -280,6 +280,12 @@ void UInterchangeAssetImportData::ProcessTranslatorCache() const
 {
 	if (UInterchangeManager::IsInterchangeImportEnabled())
 	{
+		//Verify our transient object was not garbage collect
+		if (TransientTranslatorSettings && (TransientTranslatorSettings->IsGarbageEliminationEnabled() || TransientTranslatorSettings->HasAnyFlags(RF_BeginDestroyed | RF_FinishDestroyed)) )
+		{
+			TransientTranslatorSettings = nullptr;
+		}
+
 		//de-serialize
 		if (!TransientTranslatorSettings && !CachedTranslatorSettings.Key.IsEmpty())
 		{
