[Interchange] Global socket import fixes.

#rnx

diff --git a/Engine/Plugins/Interchange/Runtime/Source/Pipelines/Private/InterchangePipelineMeshesUtilities.cpp b/Engine/Plugins/Interchange/Runtime/Source/Pipelines/Private/InterchangePipelineMeshesUtilities.cpp
--- a/Engine/Plugins/Interchange/Runtime/Source/Pipelines/Private/InterchangePipelineMeshesUtilities.cpp
+++ b/Engine/Plugins/Interchange/Runtime/Source/Pipelines/Private/InterchangePipelineMeshesUtilities.cpp
@@ -416,37 +416,59 @@ UInterchangePipelineMeshesUtilities* UInterchangePipelineMeshesUtilities::Create
 	// Do a second pass to discover sockets
 	if (bHasSockets)
 	{
-		BaseNodeContainer->IterateNodes(
-			[&PipelineMeshesUtilities, &BaseNodeContainer](const FString& NodeUid, const UInterchangeBaseNode* Node)
-			{
-				if (const UInterchangeSceneNode* SceneNode = Cast<UInterchangeSceneNode>(Node))
+		if (PipelineMeshesUtilities->MeshGeometriesPerMeshUid.Num() == 1)
+		{
+			//Import of Global Sockets (only done, in case there are only 1 mesh in the source data) :
+			TArray<FString> SocketUIDs;
+			BaseNodeContainer->IterateNodes(
+				[&PipelineMeshesUtilities, &BaseNodeContainer, &SocketUIDs](const FString& NodeUid, const UInterchangeBaseNode* Node)
 				{
-					if (IsSceneNodeASocket(SceneNode))
+					if (const UInterchangeSceneNode* SceneNode = Cast<UInterchangeSceneNode>(Node))
 					{
-						FString MeshUid;
-						if (!SceneNode->GetCustomAssetInstanceUid(MeshUid))
+						if (IsSceneNodeASocket(SceneNode))
 						{
-							const UInterchangeSceneNode* ParentMeshSceneNode = Cast<UInterchangeSceneNode>(BaseNodeContainer->GetNode(SceneNode->GetParentUid()));
-							while (ParentMeshSceneNode)
+							SocketUIDs.Add(SceneNode->GetUniqueID());
+						}
+					}
+				}
+			);
+			PipelineMeshesUtilities->MeshGeometriesPerMeshUid.begin().Value().AttachedSocketUids = SocketUIDs;
+		}
+		else
+		{
+			//Import of Local Sockets:
+			BaseNodeContainer->IterateNodes(
+				[&PipelineMeshesUtilities, &BaseNodeContainer](const FString& NodeUid, const UInterchangeBaseNode* Node)
+				{
+					if (const UInterchangeSceneNode* SceneNode = Cast<UInterchangeSceneNode>(Node))
+					{
+						if (IsSceneNodeASocket(SceneNode))
+						{
+							FString MeshUid;
+							if (!SceneNode->GetCustomAssetInstanceUid(MeshUid))
 							{
-								if (ParentMeshSceneNode->GetCustomAssetInstanceUid(MeshUid))
+								const UInterchangeSceneNode* ParentMeshSceneNode = Cast<UInterchangeSceneNode>(BaseNodeContainer->GetNode(SceneNode->GetParentUid()));
+								while (ParentMeshSceneNode)
 								{
-									break;
-								}
+									if (ParentMeshSceneNode->GetCustomAssetInstanceUid(MeshUid))
+									{
+										break;
+									}
 
-								ParentMeshSceneNode = Cast<UInterchangeSceneNode>(BaseNodeContainer->GetNode(ParentMeshSceneNode->GetParentUid()));
+									ParentMeshSceneNode = Cast<UInterchangeSceneNode>(BaseNodeContainer->GetNode(ParentMeshSceneNode->GetParentUid()));
+								}
 							}
-						}
 
-						if (!MeshUid.IsEmpty())
-						{
-							FInterchangeMeshGeometry& MeshGeometry = PipelineMeshesUtilities->MeshGeometriesPerMeshUid.FindChecked(MeshUid);
-							MeshGeometry.AttachedSocketUids.Add(SceneNode->GetUniqueID());
+							if (!MeshUid.IsEmpty())
+							{
+								FInterchangeMeshGeometry& MeshGeometry = PipelineMeshesUtilities->MeshGeometriesPerMeshUid.FindChecked(MeshUid);
+								MeshGeometry.AttachedSocketUids.Add(SceneNode->GetUniqueID());
+							}
 						}
 					}
 				}
-			}
-		);
+			);
+		}
 	}
 
 
