Interchange: fix skeletalmesh re-import socket issue when a socket bone go away.

#rnx

diff --git a/Engine/Plugins/Interchange/Runtime/Source/Import/Private/Mesh/InterchangeSkeletalMeshFactory.cpp b/Engine/Plugins/Interchange/Runtime/Source/Import/Private/Mesh/InterchangeSkeletalMeshFactory.cpp
--- a/Engine/Plugins/Interchange/Runtime/Source/Import/Private/Mesh/InterchangeSkeletalMeshFactory.cpp
+++ b/Engine/Plugins/Interchange/Runtime/Source/Import/Private/Mesh/InterchangeSkeletalMeshFactory.cpp
@@ -9,6 +9,7 @@
 #include "Components.h"
 #include "CoreGlobals.h"
 #include "Engine/SkeletalMesh.h"
+#include "Engine/SkeletalMeshSocket.h"
 #include "Engine/SkinnedAssetAsyncCompileUtils.h"
 #include "Engine/SkinnedAssetCommon.h"
 #include "GenericPlatform/GenericPlatformMisc.h"
@@ -1745,6 +1746,11 @@ UInterchangeFactoryBase::FImportAssetResult UInterchangeSkeletalMeshFactory::End
 
 	if (USkeleton* SkeletonReference = ImportAssetObjectData.SkeletonReference)
 	{
+		if (SkeletalMesh->GetSkeleton() != SkeletonReference)
+		{
+			SkeletalMesh->SetSkeleton(SkeletonReference);
+		}
+
 		constexpr bool bShowProgress = false;
 		if ((!ImportAssetObjectData.bApplyGeometryOnly || !ImportAssetObjectData.bIsReImport) && !SkeletonReference->MergeAllBonesToBoneTree(SkeletalMesh, bShowProgress))
 		{
@@ -1836,10 +1842,22 @@ UInterchangeFactoryBase::FImportAssetResult UInterchangeSkeletalMeshFactory::End
 				Async(EAsyncExecution::TaskGraphMainThread, MoveTemp(RecreateSkeleton)).Wait();
 			}
 		}
-		if (SkeletalMesh->GetSkeleton() != SkeletonReference)
+
+		//Cleanup the reference skeleton sockets that doesn't fit anymmore a bone name.
+		TArray<TObjectPtr<USkeletalMeshSocket>>& Sockets = SkeletalMesh->GetMeshOnlySocketList();
+
+		for (int32 SocketIndex = Sockets.Num()-1; SocketIndex >= 0; SocketIndex--)
 		{
-			SkeletalMesh->SetSkeleton(SkeletonReference);
+			// Find bone index the socket is attached to.
+			USkeletalMeshSocket* Socket = Sockets[SocketIndex];
+			int32 SocketBoneIndex = SkeletalMesh->GetRefSkeleton().FindBoneIndex(Socket->BoneName);
+			// If this LOD does not contain the socket bone, abort import.
+			if (SocketBoneIndex == INDEX_NONE)
+			{
+				Sockets.RemoveAt(SocketIndex, EAllowShrinking::No);
+			}
 		}
+		Sockets.Shrink();
 	}
 	else
 	{
