[Interchange - glTF] Skin usage fixes.

#rnx

diff --git a/Engine/Plugins/Interchange/Runtime/Source/Parsers/GLTFCore/Private/GLTF/GLTFReader.cpp b/Engine/Plugins/Interchange/Runtime/Source/Parsers/GLTFCore/Private/GLTF/GLTFReader.cpp
--- a/Engine/Plugins/Interchange/Runtime/Source/Parsers/GLTFCore/Private/GLTF/GLTFReader.cpp
+++ b/Engine/Plugins/Interchange/Runtime/Source/Parsers/GLTFCore/Private/GLTF/GLTFReader.cpp
@@ -1203,6 +1203,8 @@ namespace GLTF
 		if (!SetupObjects(SkinCount, TEXT("skins"), [this](const FJsonObject& Object) { SetupSkin(Object); })) { return; }
 
 		{//BuildRootJoints can affect the node hierarchy and data, as the Animation setup currently stores references to the Nodes, we have to do these setups before the SetupAnimations.
+			SetupUsedSkins();
+
 			SetupNodesType();
 
 			GenerateInverseBindPosesPerSkinIndices();
@@ -1269,6 +1271,41 @@ namespace GLTF
 		return true;
 	}
 
+	void FFileReader::SetupUsedSkins() const
+	{
+		//Acquire Skin usages:
+		TBitArray SkinIndicesUsage(false, Asset->Skins.Num());
+		for (size_t Index = 0; Index < Asset->Nodes.Num(); Index++)
+		{
+			int32 SkinIndex = Asset->Nodes[Index].Skindex;
+			if (SkinIndex != INDEX_NONE)
+			{
+				SkinIndicesUsage[SkinIndex] = true;
+			}
+		}
+
+		FString UnUsedIndicesString;
+		for (size_t SkinIndex = 0; SkinIndex < SkinIndicesUsage.Num(); SkinIndex++)
+		{
+			FSkinInfo& Skin = Asset->Skins[SkinIndex];
+			Skin.bUsed = SkinIndicesUsage[SkinIndex];
+
+			if (!Skin.bUsed)
+			{
+				if (UnUsedIndicesString.Len() > 0)
+				{
+					UnUsedIndicesString += TEXT(",");
+				}
+				UnUsedIndicesString += FString::FromInt(SkinIndex);
+			}
+		}
+
+		if (UnUsedIndicesString.Len() > 0)
+		{
+			Messages.Emplace(EMessageSeverity::Warning, FString::Printf(TEXT("Skin objects unused. At indices: %s."), *UnUsedIndicesString));
+		}
+	}
+
 	void FFileReader::SetupNodesType() const
 	{
 		// setup node types
@@ -1296,8 +1333,14 @@ namespace GLTF
 				}
 			}
 		}
+
 		for (const FSkinInfo& Skin : Asset->Skins)
 		{
+			if (!Skin.bUsed)
+			{
+				continue;
+			}
+
 			for (int32 JointIndex : Skin.Joints)
 			{
 				ensure(Asset->Nodes[JointIndex].Type == FNode::EType::None 
@@ -1511,6 +1554,11 @@ namespace GLTF
 			{
 				GLTF::FSkinInfo& Skin = Asset->Skins[SkinIndex];
 
+				if (!Skin.bUsed)
+				{
+					continue;
+				}
+
 				//0. Group RootJointNodes by ParentIndices
 				TMap<int32, FRootJoints> ParentToRootJointIndices;
 				for (size_t JointIndex = 0; JointIndex < Skin.Joints.Num(); JointIndex++)
@@ -1598,6 +1646,8 @@ namespace GLTF
 					OriginalNode.Children.Add(Node.Index);
 					Node.Type = GLTF::FNode::EType::Joint;
 					Node.ParentIndex = OriginalNode.Index;
+					Node.bHasLocalBindPose = true;
+					Node.LocalBindPose = FTransform::Identity;
 
 					for (size_t RootJointIndex = 0; RootJointIndex < Group.Value.Indices.Num(); RootJointIndex++)
 					{
diff --git a/Engine/Plugins/Interchange/Runtime/Source/Parsers/GLTFCore/Public/GLTFNode.h b/Engine/Plugins/Interchange/Runtime/Source/Parsers/GLTFCore/Public/GLTFNode.h
--- a/Engine/Plugins/Interchange/Runtime/Source/Parsers/GLTFCore/Public/GLTFNode.h
+++ b/Engine/Plugins/Interchange/Runtime/Source/Parsers/GLTFCore/Public/GLTFNode.h
@@ -155,15 +155,16 @@ namespace GLTF
 		const FAccessor& InverseBindMatrices;
 		FString          Name;
 		TArray<int32>    Joints;    // each is an index into FAsset::Nodes
-		int32            Skeleton;  // root node, index into FAsset::Nodes
+		int32            Skeleton = INDEX_NONE;  // root node, index into FAsset::Nodes
 
 		TMap<FString, FString> Extras;
 
+		bool bUsed = false;
+
 		FString          UniqueId; //will be generated in FAsset::GenerateNames
 
 		FSkinInfo(const FAccessor& InverseBindMatrices)
 		    : InverseBindMatrices(InverseBindMatrices)
-		    , Skeleton(INDEX_NONE)
 		{
 		}
 	};
diff --git a/Engine/Plugins/Interchange/Runtime/Source/Parsers/GLTFCore/Public/GLTFReader.h b/Engine/Plugins/Interchange/Runtime/Source/Parsers/GLTFCore/Public/GLTFReader.h
--- a/Engine/Plugins/Interchange/Runtime/Source/Parsers/GLTFCore/Public/GLTFReader.h
+++ b/Engine/Plugins/Interchange/Runtime/Source/Parsers/GLTFCore/Public/GLTFReader.h
@@ -60,6 +60,7 @@ namespace GLTF
 
 		template <typename SetupFunc>
 		bool SetupObjects(uint32 ObjectCount, const TCHAR* FieldName, SetupFunc Func) const;
+		void SetupUsedSkins() const;
 		void SetupNodesType() const;
 
 		void GenerateInverseBindPosesPerSkinIndices() const; //Per node
