MaterialX:
- Flip UV coordinates to be top bottom left
- deprecating Material expression Place2D in favor of the Material Function
- deprecating Material expression Rotate2d in favor of Rotator

MaterialExpressionRotator:
- Add MinimalAPI tag to avoid link error

#virtualized

diff --git a/Engine/Plugins/Interchange/Runtime/Source/Import/Private/MaterialX/MaterialExpressions/MaterialExpressionPlace2D.cpp b/Engine/Plugins/Interchange/Runtime/Source/Import/Private/MaterialX/MaterialExpressions/MaterialExpressionPlace2D.cpp
--- a/Engine/Plugins/Interchange/Runtime/Source/Import/Private/MaterialX/MaterialExpressions/MaterialExpressionPlace2D.cpp
+++ b/Engine/Plugins/Interchange/Runtime/Source/Import/Private/MaterialX/MaterialExpressions/MaterialExpressionPlace2D.cpp
@@ -14,6 +14,8 @@
 #include "Materials/MaterialExpressionSubtract.h"
 #include "Materials/MaterialExpressionTextureCoordinate.h"
 
+#include "Materials/MaterialExpressionRotator.h"
+
 #include UE_INLINE_GENERATED_CPP_BY_NAME(MaterialExpressionPlace2D)
 
 #define LOCTEXT_NAMESPACE "MaterialExpressionMaterialXPlace2D"
@@ -116,12 +118,6 @@ int32 UMaterialExpressionMaterialXPlace2D::Compile(FMaterialCompiler* Compiler,
 	FExpressionInput ApplyScale = NewExpression<UMaterialExpressionDivide>(SubPivot, AppendVectorScale);
 
 	FExpressionInput ApplyRot;
-	{
-		UMaterialExpressionMaterialXRotate2D* Rotate2D = NewObject<UMaterialExpressionMaterialXRotate2D>();
-		Rotate2D->Input = ApplyScale;
-		Rotate2D->RotationAngle = RotationAngle;
-		ApplyRot.Connect(0, Rotate2D);
-	}
 
 	FExpressionInput Applyoffset = NewExpression<UMaterialExpressionSubtract>(ApplyRot, Offset);
 	FExpressionInput AddPivot = NewExpression<UMaterialExpressionAdd>(Applyoffset, Pivot);
diff --git a/Engine/Plugins/Interchange/Runtime/Source/Import/Private/MaterialX/MaterialExpressions/MaterialExpressionPlace2D.h b/Engine/Plugins/Interchange/Runtime/Source/Import/Private/MaterialX/MaterialExpressions/MaterialExpressionPlace2D.h
--- a/Engine/Plugins/Interchange/Runtime/Source/Import/Private/MaterialX/MaterialExpressions/MaterialExpressionPlace2D.h
+++ b/Engine/Plugins/Interchange/Runtime/Source/Import/Private/MaterialX/MaterialExpressions/MaterialExpressionPlace2D.h
@@ -14,7 +14,7 @@
  * -pivot, translate, rotate, scale, +pivot which does not introduce texture shear.
  * Default is 0 "SRT" for backward compatibility.*/
 UCLASS(MinimalAPI, collapsecategories, hidecategories = Object, meta = (Private))
-class UMaterialExpressionMaterialXPlace2D : public UMaterialExpression
+class UE_DEPRECATED(5.5, "UMaterialExpressionMaterialXPlace2D is now deprecated, please use Material Function MX_Place2D") UMaterialExpressionMaterialXPlace2D : public UMaterialExpression
 {
 	GENERATED_UCLASS_BODY()
 
diff --git a/Engine/Plugins/Interchange/Runtime/Source/Import/Private/MaterialX/MaterialExpressions/MaterialExpressionRotate2D.h b/Engine/Plugins/Interchange/Runtime/Source/Import/Private/MaterialX/MaterialExpressions/MaterialExpressionRotate2D.h
--- a/Engine/Plugins/Interchange/Runtime/Source/Import/Private/MaterialX/MaterialExpressions/MaterialExpressionRotate2D.h
+++ b/Engine/Plugins/Interchange/Runtime/Source/Import/Private/MaterialX/MaterialExpressions/MaterialExpressionRotate2D.h
@@ -9,7 +9,7 @@
  * A material expression that rotates a vector2 value about the origin in 2D.
  */
 UCLASS(collapsecategories, hidecategories = Object, MinimalAPI, meta = (Private))
-class UMaterialExpressionMaterialXRotate2D : public UMaterialExpression
+class UE_DEPRECATED(5.5, "UMaterialExpressionMaterialXRotate2D is now deprecated, please use UMaterialExpressionRotator") UMaterialExpressionMaterialXRotate2D : public UMaterialExpression
 {
 	GENERATED_UCLASS_BODY()
 
diff --git a/Engine/Plugins/Interchange/Runtime/Source/Import/Private/MaterialX/MaterialXUtils/MaterialXManager.cpp b/Engine/Plugins/Interchange/Runtime/Source/Import/Private/MaterialX/MaterialXUtils/MaterialXManager.cpp
--- a/Engine/Plugins/Interchange/Runtime/Source/Import/Private/MaterialX/MaterialXUtils/MaterialXManager.cpp
+++ b/Engine/Plugins/Interchange/Runtime/Source/Import/Private/MaterialX/MaterialXUtils/MaterialXManager.cpp
@@ -85,11 +85,10 @@ FMaterialXManager::FMaterialXManager()
 		{{MaterialX::Category::Normalize,      TEXT("in")},          TEXT("VectorInput")},
 		{{MaterialX::Category::Outside,        TEXT("in")},          TEXT("A")},				// Outside is treated as Multiply node
 		{{MaterialX::Category::Outside,        TEXT("mask")},        TEXT("B")},				// Outside is treated as Multiply node
-		{{MaterialX::Category::Place2D,        TEXT("pivot")},       TEXT("Pivot")},
-		{{MaterialX::Category::Place2D,        TEXT("rotate")},      TEXT("RotationAngle")},
 		{{MaterialX::Category::Power,          TEXT("in1")},         TEXT("Base")},
 		{{MaterialX::Category::Power,	       TEXT("in2")},         TEXT("Exponent")},
-		{{MaterialX::Category::Rotate2D,       TEXT("amount")},      TEXT("RotationAngle")},
+		{{MaterialX::Category::Rotate2D,       TEXT("in")},          TEXT("Coordinate")},
+		{{MaterialX::Category::Rotate2D,       TEXT("amount")},      TEXT("Time")},
 		{{MaterialX::Category::Rotate3D,       TEXT("amount")},      TEXT("RotationAngle")},
 		{{MaterialX::Category::Rotate3D,       TEXT("axis")},		 TEXT("NormalizedRotationAxis")},
 		{{MaterialX::Category::Rotate3D,       TEXT("in")},          TEXT("Position")},
@@ -119,9 +118,7 @@ FMaterialXManager::FMaterialXManager()
 		{MaterialX::Category::Modulo,       TEXT("Fmod")},
 		{MaterialX::Category::Multiply,     TEXT("Multiply")},
 		{MaterialX::Category::Normalize,    TEXT("Normalize")},
-		{MaterialX::Category::Place2D,		TEXT("MaterialXPlace2D")},
 		{MaterialX::Category::Power,        TEXT("Power")},
-		{MaterialX::Category::Rotate2D,		TEXT("MaterialXRotate2D")},
 		{MaterialX::Category::RampLR,       TEXT("MaterialXRampLeftRight")},
 		{MaterialX::Category::RampTB,       TEXT("MaterialXRampTopBottom")},
 		{MaterialX::Category::Sign,         TEXT("Sign")},
@@ -181,6 +178,7 @@ FMaterialXManager::FMaterialXManager()
 		TEXT("Base"),
 		TEXT("C"),
 		TEXT("Center"),
+		TEXT("Coordinate"),
 		TEXT("Coordinates"),
 		TEXT("D"),
 		TEXT("Exponent"),
@@ -202,6 +200,7 @@ FMaterialXManager::FMaterialXManager()
 		TEXT("TargetLow"),
 		TEXT("TargetHigh"),
 		TEXT("Temp"),
+		TEXT("Time"),
 		TEXT("Value"),
 		TEXT("VectorInput"),
 		TEXT("X"),
@@ -241,7 +240,9 @@ FMaterialXManager::FMaterialXManager()
 			// Utility nodes
 			{mx::Category::ArtisticIOR,				FMaterialXMaterialFunction{TInPlaceType<FString>{}, TEXT("/Interchange/Functions/MX_Artistic_IOR.MX_Artistic_IOR")}},
 			{mx::Category::RoughnessAnisotropy,		FMaterialXMaterialFunction{TInPlaceType<FString>{}, TEXT("/Interchange/Functions/MX_Roughness_Anisotropy.MX_Roughness_Anisotropy")}},
-			{mx::Category::RoughnessDual,			FMaterialXMaterialFunction{TInPlaceType<FString>{}, TEXT("/Interchange/Functions/MX_Roughness_Dual.MX_Roughness_Dual")}}
+			{mx::Category::RoughnessDual,			FMaterialXMaterialFunction{TInPlaceType<FString>{}, TEXT("/Interchange/Functions/MX_Roughness_Dual.MX_Roughness_Dual")}},
+			// Math
+			{mx::Category::Place2D,					FMaterialXMaterialFunction{TInPlaceType<FString>{}, TEXT("/Interchange/Functions/MX_Place2D.MX_Place2D")}},
 		};
 
 		if(bIsSubstrateEnabled)
@@ -380,10 +381,12 @@ namespace UE::Interchange::MaterialX
 			return bAllLoaded;
 		};
 
-		static const bool bPackagesLoaded =	ArePackagesLoaded({ TEXT("MaterialFunction'/Engine/Functions/Engine_MaterialFunctions03/Procedurals/NormalFromHeightmap.NormalFromHeightmap'"),
-															    TEXT("MaterialFunction'/Interchange/Functions/MX_Artistic_IOR.MX_Artistic_IOR'"),
-															    TEXT("MaterialFunction'/Interchange/Functions/MX_Roughness_Anisotropy.MX_Roughness_Anisotropy'"), 
-															    TEXT("MaterialFunction'/Interchange/Functions/MX_Roughness_Dual.MX_Roughness_Dual'")});
+		static const bool bPackagesLoaded = ArePackagesLoaded({ TEXT("MaterialFunction'/Engine/Functions/Engine_MaterialFunctions03/Procedurals/NormalFromHeightmap.NormalFromHeightmap'"),
+																TEXT("MaterialFunction'/Interchange/Functions/MX_Artistic_IOR.MX_Artistic_IOR'"),
+																TEXT("MaterialFunction'/Interchange/Functions/MX_Roughness_Anisotropy.MX_Roughness_Anisotropy'"),
+																TEXT("MaterialFunction'/Interchange/Functions/MX_Roughness_Dual.MX_Roughness_Dual'"),
+																TEXT("MaterialFunction'/Interchange/Functions/MX_Place2D.MX_Place2D'"),
+															  });
 
 		return bPackagesLoaded;
 #else
diff --git a/Engine/Plugins/Interchange/Runtime/Source/Import/Private/MaterialX/MaterialXUtils/MaterialXSurfaceShaderAbstract.cpp b/Engine/Plugins/Interchange/Runtime/Source/Import/Private/MaterialX/MaterialXUtils/MaterialXSurfaceShaderAbstract.cpp
--- a/Engine/Plugins/Interchange/Runtime/Source/Import/Private/MaterialX/MaterialXUtils/MaterialXSurfaceShaderAbstract.cpp
+++ b/Engine/Plugins/Interchange/Runtime/Source/Import/Private/MaterialX/MaterialXUtils/MaterialXSurfaceShaderAbstract.cpp
@@ -457,6 +457,39 @@ void FMaterialXSurfaceShaderAbstract::ConnectTransformVectorInputToOutput(const
 	UInterchangeShaderPortsAPI::ConnectDefaultOuputToInput(Connect.ParentShaderNode, Connect.InputChannelName, TransformNode->GetUniqueID());
 }
 
+void FMaterialXSurfaceShaderAbstract::ConnectRotate2DInputToOutput(const FConnectNode& Connect)
+{
+	using namespace UE::Interchange::Materials::Standard::Nodes;
+
+	UInterchangeShaderNode* Rotate2DNode = CreateShaderNode(Connect.UpstreamNode->getName().c_str(), Rotator::Name.ToString());
+
+
+	if(mx::InputPtr Input = Connect.UpstreamNode->getInput("in"))
+	{
+		AddAttributeFromValueOrInterface(Input, GetInputName(Input), Rotate2DNode);
+	}
+
+	// Amount is in degrees whereas Time (which in our case is the angle) input is in radians
+	if(mx::InputPtr Input = Connect.UpstreamNode->getInput("amount"))
+	{
+		FString InputName = GetInputName(Input);
+		UInterchangeShaderNode* DegreesToRadiansNode = CreateShaderNode((Connect.UpstreamNode->getName() + "multiply").c_str(), TEXT("Multiply"));
+		constexpr float DegreesToRadians = UE_PI / 180.f;
+		DegreesToRadiansNode->AddFloatAttribute(UInterchangeShaderPortsAPI::MakeInputValueKey(TEXT("B")), DegreesToRadians);
+
+		//If it's a value we should always attach it on the A input of Divide node
+		AddAttributeFromValueOrInterface(Input, TEXT("A"), DegreesToRadiansNode);
+		UInterchangeShaderPortsAPI::ConnectDefaultOuputToInput(Rotate2DNode, InputName, DegreesToRadiansNode->GetUniqueID());
+		SetAttributeNewName(Input, "A");
+	}
+
+	Rotate2DNode->AddFloatAttribute(Rotator::Attributes::CenterX.ToString(), 0.f);
+	Rotate2DNode->AddFloatAttribute(Rotator::Attributes::CenterY.ToString(), 0.f);
+	Rotate2DNode->AddFloatAttribute(Rotator::Attributes::Speed.ToString(), 1.f);
+
+	UInterchangeShaderPortsAPI::ConnectDefaultOuputToInput(Connect.ParentShaderNode, Connect.InputChannelName, Rotate2DNode->GetUniqueID());
+}
+
 void FMaterialXSurfaceShaderAbstract::ConnectRotate3DInputToOutput(const FConnectNode& Connect)
 {
 	UInterchangeShaderNode* Rotate3DNode = CreateShaderNode(Connect.UpstreamNode->getName().c_str(), TEXT("RotateAboutAxis"));
@@ -1105,7 +1138,22 @@ void FMaterialXSurfaceShaderAbstract::ConnectTexCoordInputToOutput(const FConnec
 	{
 		TexCoord->AddInt32Attribute(UInterchangeShaderPortsAPI::MakeInputValueKey(TextureCoordinate::Inputs::Index.ToString()), mx::fromValueString<int>(Input->getValueString()));
 	}
-	UInterchangeShaderPortsAPI::ConnectDefaultOuputToInput(Connect.ParentShaderNode, Connect.InputChannelName, TexCoord->GetUniqueID());
+
+	// Flip UV coordinates to be bottom left instead of top left
+	UInterchangeShaderNode* TexCoordR = CreateMaskShaderNode(0b1000, (Connect.UpstreamNode->getName() + "ComponentMaskR").c_str());
+	UInterchangeShaderPortsAPI::ConnectDefaultOuputToInput(TexCoordR, Mask::Inputs::Input.ToString(), TexCoord->GetUniqueID());
+
+	UInterchangeShaderNode* TexCoordG = CreateMaskShaderNode(0b0100, (Connect.UpstreamNode->getName() + "ComponentMaskG").c_str());
+	UInterchangeShaderPortsAPI::ConnectDefaultOuputToInput(TexCoordG, Mask::Inputs::Input.ToString(), TexCoord->GetUniqueID());
+
+	UInterchangeShaderNode* OneMinusG = CreateShaderNode((Connect.UpstreamNode->getName() + "OneMinusG").c_str(), OneMinus::Name.ToString());
+	UInterchangeShaderPortsAPI::ConnectDefaultOuputToInput(OneMinusG, OneMinus::Inputs::Input.ToString(), TexCoordG->GetUniqueID());
+
+	UInterchangeShaderNode* TexCoordAppend = CreateShaderNode((Connect.UpstreamNode->getName() + "Append").c_str(), TEXT("AppendVector"));
+	UInterchangeShaderPortsAPI::ConnectDefaultOuputToInput(TexCoordAppend, TEXT("A"), TexCoordR->GetUniqueID());
+	UInterchangeShaderPortsAPI::ConnectDefaultOuputToInput(TexCoordAppend, TEXT("B"), OneMinusG->GetUniqueID());
+
+	UInterchangeShaderPortsAPI::ConnectDefaultOuputToInput(Connect.ParentShaderNode, Connect.InputChannelName, TexCoordAppend->GetUniqueID());
 }
 
 void FMaterialXSurfaceShaderAbstract::ConnectSeparateInputToOutput(const FConnectNode& Connect)
@@ -1360,6 +1408,7 @@ void FMaterialXSurfaceShaderAbstract::RegisterConnectNodeOutputToInputDelegates(
 	MatchingConnectNodeDelegates.Add(mx::Category::TransformPoint,	FOnConnectNodeOutputToInput::CreateSP(this, &FMaterialXSurfaceShaderAbstract::ConnectTransformPositionInputToOutput));
 	MatchingConnectNodeDelegates.Add(mx::Category::TransformVector, FOnConnectNodeOutputToInput::CreateSP(this, &FMaterialXSurfaceShaderAbstract::ConnectTransformVectorInputToOutput));
 	MatchingConnectNodeDelegates.Add(mx::Category::TransformNormal, FOnConnectNodeOutputToInput::CreateSP(this, &FMaterialXSurfaceShaderAbstract::ConnectTransformVectorInputToOutput));
+	MatchingConnectNodeDelegates.Add(mx::Category::Rotate2D,		FOnConnectNodeOutputToInput::CreateSP(this, &FMaterialXSurfaceShaderAbstract::ConnectRotate2DInputToOutput));
 	MatchingConnectNodeDelegates.Add(mx::Category::Rotate3D,		FOnConnectNodeOutputToInput::CreateSP(this, &FMaterialXSurfaceShaderAbstract::ConnectRotate3DInputToOutput));
 	MatchingConnectNodeDelegates.Add(mx::Category::Image,			FOnConnectNodeOutputToInput::CreateSP(this, &FMaterialXSurfaceShaderAbstract::ConnectImageInputToOutput));
 	MatchingConnectNodeDelegates.Add(mx::Category::Convert,			FOnConnectNodeOutputToInput::CreateSP(this, &FMaterialXSurfaceShaderAbstract::ConnectConvertInputToOutput));
diff --git a/Engine/Plugins/Interchange/Runtime/Source/Import/Public/MaterialX/MaterialXUtils/MaterialXSurfaceShaderAbstract.h b/Engine/Plugins/Interchange/Runtime/Source/Import/Public/MaterialX/MaterialXUtils/MaterialXSurfaceShaderAbstract.h
--- a/Engine/Plugins/Interchange/Runtime/Source/Import/Public/MaterialX/MaterialXUtils/MaterialXSurfaceShaderAbstract.h
+++ b/Engine/Plugins/Interchange/Runtime/Source/Import/Public/MaterialX/MaterialXUtils/MaterialXSurfaceShaderAbstract.h
@@ -215,6 +215,9 @@ protected:
 	/** <transformvector> */
 	void ConnectTransformVectorInputToOutput(const FConnectNode& Connect);
 
+	/** <rotate2d> */
+	void ConnectRotate2DInputToOutput(const FConnectNode& Connect);
+
 	/** <rotate3d> */
 	void ConnectRotate3DInputToOutput(const FConnectNode& Connect);
 
diff --git a/Engine/Plugins/Interchange/Runtime/Source/Nodes/Public/InterchangeMaterialDefinitions.h b/Engine/Plugins/Interchange/Runtime/Source/Nodes/Public/InterchangeMaterialDefinitions.h
--- a/Engine/Plugins/Interchange/Runtime/Source/Nodes/Public/InterchangeMaterialDefinitions.h
+++ b/Engine/Plugins/Interchange/Runtime/Source/Nodes/Public/InterchangeMaterialDefinitions.h
@@ -137,6 +137,24 @@ namespace Interchange
 					}
 				}
 
+				namespace Rotator
+				{
+					const FName Name = TEXT("Rotator");
+
+					namespace Attributes
+					{
+						const FName CenterX = TEXT("CenterX");
+						const FName CenterY = TEXT("CenterY");
+						const FName Speed = TEXT("Speed");
+					}
+
+					namespace Inputs
+					{
+						const FName Coordinate = TEXT("Coordinate");
+						const FName Time = TEXT("Time");
+					}
+				}
+
 				namespace Swizzle
 				{
 					const FName Name = TEXT("MaterialXSwizzle");
diff --git a/Engine/Plugins/Interchange/Runtime/Source/Pipelines/Private/InterchangeGenericMaterialPipeline.cpp b/Engine/Plugins/Interchange/Runtime/Source/Pipelines/Private/InterchangeGenericMaterialPipeline.cpp
--- a/Engine/Plugins/Interchange/Runtime/Source/Pipelines/Private/InterchangeGenericMaterialPipeline.cpp
+++ b/Engine/Plugins/Interchange/Runtime/Source/Pipelines/Private/InterchangeGenericMaterialPipeline.cpp
@@ -31,6 +31,7 @@
 #include "Materials/MaterialExpressionMultiply.h"
 #include "Materials/MaterialExpressionNoise.h"
 #include "Materials/MaterialExpressionOneMinus.h"
+#include "Materials/MaterialExpressionRotator.h"
 #include "Materials/MaterialExpressionScalarParameter.h"
 #include "Materials/MaterialExpressionStaticBoolParameter.h"
 #include "Materials/MaterialExpressionConstant.h"
@@ -2309,6 +2310,58 @@ void UInterchangeGenericMaterialPipeline::HandleMaskNode(const UInterchangeShade
 	}
 }
 
+void UInterchangeGenericMaterialPipeline::HandleRotatorNode(const UInterchangeShaderNode* ShaderNode, UInterchangeBaseMaterialFactoryNode* MaterialFactoryNode, UInterchangeMaterialExpressionFactoryNode* RotatorFactoryNode)
+{
+	using namespace UE::Interchange::Materials::Standard::Nodes;
+	RotatorFactoryNode->SetCustomExpressionClassName(UMaterialExpressionRotator::StaticClass()->GetName());
+
+	// Coordinate
+	{
+		TTuple<UInterchangeMaterialExpressionFactoryNode*, FString> CoordinateExpression =
+			CreateMaterialExpressionForInput(MaterialFactoryNode, ShaderNode, Rotator::Inputs::Coordinate.ToString(), RotatorFactoryNode->GetUniqueID());
+		if(CoordinateExpression.Get<0>())
+		{
+			UInterchangeShaderPortsAPI::ConnectOuputToInputByName(RotatorFactoryNode, GET_MEMBER_NAME_CHECKED(UMaterialExpressionRotator, Coordinate).ToString(),
+																  CoordinateExpression.Get<0>()->GetUniqueID(), CoordinateExpression.Get<1>());
+		}
+	}
+
+	// Time
+	{
+		TTuple<UInterchangeMaterialExpressionFactoryNode*, FString> TimeExpression =
+			CreateMaterialExpressionForInput(MaterialFactoryNode, ShaderNode, Rotator::Inputs::Time.ToString(), RotatorFactoryNode->GetUniqueID());
+		if(TimeExpression.Get<0>())
+		{
+			UInterchangeShaderPortsAPI::ConnectOuputToInputByName(RotatorFactoryNode, GET_MEMBER_NAME_CHECKED(UMaterialExpressionRotator, Time).ToString(),
+																  TimeExpression.Get<0>()->GetUniqueID(), TimeExpression.Get<1>());
+		}
+	}
+
+	// CenterX
+	if(float CenterX; ShaderNode->GetFloatAttribute(Rotator::Attributes::CenterX.ToString(), CenterX))
+	{
+		const FName CenterXMemberName = GET_MEMBER_NAME_CHECKED(UMaterialExpressionRotator, CenterX);
+		RotatorFactoryNode->AddFloatAttribute(CenterXMemberName.ToString(), CenterX);
+		RotatorFactoryNode->AddApplyAndFillDelegates<float>(CenterXMemberName.ToString(), UMaterialExpressionRotator::StaticClass(), CenterXMemberName);
+	}
+
+	// CenterY
+	if(float CenterY; ShaderNode->GetFloatAttribute(Rotator::Attributes::CenterY.ToString(), CenterY))
+	{
+		const FName CenterYMemberName = GET_MEMBER_NAME_CHECKED(UMaterialExpressionRotator, CenterY);
+		RotatorFactoryNode->AddFloatAttribute(CenterYMemberName.ToString(), CenterY);
+		RotatorFactoryNode->AddApplyAndFillDelegates<float>(CenterYMemberName.ToString(), UMaterialExpressionRotator::StaticClass(), CenterYMemberName);
+	}
+
+	// Speed
+	if(float Speed; ShaderNode->GetFloatAttribute(Rotator::Attributes::Speed.ToString(), Speed))
+	{
+		const FName SpeedMemberName = GET_MEMBER_NAME_CHECKED(UMaterialExpressionRotator, Speed);
+		RotatorFactoryNode->AddFloatAttribute(SpeedMemberName.ToString(), Speed);
+		RotatorFactoryNode->AddApplyAndFillDelegates<float>(SpeedMemberName.ToString(), UMaterialExpressionRotator::StaticClass(), SpeedMemberName);
+	}
+}
+
 void UInterchangeGenericMaterialPipeline::HandleTimeNode(const UInterchangeShaderNode* ShaderNode, UInterchangeBaseMaterialFactoryNode* MaterialFactoryNode, UInterchangeMaterialExpressionFactoryNode* TimeFactoryNode)
 {
 	using namespace UE::Interchange::Materials::Standard::Nodes;
@@ -2641,7 +2694,6 @@ void UInterchangeGenericMaterialPipeline::HandleStaticBooleanParameterNode(const
 	StaticBoolParameterFactoryNode->SetDisplayLabel(ShaderNode->GetDisplayLabel());
 }
 
-
 UInterchangeMaterialExpressionFactoryNode* UInterchangeGenericMaterialPipeline::CreateMaterialExpressionForShaderNode(UInterchangeBaseMaterialFactoryNode* MaterialFactoryNode,
 	const UInterchangeShaderNode* ShaderNode, const FString& ParentUid)
 {
@@ -2704,6 +2756,10 @@ UInterchangeMaterialExpressionFactoryNode* UInterchangeGenericMaterialPipeline::
 	{
 		HandleNormalFromHeightMapNode(ShaderNode, MaterialFactoryNode, MaterialExpression);
 	}
+	else if(*ShaderType == Nodes::Rotator::Name)
+	{
+		HandleRotatorNode(ShaderNode, MaterialFactoryNode, MaterialExpression);
+	}
 	else if(*ShaderType == Nodes::Swizzle::Name)
 	{
 		HandleSwizzleNode(ShaderNode, MaterialFactoryNode, MaterialExpression);
diff --git a/Engine/Plugins/Interchange/Runtime/Source/Pipelines/Public/InterchangeGenericMaterialPipeline.h b/Engine/Plugins/Interchange/Runtime/Source/Pipelines/Public/InterchangeGenericMaterialPipeline.h
--- a/Engine/Plugins/Interchange/Runtime/Source/Pipelines/Public/InterchangeGenericMaterialPipeline.h
+++ b/Engine/Plugins/Interchange/Runtime/Source/Pipelines/Public/InterchangeGenericMaterialPipeline.h
@@ -209,6 +209,7 @@ private:
 	void HandleTextureCoordinateNode(const UInterchangeShaderNode* ShaderNode, UInterchangeBaseMaterialFactoryNode* MaterialFactoryNode, UInterchangeMaterialExpressionFactoryNode*& TextureSampleFactoryNode);
 	void HandleLerpNode(const UInterchangeShaderNode* ShaderNode, UInterchangeBaseMaterialFactoryNode* MaterialFactoryNode, UInterchangeMaterialExpressionFactoryNode* LerpFactoryNode);
 	void HandleMaskNode(const UInterchangeShaderNode* ShaderNode, UInterchangeBaseMaterialFactoryNode* MaterialFactoryNode, UInterchangeMaterialExpressionFactoryNode* MaskFactoryNode);
+	void HandleRotatorNode(const UInterchangeShaderNode* ShaderNode, UInterchangeBaseMaterialFactoryNode* MaterialFactoryNode, UInterchangeMaterialExpressionFactoryNode* RotatorFactoryNode);
 	void HandleTimeNode(const UInterchangeShaderNode* ShaderNode, UInterchangeBaseMaterialFactoryNode* MaterialFactoryNode, UInterchangeMaterialExpressionFactoryNode* TimeFactoryNode);
 	void HandleTransformPositionNode(const UInterchangeShaderNode* ShaderNode, UInterchangeBaseMaterialFactoryNode* MaterialFactoryNode, UInterchangeMaterialExpressionFactoryNode* TransformPositionFactoryNode);
 	void HandleTransformVectorNode(const UInterchangeShaderNode* ShaderNode, UInterchangeBaseMaterialFactoryNode* MaterialFactoryNode, UInterchangeMaterialExpressionFactoryNode* TransformVectorFactoryNode);
diff --git a/Engine/Source/Runtime/Engine/Public/Materials/MaterialExpressionRotator.h b/Engine/Source/Runtime/Engine/Public/Materials/MaterialExpressionRotator.h
--- a/Engine/Source/Runtime/Engine/Public/Materials/MaterialExpressionRotator.h
+++ b/Engine/Source/Runtime/Engine/Public/Materials/MaterialExpressionRotator.h
@@ -9,7 +9,7 @@
 #include "Materials/MaterialExpression.h"
 #include "MaterialExpressionRotator.generated.h"
 
-UCLASS(collapsecategories, hidecategories=Object)
+UCLASS(collapsecategories, hidecategories=Object, MinimalAPI)
 class UMaterialExpressionRotator : public UMaterialExpression
 {
 	GENERATED_UCLASS_BODY()
