MaterialX:
- Add support of thin-film-IOR
- Improve thin film iridescence

#virtualized

diff --git a/Engine/Plugins/Interchange/Runtime/Source/Import/Private/MaterialX/MaterialXUtils/MaterialXStandardSurfaceShader.cpp b/Engine/Plugins/Interchange/Runtime/Source/Import/Private/MaterialX/MaterialXUtils/MaterialXStandardSurfaceShader.cpp
--- a/Engine/Plugins/Interchange/Runtime/Source/Import/Private/MaterialX/MaterialXUtils/MaterialXStandardSurfaceShader.cpp
+++ b/Engine/Plugins/Interchange/Runtime/Source/Import/Private/MaterialX/MaterialXUtils/MaterialXStandardSurfaceShader.cpp
@@ -114,6 +114,9 @@ void FMaterialXStandardSurfaceShader::ConnectToStandardSurface(UInterchangeFunct
 	//Thin Film Thickness
 	ConnectNodeOutputToInput(Input::ThinFilmThickness, StandardSurfaceShaderNode, StandardSurface::Parameters::ThinFilmThickness.ToString(), DefaultValue::Float::ThinFilmThickness);
 
+	//Thin Film IOR
+	ConnectNodeOutputToInput(Input::ThinFilmIOR, StandardSurfaceShaderNode, StandardSurface::Parameters::ThinFilmIOR.ToString(), DefaultValue::Float::ThinFilmIOR);
+
 	//Emission
 	ConnectNodeOutputToInput(Input::Emission, StandardSurfaceShaderNode, StandardSurface::Parameters::Emission.ToString(), DefaultValue::Float::Emission);
 
