MaterialX:
- Match MF_FlattenNormal to <normalmap> node

diff --git a/Engine/Plugins/Interchange/Runtime/Source/Import/Private/MaterialX/MaterialXUtils/MaterialXManager.cpp b/Engine/Plugins/Interchange/Runtime/Source/Import/Private/MaterialX/MaterialXUtils/MaterialXManager.cpp
--- a/Engine/Plugins/Interchange/Runtime/Source/Import/Private/MaterialX/MaterialXUtils/MaterialXManager.cpp
+++ b/Engine/Plugins/Interchange/Runtime/Source/Import/Private/MaterialX/MaterialXUtils/MaterialXManager.cpp
@@ -83,6 +83,8 @@ FMaterialXManager::FMaterialXManager()
 		{{MaterialX::Category::Noise3D,        TEXT("amplitude")},   TEXT("B")},            // The amplitude of the noise is connected to a multiply node
 		{{MaterialX::Category::Noise3D,        TEXT("pivot")},       TEXT("B")},            // The pivot of the noise is connected to a add node
 		{{MaterialX::Category::Normalize,      TEXT("in")},          TEXT("VectorInput")},
+		{{MaterialX::Category::NormalMap,      TEXT("in")},          TEXT("Normal")},
+		{{MaterialX::Category::NormalMap,      TEXT("scale")},       TEXT("Flatness")},
 		{{MaterialX::Category::Outside,        TEXT("in")},          TEXT("A")},				// Outside is treated as Multiply node
 		{{MaterialX::Category::Outside,        TEXT("mask")},        TEXT("B")},				// Outside is treated as Multiply node
 		{{MaterialX::Category::Power,          TEXT("in1")},         TEXT("Base")},
diff --git a/Engine/Plugins/Interchange/Runtime/Source/Import/Private/MaterialX/MaterialXUtils/MaterialXSurfaceShaderAbstract.cpp b/Engine/Plugins/Interchange/Runtime/Source/Import/Private/MaterialX/MaterialXUtils/MaterialXSurfaceShaderAbstract.cpp
--- a/Engine/Plugins/Interchange/Runtime/Source/Import/Private/MaterialX/MaterialXUtils/MaterialXSurfaceShaderAbstract.cpp
+++ b/Engine/Plugins/Interchange/Runtime/Source/Import/Private/MaterialX/MaterialXUtils/MaterialXSurfaceShaderAbstract.cpp
@@ -1217,6 +1217,28 @@ void FMaterialXSurfaceShaderAbstract::ConnectSwizzleInputToOutput(const FConnect
 	UInterchangeShaderPortsAPI::ConnectDefaultOuputToInput(Connect.ParentShaderNode, Connect.InputChannelName, SwizzleNode->GetUniqueID());
 }
 
+void FMaterialXSurfaceShaderAbstract::ConnectNormalMapInputToOutput(const FConnectNode& Connect)
+{
+	using namespace UE::Interchange::Materials::Standard::Nodes;
+
+	// Only create a FunctionCall if there's a scale, otherwise just like dot
+	if(mx::InputPtr Input = Connect.UpstreamNode->getInput("scale"))
+	{
+		UInterchangeShaderNode* FlattenNormalNode = CreateFunctionCallShaderNode(Connect.UpstreamNode->getName().c_str(), TEXT("/Engine/Functions/Engine_MaterialFunctions01/Texturing/FlattenNormal.FlattenNormal"));
+
+		UInterchangeShaderNode* OneMinusNode = CreateShaderNode((Connect.UpstreamNode->getName() + "_OneMinusFlatness").c_str(), OneMinus::Name.ToString());
+		UInterchangeShaderPortsAPI::ConnectDefaultOuputToInput(FlattenNormalNode, FlattenNormal::Inputs::Flatness.ToString(), OneMinusNode->GetUniqueID());
+
+		AddAttributeFromValueOrInterface(Input, OneMinus::Inputs::Input.ToString(), OneMinusNode);
+		UInterchangeShaderPortsAPI::ConnectDefaultOuputToInput(Connect.ParentShaderNode, Connect.InputChannelName, FlattenNormalNode->GetUniqueID());
+	}
+	else
+	{
+		SetAttributeNewName(Connect.UpstreamNode->getInput("in"), TCHAR_TO_UTF8(*Connect.InputChannelName)); //let's take the parent node's input name
+		ShaderNodes.Add({ Connect.UpstreamNode->getName().c_str(), Connect.OutputName }, Connect.ParentShaderNode);
+	}
+}
+
 UInterchangeShaderNode* FMaterialXSurfaceShaderAbstract::CreateMaskShaderNode(uint8 RGBA, const FString& NodeName, const FString& OutputName)
 {
 	bool bR = (0b1000 & RGBA) >> 3;
@@ -1404,7 +1426,7 @@ void FMaterialXSurfaceShaderAbstract::RegisterConnectNodeOutputToInputDelegates(
 	MatchingConnectNodeDelegates.Add(mx::Category::Constant,		FOnConnectNodeOutputToInput::CreateSP(this, &FMaterialXSurfaceShaderAbstract::ConnectConstantInputToOutput));
 	MatchingConnectNodeDelegates.Add(mx::Category::Extract,			FOnConnectNodeOutputToInput::CreateSP(this, &FMaterialXSurfaceShaderAbstract::ConnectExtractInputToOutput));
 	MatchingConnectNodeDelegates.Add(mx::Category::Dot,				FOnConnectNodeOutputToInput::CreateSP(this, &FMaterialXSurfaceShaderAbstract::ConnectDotInputToOutput));
-	MatchingConnectNodeDelegates.Add(mx::Category::NormalMap,		FOnConnectNodeOutputToInput::CreateSP(this, &FMaterialXSurfaceShaderAbstract::ConnectDotInputToOutput));
+	MatchingConnectNodeDelegates.Add(mx::Category::NormalMap,		FOnConnectNodeOutputToInput::CreateSP(this, &FMaterialXSurfaceShaderAbstract::ConnectNormalMapInputToOutput));
 	MatchingConnectNodeDelegates.Add(mx::Category::TransformPoint,	FOnConnectNodeOutputToInput::CreateSP(this, &FMaterialXSurfaceShaderAbstract::ConnectTransformPositionInputToOutput));
 	MatchingConnectNodeDelegates.Add(mx::Category::TransformVector, FOnConnectNodeOutputToInput::CreateSP(this, &FMaterialXSurfaceShaderAbstract::ConnectTransformVectorInputToOutput));
 	MatchingConnectNodeDelegates.Add(mx::Category::TransformNormal, FOnConnectNodeOutputToInput::CreateSP(this, &FMaterialXSurfaceShaderAbstract::ConnectTransformVectorInputToOutput));
diff --git a/Engine/Plugins/Interchange/Runtime/Source/Import/Public/MaterialX/MaterialXUtils/MaterialXSurfaceShaderAbstract.h b/Engine/Plugins/Interchange/Runtime/Source/Import/Public/MaterialX/MaterialXUtils/MaterialXSurfaceShaderAbstract.h
--- a/Engine/Plugins/Interchange/Runtime/Source/Import/Public/MaterialX/MaterialXUtils/MaterialXSurfaceShaderAbstract.h
+++ b/Engine/Plugins/Interchange/Runtime/Source/Import/Public/MaterialX/MaterialXUtils/MaterialXSurfaceShaderAbstract.h
@@ -278,6 +278,9 @@ protected:
 	/** <swizzle> */
 	void ConnectSwizzleInputToOutput(const FConnectNode& Connect);
 
+	/** <normalmap> */
+	void ConnectNormalMapInputToOutput(const FConnectNode& Connect);
+
 	/** End Connect MaterialX nodes*/
 
 	/**
