[Interchange glTF] Add support for import of KHR_materials_anisotropy

#rnx

#virtualized

diff --git a/Engine/Plugins/Interchange/Runtime/Source/Import/Private/Gltf/InterchangeGLTFMaterial.cpp b/Engine/Plugins/Interchange/Runtime/Source/Import/Private/Gltf/InterchangeGLTFMaterial.cpp
--- a/Engine/Plugins/Interchange/Runtime/Source/Import/Private/Gltf/InterchangeGLTFMaterial.cpp
+++ b/Engine/Plugins/Interchange/Runtime/Source/Import/Private/Gltf/InterchangeGLTFMaterial.cpp
@@ -632,6 +632,31 @@ namespace UE::Interchange::GLTFMaterials
 					//Configuration
 					SetBool(Inputs::Configuration::bHasIridescenceTexture, false);
 				}
+
+				if (GltfMaterial.Anisotropy.bHasAnisotropy)
+				{
+					//AnisotropyTexture
+					//AnisotropyTexture_OffsetScale
+					//AnisotropyTexture_Rotation
+					//AnisotropyTexture_TexCoord
+					//AnisotropyTexture_TilingMethod
+					bool bHasAnisotropyTextureAndOrRotation = SetMap(Inputs::AnisotropyTexture, GltfMaterial.Anisotropy.Texture);
+
+					//AnisotropyStrength
+					SetScalar(Inputs::AnisotropyStrength, GltfMaterial.Anisotropy.Strength, 0);
+
+					//AnisotropyRotation
+					SetScalar(Inputs::AnisotropyRotation, GltfMaterial.Anisotropy.Rotation, 0);
+					bHasAnisotropyTextureAndOrRotation |= GltfMaterial.Anisotropy.Rotation != 0;
+
+					//Configuration
+					SetBool(Inputs::Configuration::bHasAnisotropyTextureAndOrRotation, bHasAnisotropyTextureAndOrRotation);
+				}
+				else
+				{
+					//Configuration
+					SetBool(Inputs::Configuration::bHasAnisotropyTextureAndOrRotation, false);
+				}
 			}
 			else
 			{
diff --git a/Engine/Plugins/Interchange/Runtime/Source/Import/Public/Gltf/InterchangeGLTFMaterial.h b/Engine/Plugins/Interchange/Runtime/Source/Import/Public/Gltf/InterchangeGLTFMaterial.h
--- a/Engine/Plugins/Interchange/Runtime/Source/Import/Public/Gltf/InterchangeGLTFMaterial.h
+++ b/Engine/Plugins/Interchange/Runtime/Source/Import/Public/Gltf/InterchangeGLTFMaterial.h
@@ -125,6 +125,12 @@ namespace UE::Interchange::GLTFMaterials
 		const FString IridescenceThicknessMinimum = TEXT("IridescenceThicknessMinimum");
 		const FString IridescenceThicknessMaximum = TEXT("IridescenceThicknessMaximum");
 
+
+		//Anisotropy Specific:
+		DECLARE_INTERCHANGE_GLTF_MI_MAP(Anisotropy);
+		const FString AnisotropyStrength = TEXT("AnisotropyStrength");
+		const FString AnisotropyRotation = TEXT("AnisotropyRotation");
+
 		namespace Configuration
 		{
 			//For StaticSwitch optimizations:
@@ -140,6 +146,7 @@ namespace UE::Interchange::GLTFMaterials
 			const FString bHasSheenTexture = TEXT("bHasSheenTexture");
 			const FString bHasTransmissionTexture = TEXT("bHasTransmissionTexture");
 			const FString bHasIridescenceTexture = TEXT("bHasIridescenceTexture");
+			const FString bHasAnisotropyTextureAndOrRotation = TEXT("bHasAnisotropyTextureAndOrRotation"); //AnistropyRotation can affect the Tangents so we check
 		}
 	}
 
diff --git a/Engine/Plugins/Interchange/Runtime/Source/Parsers/GLTFCore/Private/GLTF/ExtensionsHandler.cpp b/Engine/Plugins/Interchange/Runtime/Source/Parsers/GLTFCore/Private/GLTF/ExtensionsHandler.cpp
--- a/Engine/Plugins/Interchange/Runtime/Source/Parsers/GLTFCore/Private/GLTF/ExtensionsHandler.cpp
+++ b/Engine/Plugins/Interchange/Runtime/Source/Parsers/GLTFCore/Private/GLTF/ExtensionsHandler.cpp
@@ -219,6 +219,7 @@ namespace GLTF
 			EExtension::KHR_MaterialsSpecular,
 			EExtension::KHR_MaterialsEmissiveStrength,
 			EExtension::KHR_MaterialsIridescence,
+			EExtension::KHR_MaterialsAnisotropy,
 			EExtension::MSFT_PackingOcclusionRoughnessMetallic,
 			EExtension::MSFT_PackingNormalRoughnessMetallic 
 		};
@@ -385,6 +386,17 @@ namespace GLTF
 					GLTF::SetTextureMap(Iridescence, TEXT("iridescenceThicknessTexture"), nullptr, Asset->Textures, Material.Iridescence.Thickness.Texture, Messages);
 				}
 				break;
+				case EExtension::KHR_MaterialsAnisotropy:
+				{
+					const FJsonObject& Anisotropy = ExtObj;
+
+					Material.Anisotropy.bHasAnisotropy = true;
+
+					Material.Anisotropy.Strength = GetScalar(Anisotropy, TEXT("anisotropyStrength"), 0.0f);
+					Material.Anisotropy.Rotation = GetScalar(Anisotropy, TEXT("anisotropyRotation"), 0.0f);
+					GLTF::SetTextureMap(Anisotropy, TEXT("anisotropyTexture"), nullptr, Asset->Textures, Material.Anisotropy.Texture, Messages);
+				}
+				break;
 				default:
 					if (!ensure(false))
 					{
diff --git a/Engine/Plugins/Interchange/Runtime/Source/Parsers/GLTFCore/Private/GLTF/GLTFAsset.cpp b/Engine/Plugins/Interchange/Runtime/Source/Parsers/GLTFCore/Private/GLTF/GLTFAsset.cpp
--- a/Engine/Plugins/Interchange/Runtime/Source/Parsers/GLTFCore/Private/GLTF/GLTFAsset.cpp
+++ b/Engine/Plugins/Interchange/Runtime/Source/Parsers/GLTFCore/Private/GLTF/GLTFAsset.cpp
@@ -395,6 +395,8 @@ namespace GLTF
 				return TEXT("KHR_materials_emissive_strength");
 			case GLTF::EExtension::KHR_MaterialsIridescence:
 				return TEXT("KHR_materials_iridescence");
+			case GLTF::EExtension::KHR_MaterialsAnisotropy:
+				return TEXT("KHR_materials_anisotropy");
 			case GLTF::EExtension::KHR_TextureTransform:
 				return TEXT("KHR_texture_transform");
 			case GLTF::EExtension::KHR_DracoMeshCompression:
diff --git a/Engine/Plugins/Interchange/Runtime/Source/Parsers/GLTFCore/Public/GLTFAsset.h b/Engine/Plugins/Interchange/Runtime/Source/Parsers/GLTFCore/Public/GLTFAsset.h
--- a/Engine/Plugins/Interchange/Runtime/Source/Parsers/GLTFCore/Public/GLTFAsset.h
+++ b/Engine/Plugins/Interchange/Runtime/Source/Parsers/GLTFCore/Public/GLTFAsset.h
@@ -38,6 +38,7 @@ namespace GLTF
 		KHR_MaterialsSpecular,
 		KHR_MaterialsEmissiveStrength,
 		KHR_MaterialsIridescence,
+		KHR_MaterialsAnisotropy,
 		KHR_TextureTransform,
 		KHR_DracoMeshCompression,
 		KHR_LightsPunctual,
diff --git a/Engine/Plugins/Interchange/Runtime/Source/Parsers/GLTFCore/Public/GLTFMaterial.h b/Engine/Plugins/Interchange/Runtime/Source/Parsers/GLTFCore/Public/GLTFMaterial.h
--- a/Engine/Plugins/Interchange/Runtime/Source/Parsers/GLTFCore/Public/GLTFMaterial.h
+++ b/Engine/Plugins/Interchange/Runtime/Source/Parsers/GLTFCore/Public/GLTFMaterial.h
@@ -146,6 +146,16 @@ namespace GLTF
 			bool bHasIridescence = false;
 		};
 
+		struct FAnisotropy
+		{
+			FTextureMap Texture;
+
+			float Strength = 0.0;
+			float Rotation = 0.0;
+
+			bool bHasAnisotropy = false;
+		};
+
 		FString Name;
 
 		// PBR properties
@@ -184,6 +194,7 @@ namespace GLTF
 		bool            bHasSpecular;
 		bool            bHasEmissiveStrength;
 		FIridescence    Iridescence;
+		FAnisotropy		Anisotropy;
 
 		TMap<FString, FString> Extras;
 
