[Interchange] Ensure fix in InterchangeAnimSequenceFactory.
(Converting from Curve to Baked wtih negative RangeEnd)

#rnx

diff --git a/Engine/Plugins/Interchange/Runtime/Source/Parsers/CommonParser/Private/InterchangeCommonAnimationPayload.cpp b/Engine/Plugins/Interchange/Runtime/Source/Parsers/CommonParser/Private/InterchangeCommonAnimationPayload.cpp
--- a/Engine/Plugins/Interchange/Runtime/Source/Parsers/CommonParser/Private/InterchangeCommonAnimationPayload.cpp
+++ b/Engine/Plugins/Interchange/Runtime/Source/Parsers/CommonParser/Private/InterchangeCommonAnimationPayload.cpp
@@ -3,6 +3,7 @@
 #include "InterchangeCommonAnimationPayload.h"
 
 #include "CoreMinimal.h"
+#include "Animation/AnimTypes.h"
 
 #include UE_INLINE_GENERATED_CPP_BY_NAME(InterchangeCommonAnimationPayload)
 
@@ -64,13 +65,12 @@ namespace UE::Interchange
 			}
 			if (RangeEndTime < 0)
 			{
-				return;
+				RangeEndTime = 0;
 			}
 
 			const double BakeInterval = 1.0 / BakeFrequency;
-			const double SequenceLength = FMath::Max<double>(RangeEndTime - RangeStartTime, BakeInterval);
+			const double SequenceLength = FMath::Max<double>(RangeEndTime - RangeStartTime, MINIMUM_ANIMATION_LENGTH);
 			int32 BakeKeyCount = FMath::RoundToInt32(SequenceLength * BakeFrequency) + 1;
-			const FFrameRate ResampleFrameRate(BakeFrequency, 1);
 
 			auto EvaluateCurve = [this](const int32& CurveIndex, double CurrentTime, float DefaultValue)
 			{
diff --git a/Engine/Plugins/Interchange/Runtime/Source/Parsers/Fbx/Private/FbxAnimation.cpp b/Engine/Plugins/Interchange/Runtime/Source/Parsers/Fbx/Private/FbxAnimation.cpp
--- a/Engine/Plugins/Interchange/Runtime/Source/Parsers/Fbx/Private/FbxAnimation.cpp
+++ b/Engine/Plugins/Interchange/Runtime/Source/Parsers/Fbx/Private/FbxAnimation.cpp
@@ -21,6 +21,7 @@
 #include "StaticMeshAttributes.h"
 #include "InterchangeAnimationTrackSetNode.h"
 #include "Nodes/InterchangeBaseNodeContainer.h"
+#include "Animation/AnimTypes.h"
 
 #define LOCTEXT_NAMESPACE "InterchangeFbxMesh"
 
@@ -381,8 +382,11 @@ namespace UE::Interchange::Private
 		FbxTime TimeStep = 0;
 		TimeStep.SetSecondDouble(TimeStepSecond);
 
-		const int32 NumFrame = FMath::RoundToInt32((AnimationBakeTransformPayloadData.RangeEndTime - AnimationBakeTransformPayloadData.RangeStartTime) * AnimationBakeTransformPayloadData.BakeFrequency);
-		check(NumFrame >= 0);
+		const double SequenceLength = FMath::Max<double>(AnimationBakeTransformPayloadData.RangeEndTime - AnimationBakeTransformPayloadData.RangeStartTime, MINIMUM_ANIMATION_LENGTH);
+		const int32 NumFrame = FMath::RoundToInt32(SequenceLength * AnimationBakeTransformPayloadData.BakeFrequency);
+		int32 BakeKeyCount = NumFrame + 1;
+
+		ensure(NumFrame >= 0);
 
 		//Add a threshold when we compare if we have reach the end of the animation
 		const FbxTime TimeComparisonThreshold = (UE_DOUBLE_KINDA_SMALL_NUMBER * static_cast<double>(FBXSDK_TC_SECOND));
@@ -403,13 +407,14 @@ namespace UE::Interchange::Private
 				bNanErrorLogged = true;
 			};
 
-		for (FbxTime CurTime = StartTime; CurTime < (EndTime + TimeComparisonThreshold); CurTime += TimeStep)
+		FbxTime CurrentTime = StartTime;
+		for (size_t FrameIndex = 0; FrameIndex < BakeKeyCount; FrameIndex++, CurrentTime+=TimeStep)
 		{
 			FTransform LocalTransform;
 			FbxNode* ParentNode = FetchPayloadData.Node->GetParent();
 			if (ParentNode)
 			{
-				FbxAMatrix NodeTransform = FetchPayloadData.Node->EvaluateGlobalTransform(CurTime);
+				FbxAMatrix NodeTransform = FetchPayloadData.Node->EvaluateGlobalTransform(CurrentTime);
 				FTransform GlobalTransform = UE::Interchange::Private::FFbxConvert::ConvertTransform<FTransform, FVector, FQuat>(NodeTransform);
 				if (GlobalTransform.ContainsNaN())
 				{
@@ -417,7 +422,7 @@ namespace UE::Interchange::Private
 					GlobalTransform.SetIdentity();
 				}
 
-				FbxAMatrix ParentTransform = ParentNode->EvaluateGlobalTransform(CurTime);
+				FbxAMatrix ParentTransform = ParentNode->EvaluateGlobalTransform(CurrentTime);
 				FTransform ParentGlobalTransform = UE::Interchange::Private::FFbxConvert::ConvertTransform<FTransform, FVector, FQuat>(ParentTransform);
 				if (ParentGlobalTransform.ContainsNaN())
 				{
@@ -428,7 +433,7 @@ namespace UE::Interchange::Private
 			}
 			else
 			{
-				FbxAMatrix& LocalMatrix = FetchPayloadData.Node->EvaluateLocalTransform(CurTime);
+				FbxAMatrix& LocalMatrix = FetchPayloadData.Node->EvaluateLocalTransform(CurrentTime);
 				FbxVector4 NewLocalT = LocalMatrix.GetT();
 				FbxVector4 NewLocalS = LocalMatrix.GetS();
 				FbxQuaternion NewLocalQ = LocalMatrix.GetQ();
