[InterchangeWorker] Reverting minimum animation length usage.

#rnx

diff --git a/Engine/Plugins/Interchange/Runtime/Source/Parsers/CommonParser/Private/InterchangeCommonAnimationPayload.cpp b/Engine/Plugins/Interchange/Runtime/Source/Parsers/CommonParser/Private/InterchangeCommonAnimationPayload.cpp
--- a/Engine/Plugins/Interchange/Runtime/Source/Parsers/CommonParser/Private/InterchangeCommonAnimationPayload.cpp
+++ b/Engine/Plugins/Interchange/Runtime/Source/Parsers/CommonParser/Private/InterchangeCommonAnimationPayload.cpp
@@ -3,7 +3,6 @@
 #include "InterchangeCommonAnimationPayload.h"
 
 #include "CoreMinimal.h"
-#include "Animation/AnimTypes.h"
 
 #include UE_INLINE_GENERATED_CPP_BY_NAME(InterchangeCommonAnimationPayload)
 
@@ -69,7 +68,7 @@ namespace UE::Interchange
 			}
 
 			const double BakeInterval = 1.0 / BakeFrequency;
-			const double SequenceLength = FMath::Max<double>(RangeEndTime - RangeStartTime, MINIMUM_ANIMATION_LENGTH);
+			const double SequenceLength = FMath::Max<double>(RangeEndTime - RangeStartTime, BakeInterval);
 			int32 BakeKeyCount = FMath::RoundToInt32(SequenceLength * BakeFrequency) + 1;
 
 			auto EvaluateCurve = [this](const int32& CurveIndex, double CurrentTime, float DefaultValue)
diff --git a/Engine/Plugins/Interchange/Runtime/Source/Parsers/Fbx/Private/FbxAnimation.cpp b/Engine/Plugins/Interchange/Runtime/Source/Parsers/Fbx/Private/FbxAnimation.cpp
--- a/Engine/Plugins/Interchange/Runtime/Source/Parsers/Fbx/Private/FbxAnimation.cpp
+++ b/Engine/Plugins/Interchange/Runtime/Source/Parsers/Fbx/Private/FbxAnimation.cpp
@@ -21,7 +21,6 @@
 #include "StaticMeshAttributes.h"
 #include "InterchangeAnimationTrackSetNode.h"
 #include "Nodes/InterchangeBaseNodeContainer.h"
-#include "Animation/AnimTypes.h"
 
 #define LOCTEXT_NAMESPACE "InterchangeFbxMesh"
 
@@ -382,7 +381,7 @@ namespace UE::Interchange::Private
 		FbxTime TimeStep = 0;
 		TimeStep.SetSecondDouble(TimeStepSecond);
 
-		const double SequenceLength = FMath::Max<double>(AnimationBakeTransformPayloadData.RangeEndTime - AnimationBakeTransformPayloadData.RangeStartTime, MINIMUM_ANIMATION_LENGTH);
+		const double SequenceLength = FMath::Max<double>(AnimationBakeTransformPayloadData.RangeEndTime - AnimationBakeTransformPayloadData.RangeStartTime, TimeStepSecond);
 		const int32 NumFrame = FMath::RoundToInt32(SequenceLength * AnimationBakeTransformPayloadData.BakeFrequency);
 		int32 BakeKeyCount = NumFrame + 1;
 
