[Interchange] Unique Joint Name fixes.
(GLTF specific solution)

#rnx

diff --git a/Engine/Plugins/Interchange/Runtime/Source/Parsers/GLTFCore/Private/GLTF/GLTFAsset.cpp b/Engine/Plugins/Interchange/Runtime/Source/Parsers/GLTFCore/Private/GLTF/GLTFAsset.cpp
--- a/Engine/Plugins/Interchange/Runtime/Source/Parsers/GLTFCore/Private/GLTF/GLTFAsset.cpp
+++ b/Engine/Plugins/Interchange/Runtime/Source/Parsers/GLTFCore/Private/GLTF/GLTFAsset.cpp
@@ -9,6 +9,7 @@
 #include "GLTFNode.h"
 #include "GLTFTexture.h"
 #include "InterchangeHelper.h"
+#include "Misc/SecureHash.h"
 
 namespace GLTF
 {
@@ -118,6 +119,12 @@ namespace GLTF
 		}
 
 		{
+			//Nodes/Joints need to have (case ignored) unique naming,
+			// as Skeletons don't accept (casen ignored) identical bone names
+			//For that reason we are making sure here that the Node names are unique.
+
+			TSet<FString> Names;
+
 			const FString NodePrefix = Name + TEXT("_node_");
 			const FString JointPrefix = Name + TEXT("_joint_");
 			for (int32 NodeIndex = 0; NodeIndex < Nodes.Num(); ++NodeIndex)
@@ -135,6 +142,15 @@ namespace GLTF
 				}
 
 				Node.Name = UE::Interchange::MakeName(Node.Name, bIsJoint);
+
+				FString LowerCaseName = Node.Name.ToLower();
+				if (Names.Contains(LowerCaseName))
+				{
+					//to make the name unique we encode the UID into the name:
+					Node.Name += TEXT("_") + FMD5::HashAnsiString(*Node.UniqueId);
+				}
+
+				Names.Add(Node.Name);
 			}
 		}
 
