MaterialX:
- Connect directly the output of VectorParameters without using a ComponentMask node
- Update the extract node in case of value attribute
- Avoid creating a MF_FlattenNormal if the scale is 1.f

diff --git a/Engine/Plugins/Interchange/Runtime/Source/Import/Private/MaterialX/MaterialXUtils/MaterialXSurfaceShaderAbstract.cpp b/Engine/Plugins/Interchange/Runtime/Source/Import/Private/MaterialX/MaterialXUtils/MaterialXSurfaceShaderAbstract.cpp
--- a/Engine/Plugins/Interchange/Runtime/Source/Import/Private/MaterialX/MaterialXUtils/MaterialXSurfaceShaderAbstract.cpp
+++ b/Engine/Plugins/Interchange/Runtime/Source/Import/Private/MaterialX/MaterialXUtils/MaterialXSurfaceShaderAbstract.cpp
@@ -28,7 +28,7 @@ FMaterialXSurfaceShaderAbstract::FMaterialXSurfaceShaderAbstract(UInterchangeBas
 	, bTangentSpaceInput{ false }
 {}
 
-bool FMaterialXSurfaceShaderAbstract::AddAttribute(MaterialX::InputPtr Input, const FString& InputChannelName, UInterchangeShaderNode* ShaderNode)
+bool FMaterialXSurfaceShaderAbstract::AddAttribute(MaterialX::InputPtr Input, const FString& InputChannelName, UInterchangeShaderNode* ShaderNode, int32 OutputIndex)
 {
 	using namespace UE::Interchange::Materials::Standard::Nodes;
 	if(Input)
@@ -47,7 +47,7 @@ bool FMaterialXSurfaceShaderAbstract::AddAttribute(MaterialX::InputPtr Input, co
 		}
 		else if(Input->getType() == mx::Type::Color3 || Input->getType() == mx::Type::Color4)
 		{
-			return AddLinearColorAttribute(Input, InputChannelName, ShaderNode);
+			return AddLinearColorAttribute(Input, InputChannelName, ShaderNode, FLinearColor{ std::numeric_limits<float>::max(),std::numeric_limits<float>::max(),std::numeric_limits<float>::max() }, OutputIndex);
 		}
 		else if(Input->getType() == mx::Type::Vector2)
 		{
@@ -56,14 +56,14 @@ bool FMaterialXSurfaceShaderAbstract::AddAttribute(MaterialX::InputPtr Input, co
 		}
 		else if(Input->getType() == mx::Type::Vector3 || Input->getType() == mx::Type::Vector4)
 		{
-			return AddVectorAttribute(Input, InputChannelName, ShaderNode);
+			return AddVectorAttribute(Input, InputChannelName, ShaderNode, FVector4f{ std::numeric_limits<float>::max(),std::numeric_limits<float>::max(),std::numeric_limits<float>::max() }, OutputIndex);
 		}
 	}
 
 	return false;
 }
 
-bool FMaterialXSurfaceShaderAbstract::AddAttributeFromValueOrInterface(MaterialX::InputPtr Input, const FString& InputChannelName, UInterchangeShaderNode* ShaderNode)
+bool FMaterialXSurfaceShaderAbstract::AddAttributeFromValueOrInterface(MaterialX::InputPtr Input, const FString& InputChannelName, UInterchangeShaderNode* ShaderNode, int32 OutputIndex)
 {
 	bool bAttribute = false;
 
@@ -87,13 +87,13 @@ bool FMaterialXSurfaceShaderAbstract::AddAttributeFromValueOrInterface(MaterialX
 
 		if(Input->hasValue())
 		{
-			bAttribute = AddAttribute(Input, InputToConnectTo, ShaderNodeToConnectTo);
+			bAttribute = AddAttribute(Input, InputToConnectTo, ShaderNodeToConnectTo, OutputIndex);
 		}
 		else if(Input->hasInterfaceName())
 		{
 			if(mx::InputPtr InputInterface = Input->getInterfaceInput(); InputInterface->hasValue())
 			{
-				bAttribute = AddAttribute(InputInterface, InputToConnectTo, ShaderNodeToConnectTo);
+				bAttribute = AddAttribute(InputInterface, InputToConnectTo, ShaderNodeToConnectTo, OutputIndex);
 			}
 		}
 	}
@@ -144,7 +144,7 @@ bool FMaterialXSurfaceShaderAbstract::AddFloatAttribute(MaterialX::InputPtr Inpu
 	return false;
 }
 
-bool FMaterialXSurfaceShaderAbstract::AddLinearColorAttribute(MaterialX::InputPtr Input, const FString& InputChannelName, UInterchangeShaderNode* ShaderNode, const FLinearColor& DefaultValue)
+bool FMaterialXSurfaceShaderAbstract::AddLinearColorAttribute(MaterialX::InputPtr Input, const FString& InputChannelName, UInterchangeShaderNode* ShaderNode, const FLinearColor& DefaultValue, int32 OutputIndex)
 {
 	using namespace UE::Interchange::Materials::Standard::Nodes;
 	if(Input && Input->hasValue())
@@ -160,14 +160,14 @@ bool FMaterialXSurfaceShaderAbstract::AddLinearColorAttribute(MaterialX::InputPt
 
 			UInterchangeShaderNode* VectorParameterNode = CreateShaderNode(NodeName, VectorParameter::Name.ToString());
 			VectorParameterNode->AddLinearColorAttribute(UInterchangeShaderPortsAPI::MakeInputParameterKey(VectorParameter::Attributes::DefaultValue.ToString()), Value);
-			return UInterchangeShaderPortsAPI::ConnectDefaultOuputToInput(ShaderNode, InputChannelName, VectorParameterNode->GetUniqueID());
+			return UInterchangeShaderPortsAPI::ConnectOuputToInputByIndex(ShaderNode, InputChannelName, VectorParameterNode->GetUniqueID(), OutputIndex);
 		}
 	}
 
 	return false;
 }
 
-bool FMaterialXSurfaceShaderAbstract::AddVectorAttribute(MaterialX::InputPtr Input, const FString& InputChannelName, UInterchangeShaderNode* ShaderNode, const FVector4f& DefaultValue)
+bool FMaterialXSurfaceShaderAbstract::AddVectorAttribute(MaterialX::InputPtr Input, const FString& InputChannelName, UInterchangeShaderNode* ShaderNode, const FVector4f& DefaultValue, int32 OutputIndex)
 {
 	using namespace UE::Interchange::Materials::Standard::Nodes;
 	if(Input && Input->hasValue())
@@ -183,7 +183,7 @@ bool FMaterialXSurfaceShaderAbstract::AddVectorAttribute(MaterialX::InputPtr Inp
 
 			UInterchangeShaderNode* VectorParameterNode = CreateShaderNode(NodeName, VectorParameter::Name.ToString());
 			VectorParameterNode->AddLinearColorAttribute(UInterchangeShaderPortsAPI::MakeInputParameterKey(VectorParameter::Attributes::DefaultValue.ToString()), Value);
-			return UInterchangeShaderPortsAPI::ConnectDefaultOuputToInput(ShaderNode, InputChannelName, VectorParameterNode->GetUniqueID());
+			return UInterchangeShaderPortsAPI::ConnectOuputToInputByIndex(ShaderNode, InputChannelName, VectorParameterNode->GetUniqueID(), OutputIndex);
 		}
 	}
 
@@ -414,24 +414,22 @@ void FMaterialXSurfaceShaderAbstract::ConnectExtractInputToOutput(const FConnect
 {
 	using namespace UE::Interchange::Materials::Standard::Nodes;
 
-	UInterchangeShaderNode* MaskShaderNode = CreateShaderNode(Connect.UpstreamNode->getName().c_str(), Mask::Name.ToString());
-
+	uint8 Index = 0;
 	if(mx::InputPtr InputIndex = Connect.UpstreamNode->getInput("index"))
 	{
-		const int32 Index = mx::fromValueString<int>(InputIndex->getValueString());
-		switch(Index)
-		{
-		case 0: MaskShaderNode->AddBooleanAttribute(Mask::Attributes::R.ToString(), true); break;
-		case 1: MaskShaderNode->AddBooleanAttribute(Mask::Attributes::G.ToString(), true); break;
-		case 2: MaskShaderNode->AddBooleanAttribute(Mask::Attributes::B.ToString(), true); break;
-		case 3: MaskShaderNode->AddBooleanAttribute(Mask::Attributes::A.ToString(), true); break;
-		default:
-			UE_LOG(LogInterchangeImport, Warning, TEXT("Wrong index number for extracted node. Values are from [0-3]."));
-			break;
-		}
+		Index = mx::fromValueString<int>(InputIndex->getValueString());
 	}
 
-	UInterchangeShaderPortsAPI::ConnectDefaultOuputToInput(Connect.ParentShaderNode, Connect.InputChannelName, MaskShaderNode->GetUniqueID());
+	if(mx::InputPtr Input = Connect.UpstreamNode->getInput("in"); Input && Input->hasValue())
+	{
+		// Output 0 means RGB, 1st channel starts at 1
+		AddAttributeFromValueOrInterface(Input, Connect.InputChannelName, Connect.ParentShaderNode, Index + 1);
+	}
+	else
+	{
+		UInterchangeShaderNode* MaskShaderNode = CreateMaskShaderNode(1 << (3 - Index), Connect.UpstreamNode->getName().c_str());
+		UInterchangeShaderPortsAPI::ConnectDefaultOuputToInput(Connect.ParentShaderNode, Connect.InputChannelName, MaskShaderNode->GetUniqueID());
+	}
 }
 
 void FMaterialXSurfaceShaderAbstract::ConnectDotInputToOutput(const FConnectNode& Connect)
@@ -1158,46 +1156,49 @@ void FMaterialXSurfaceShaderAbstract::ConnectTexCoordInputToOutput(const FConnec
 
 void FMaterialXSurfaceShaderAbstract::ConnectSeparateInputToOutput(const FConnectNode& Connect)
 {
+	bool bHasValueOrInterface = false;
+	if(mx::InputPtr Input = Connect.UpstreamNode->getInput("in"))
+	{
+		bHasValueOrInterface = Input->hasValue() || Input->hasInterfaceName();
+	}
+
+	uint8 Mask = 0b1110;
+	int32 OutpuxIndex = 0;
+
 	if(Connect.OutputName == TEXT("outx") || Connect.OutputName == TEXT("outr"))
 	{
-		UInterchangeShaderNode* OutXNode = CreateMaskShaderNode(0b1000, Connect.UpstreamNode->getName().c_str(), Connect.OutputName);
-		if(mx::InputPtr Input = Connect.UpstreamNode->getInput("in"))
-		{
-			AddAttributeFromValueOrInterface(Input, GetInputName(Input), OutXNode);
-		}
-		UInterchangeShaderPortsAPI::ConnectDefaultOuputToInput(Connect.ParentShaderNode, Connect.InputChannelName, OutXNode->GetUniqueID());
+		Mask = 0b1000;
+		OutpuxIndex = 1;
 	}
 	else if(Connect.OutputName == TEXT("outy") || Connect.OutputName == TEXT("outg"))
 	{
-		UInterchangeShaderNode* OutYNode = CreateMaskShaderNode(0b0100, Connect.UpstreamNode->getName().c_str(), Connect.OutputName);
-		if(mx::InputPtr Input = Connect.UpstreamNode->getInput("in"))
-		{
-			AddAttributeFromValueOrInterface(Input, GetInputName(Input), OutYNode);
-		}
-		UInterchangeShaderPortsAPI::ConnectDefaultOuputToInput(Connect.ParentShaderNode, Connect.InputChannelName, OutYNode->GetUniqueID());
+		Mask = 0b0100;
+		OutpuxIndex = 2;
 	}
 	else if(Connect.OutputName == TEXT("outz") || Connect.OutputName == TEXT("outb"))
 	{
-		UInterchangeShaderNode* OutZNode = CreateMaskShaderNode(0b0010, Connect.UpstreamNode->getName().c_str(), Connect.OutputName);
-		if(mx::InputPtr Input = Connect.UpstreamNode->getInput("in"))
-		{
-			AddAttributeFromValueOrInterface(Input, GetInputName(Input), OutZNode);
-		}
-		UInterchangeShaderPortsAPI::ConnectDefaultOuputToInput(Connect.ParentShaderNode, Connect.InputChannelName, OutZNode->GetUniqueID());
+		Mask = 0b0010;
+		OutpuxIndex = 3;
 	}
 	else if(Connect.OutputName == TEXT("outw") || Connect.OutputName == TEXT("outa"))
 	{
-		UInterchangeShaderNode* OutWNode = CreateMaskShaderNode(0b0001, Connect.UpstreamNode->getName().c_str(), Connect.OutputName);
-		if(mx::InputPtr Input = Connect.UpstreamNode->getInput("in"))
-		{
-			AddAttributeFromValueOrInterface(Input, GetInputName(Input), OutWNode);
-		}
-		UInterchangeShaderPortsAPI::ConnectDefaultOuputToInput(Connect.ParentShaderNode, Connect.InputChannelName, OutWNode->GetUniqueID());
+		Mask = 0b0001;
+		OutpuxIndex = 4;
 	}
 	else
 	{
 		UE_LOG(LogInterchangeImport, Warning, TEXT("output <%s> not defined in <%s>."), *Connect.OutputName, ANSI_TO_TCHAR(Connect.UpstreamNode->getCategory().c_str()));
 	}
+
+	if(!bHasValueOrInterface)
+	{
+		UInterchangeShaderNode* OutNode = CreateMaskShaderNode(Mask, Connect.UpstreamNode->getName().c_str(), Connect.OutputName);
+		UInterchangeShaderPortsAPI::ConnectDefaultOuputToInput(Connect.ParentShaderNode, Connect.InputChannelName, OutNode->GetUniqueID());
+	}
+	else
+	{
+		AddAttributeFromValueOrInterface(Connect.UpstreamNode->getInput("in"), Connect.InputChannelName, Connect.ParentShaderNode, OutpuxIndex);
+	}
 }
 
 void FMaterialXSurfaceShaderAbstract::ConnectSwizzleInputToOutput(const FConnectNode& Connect)
@@ -1222,7 +1223,7 @@ void FMaterialXSurfaceShaderAbstract::ConnectNormalMapInputToOutput(const FConne
 	using namespace UE::Interchange::Materials::Standard::Nodes;
 
 	// Only create a FunctionCall if there's a scale, otherwise just like dot
-	if(mx::InputPtr Input = Connect.UpstreamNode->getInput("scale"))
+	if(mx::InputPtr Input = Connect.UpstreamNode->getInput("scale"); Input && Input->hasValueString() && (mx::fromValueString<float>(Input->getValueString()) != 1.f))
 	{
 		UInterchangeShaderNode* FlattenNormalNode = CreateFunctionCallShaderNode(Connect.UpstreamNode->getName().c_str(), TEXT("/Engine/Functions/Engine_MaterialFunctions01/Texturing/FlattenNormal.FlattenNormal"));
 
diff --git a/Engine/Plugins/Interchange/Runtime/Source/Import/Public/MaterialX/MaterialXUtils/MaterialXSurfaceShaderAbstract.h b/Engine/Plugins/Interchange/Runtime/Source/Import/Public/MaterialX/MaterialXUtils/MaterialXSurfaceShaderAbstract.h
--- a/Engine/Plugins/Interchange/Runtime/Source/Import/Public/MaterialX/MaterialXUtils/MaterialXSurfaceShaderAbstract.h
+++ b/Engine/Plugins/Interchange/Runtime/Source/Import/Public/MaterialX/MaterialXUtils/MaterialXSurfaceShaderAbstract.h
@@ -49,7 +49,7 @@ protected:
 	 *
 	 * @return true if the attribute was successfully added.
 	 */
-	bool AddAttribute(MaterialX::InputPtr Input, const FString& InputChannelName, UInterchangeShaderNode* ShaderNode);
+	bool AddAttribute(MaterialX::InputPtr Input, const FString& InputChannelName, UInterchangeShaderNode* ShaderNode, int32 OutputIndex = 0);
 
 	/**
 	 * Add an attribute to a shader node from the given MaterialX input if that input has either a value or an interface name.
@@ -60,7 +60,7 @@ protected:
 	 *
 	 * @return true if the attribute was successfully added.
 	 */
-	bool AddAttributeFromValueOrInterface(MaterialX::InputPtr Input, const FString& InputChannelName, UInterchangeShaderNode* ShaderNode);
+	bool AddAttributeFromValueOrInterface(MaterialX::InputPtr Input, const FString& InputChannelName, UInterchangeShaderNode* ShaderNode, int32 OutputIndex = 0);
 
 	/**
 	 * Add a bool attribute to a shader node only if its value taken from the input is not equal to its default value. Return false if the attribute does not exist or if we cannot add it
@@ -96,7 +96,7 @@ protected:
 	 *
 	 * @return true if the attribute was successfully added.
 	 */
-	bool AddLinearColorAttribute(MaterialX::InputPtr Input, const FString& InputChannelName, UInterchangeShaderNode* ShaderNode, const FLinearColor& DefaultValue = FLinearColor{std::numeric_limits<float>::max(), std::numeric_limits<float>::max(), std::numeric_limits<float>::max()});
+	bool AddLinearColorAttribute(MaterialX::InputPtr Input, const FString& InputChannelName, UInterchangeShaderNode* ShaderNode, const FLinearColor& DefaultValue = FLinearColor{std::numeric_limits<float>::max(), std::numeric_limits<float>::max(), std::numeric_limits<float>::max()}, int32 OutputIndex = 0);
 
 	/**
 	 * Add an FLinearColor attribute to a shader node only if its value taken from the input is not equal to its default value. Return false if the attribute does not exist or if we cannot add it.
@@ -108,7 +108,7 @@ protected:
 	 *
 	 * @return true if the attribute was successfully added.
 	 */
-	bool AddVectorAttribute(MaterialX::InputPtr Input, const FString& InputChannelName, UInterchangeShaderNode* ShaderNode, const FVector4f& DefaultValue = FVector4f{ std::numeric_limits<float>::max(), std::numeric_limits<float>::max(), std::numeric_limits<float>::max() });
+	bool AddVectorAttribute(MaterialX::InputPtr Input, const FString& InputChannelName, UInterchangeShaderNode* ShaderNode, const FVector4f& DefaultValue = FVector4f{ std::numeric_limits<float>::max(), std::numeric_limits<float>::max(), std::numeric_limits<float>::max() }, int32 OutputIndex = 0);
 
 	/**
 	 * Connect an output either from a node name or a node graph from a MaterialX input to the shader node.
