Interchange: Fix gltf pipeline display name

#rnx

#virtualized

diff --git a/Engine/Plugins/Interchange/Runtime/Source/Pipelines/Public/InterchangeglTFPipeline.h b/Engine/Plugins/Interchange/Runtime/Source/Pipelines/Public/InterchangeglTFPipeline.h
--- a/Engine/Plugins/Interchange/Runtime/Source/Pipelines/Public/InterchangeglTFPipeline.h
+++ b/Engine/Plugins/Interchange/Runtime/Source/Pipelines/Public/InterchangeglTFPipeline.h
@@ -51,6 +51,11 @@ class INTERCHANGEPIPELINES_API UInterchangeGLTFPipeline : public UInterchangePip
 	UInterchangeGLTFPipeline();
 
 public:
+	
+	/** The name of the pipeline that will be display in the import dialog. */
+	UPROPERTY(EditAnywhere, BlueprintReadWrite, Category = "Common", meta = (StandAlonePipelineProperty = "True", PipelineInternalEditionData = "True"))
+	FString PipelineDisplayName;
+
 	TObjectPtr<UGLTFPipelineSettings> GLTFPipelineSettings;
 
 protected:
