MaterialX:
- add interfacename for Texture node

diff --git a/Engine/Plugins/Interchange/Runtime/Source/Import/Public/MaterialX/MaterialXUtils/MaterialXSurfaceShaderAbstract.h b/Engine/Plugins/Interchange/Runtime/Source/Import/Public/MaterialX/MaterialXUtils/MaterialXSurfaceShaderAbstract.h
--- a/Engine/Plugins/Interchange/Runtime/Source/Import/Public/MaterialX/MaterialXUtils/MaterialXSurfaceShaderAbstract.h
+++ b/Engine/Plugins/Interchange/Runtime/Source/Import/Public/MaterialX/MaterialXUtils/MaterialXSurfaceShaderAbstract.h
@@ -332,9 +332,9 @@ protected:
 		//A node image should have an input file otherwise the user should check its default value
 		if(Node)
 		{
-			if(MaterialX::InputPtr InputFile = Node->getInput("file"); InputFile && InputFile->hasValue())
+			if(MaterialX::InputPtr InputFile = Node->getInput("file"); InputFile && (InputFile->hasValue() || InputFile->hasInterfaceName()))
 			{
-				FString Filepath{ InputFile->getValueString().c_str() };
+				FString Filepath{ InputFile->hasValue() ? InputFile->getValueString().c_str() : InputFile->getInterfaceInput()->getValueString().c_str()};
 				const FString FilePrefix = GetFilePrefix(InputFile);
 				Filepath = FPaths::Combine(FilePrefix, Filepath);
 				const FString Filename = FPaths::GetCleanFilename(Filepath);
