[Interchange] Animation Interface optimisation.

#rnx

diff --git a/Engine/Plugins/Experimental/Enterprise/DatasmithInterchange/Source/DatasmithInterchange/Private/InterchangeDatasmithTranslator.cpp b/Engine/Plugins/Experimental/Enterprise/DatasmithInterchange/Source/DatasmithInterchange/Private/InterchangeDatasmithTranslator.cpp
--- a/Engine/Plugins/Experimental/Enterprise/DatasmithInterchange/Source/DatasmithInterchange/Private/InterchangeDatasmithTranslator.cpp
+++ b/Engine/Plugins/Experimental/Enterprise/DatasmithInterchange/Source/DatasmithInterchange/Private/InterchangeDatasmithTranslator.cpp
@@ -815,7 +815,7 @@ TFuture<TOptional<UE::Interchange::FMeshPayloadData>> UInterchangeDatasmithTrans
 	);
 }
 
-TFuture<TOptional<UE::Interchange::FAnimationPayloadData>> UInterchangeDatasmithTranslator::GetAnimationPayloadData(const FInterchangeAnimationPayLoadKey& PayLoadKey, const double BakeFrequency, const double RangeStartSecond, const double RangeStopSecond) const
+TFuture<TOptional<UE::Interchange::FAnimationPayloadData>> UInterchangeDatasmithTranslator::GetAnimationPayloadData(const UE::Interchange::FAnimationPayloadQuery& PayloadQuery) const
 {
 	TPromise<TOptional<UE::Interchange::FAnimationPayloadData>> EmptyPromise;
 	EmptyPromise.SetValue(TOptional<UE::Interchange::FAnimationPayloadData>());
@@ -828,7 +828,7 @@ TFuture<TOptional<UE::Interchange::FAnimationPayloadData>> UInterchangeDatasmith
 
 	TSharedPtr<IDatasmithBaseAnimationElement> AnimationElement;
 	float FrameRate = 0.f;
-	if (UE::DatasmithInterchange::AnimUtils::FAnimationPayloadDesc* PayloadDescPtr = AnimationPayLoadMapping.Find(PayLoadKey.UniqueId))
+	if (UE::DatasmithInterchange::AnimUtils::FAnimationPayloadDesc* PayloadDescPtr = AnimationPayLoadMapping.Find(PayloadQuery.PayloadKey.UniqueId))
 	{
 		AnimationElement = PayloadDescPtr->Value;
 		if (!ensure(AnimationElement))
@@ -840,15 +840,15 @@ TFuture<TOptional<UE::Interchange::FAnimationPayloadData>> UInterchangeDatasmith
 		FrameRate = PayloadDescPtr->Key;
 	}
 
-	if (PayLoadKey.Type != EInterchangeAnimationPayLoadType::NONE)
+	if (PayloadQuery.PayloadKey.Type != EInterchangeAnimationPayLoadType::NONE)
 	{
-		return Async(EAsyncExecution::TaskGraph, [this, AnimationElement = MoveTemp(AnimationElement), FrameRate, PayLoadType = PayLoadKey.Type]
+		return Async(EAsyncExecution::TaskGraph, [this, AnimationElement = MoveTemp(AnimationElement), FrameRate, PayloadQuery]
 			{
 
-				UE::Interchange::FAnimationPayloadData TransformPayloadData(PayLoadType);
+				UE::Interchange::FAnimationPayloadData TransformPayloadData(PayloadQuery.SceneNodeUniqueID, PayloadQuery.PayloadKey);
 				TOptional<UE::Interchange::FAnimationPayloadData> Result;
 
-				if (UE::DatasmithInterchange::AnimUtils::GetAnimationPayloadData(*AnimationElement, FrameRate, PayLoadType, TransformPayloadData))
+				if (UE::DatasmithInterchange::AnimUtils::GetAnimationPayloadData(*AnimationElement, FrameRate, PayloadQuery.PayloadKey.Type, TransformPayloadData))
 				{
 					Result.Emplace(MoveTemp(TransformPayloadData));
 				}
@@ -861,6 +861,28 @@ TFuture<TOptional<UE::Interchange::FAnimationPayloadData>> UInterchangeDatasmith
 	return EmptyPromise.GetFuture();
 }
 
+TArray<UE::Interchange::FAnimationPayloadData> UInterchangeDatasmithTranslator::GetAnimationPayloadData(const TArray<UE::Interchange::FAnimationPayloadQuery>& PayloadQueries) const
+{
+	TArray<TFuture<TOptional<UE::Interchange::FAnimationPayloadData>>> AnimationPayloadFutures;
+	for (const UE::Interchange::FAnimationPayloadQuery& PayloadQuery : PayloadQueries)
+	{
+		AnimationPayloadFutures.Add(GetAnimationPayloadData(PayloadQuery));
+	}
+
+	TArray<UE::Interchange::FAnimationPayloadData> AnimationPayloads;
+	for (TFuture<TOptional<UE::Interchange::FAnimationPayloadData>>& AnimationPayloadFuture : AnimationPayloadFutures)
+	{
+		TOptional<UE::Interchange::FAnimationPayloadData> OptionalPayloadData = AnimationPayloadFuture.Get();
+		if (!OptionalPayloadData.IsSet())
+		{
+			continue;
+		}
+		AnimationPayloads.Add(OptionalPayloadData.GetValue());
+	}
+
+	return AnimationPayloads;
+}
+
 TFuture<TOptional<UE::Interchange::FVariantSetPayloadData>> UInterchangeDatasmithTranslator::GetVariantSetPayloadData(const FString& PayloadKey) const
 {
 	using namespace UE::Interchange;
diff --git a/Engine/Plugins/Experimental/Enterprise/DatasmithInterchange/Source/DatasmithInterchange/Public/InterchangeDatasmithTranslator.h b/Engine/Plugins/Experimental/Enterprise/DatasmithInterchange/Source/DatasmithInterchange/Public/InterchangeDatasmithTranslator.h
--- a/Engine/Plugins/Experimental/Enterprise/DatasmithInterchange/Source/DatasmithInterchange/Public/InterchangeDatasmithTranslator.h
+++ b/Engine/Plugins/Experimental/Enterprise/DatasmithInterchange/Source/DatasmithInterchange/Public/InterchangeDatasmithTranslator.h
@@ -116,7 +116,8 @@ public:
 	/* IInterchangeStaticMeshPayloadInterface End */
 
 	/* IInterchangeAnimationPayloadInterface Begin */
-	virtual TFuture<TOptional<UE::Interchange::FAnimationPayloadData>> GetAnimationPayloadData(const FInterchangeAnimationPayLoadKey& PayLoadKey, const double BakeFrequency = 0, const double RangeStartSecond = 0, const double RangeStopSecond = 0) const override;
+	TFuture<TOptional<UE::Interchange::FAnimationPayloadData>> GetAnimationPayloadData(const UE::Interchange::FAnimationPayloadQuery& PayloadQuery) const;
+	virtual TArray<UE::Interchange::FAnimationPayloadData> GetAnimationPayloadData(const TArray<UE::Interchange::FAnimationPayloadQuery>& PayloadQueries) const override;
 	/* IInterchangeAnimationPayloadInterface End */
 
 	/* IInterchangeVariantSetPayloadInterface Begin */
diff --git a/Engine/Plugins/Interchange/Runtime/Source/Dispatcher/Private/InterchangeDispatcherTask.cpp b/Engine/Plugins/Interchange/Runtime/Source/Dispatcher/Private/InterchangeDispatcherTask.cpp
--- a/Engine/Plugins/Interchange/Runtime/Source/Dispatcher/Private/InterchangeDispatcherTask.cpp
+++ b/Engine/Plugins/Interchange/Runtime/Source/Dispatcher/Private/InterchangeDispatcherTask.cpp
@@ -311,23 +311,23 @@ namespace UE
 			return true;
 		}
 
-		FString FJsonFetchAnimationBakeTransformPayloadCmd::ToJson() const
+		FString FJsonFetchAnimationQueriesCmd::ToJson() const
 		{
 			//Code should not do a ToJson if the data was not set before
 			ensure(bIsDataInitialize);
 
 			TSharedPtr<FJsonObject> CmdObject = MakeShared<FJsonObject>();
+			
 			TSharedPtr<FJsonObject> ActionDataObject = MakeShared<FJsonObject>();
 			//CmdObject
 			CmdObject->SetStringField(GetCommandIDJsonKey(), GetAction());
 			CmdObject->SetStringField(GetTranslatorIDJsonKey(), GetTranslatorID());
-			ActionDataObject->SetStringField(GetPayloadKeyJsonKey(), GetPayloadKey());
-			
-			//Bake settings
-			ActionDataObject->SetNumberField(GetBakeFrequencyJsonKey(), GetBakeFrequency());
-			ActionDataObject->SetNumberField(GetRangeStartTimeJsonKey(), GetRangeStartTime());
-			ActionDataObject->SetNumberField(GetRangeEndTimeJsonKey(), GetRangeEndTime());
 			
+			ActionDataObject->SetStringField(GetPayloadKeyJsonKey(), GetPayloadKey());
+
+			//Queries:
+			ActionDataObject->SetStringField(GetQueriesJsonStringKey(), QueriesJsonString);
+
 			CmdObject->SetObjectField(GetCommandDataJsonKey(), ActionDataObject);
 
 			FString FetchPayloadCmd;
@@ -340,7 +340,7 @@ namespace UE
 			return FetchPayloadCmd;
 		}
 
-		bool FJsonFetchAnimationBakeTransformPayloadCmd::FromJson(const FString& JsonString)
+		bool FJsonFetchAnimationQueriesCmd::FromJson(const FString& JsonString)
 		{
 			TSharedRef<TJsonReader<TCHAR>> Reader = FJsonStringReader::Create(JsonString);
 
@@ -350,6 +350,7 @@ namespace UE
 				//Cannot read the json file
 				return false;
 			}
+
 			FString JsonActionValue;
 			if (!CmdObject->TryGetStringField(GetCommandIDJsonKey(), JsonActionValue))
 			{
@@ -380,22 +381,55 @@ namespace UE
 				return false;
 			}
 
-			//Bake settings
-			if (!(*ActionDataObject)->TryGetNumberField(GetBakeFrequencyJsonKey(), BakeFrequency))
+			if (!((*ActionDataObject)->TryGetStringField(GetQueriesJsonStringKey(), QueriesJsonString)))
 			{
 				return false;
 			}
-			if (!(*ActionDataObject)->TryGetNumberField(GetRangeStartTimeJsonKey(), RangeStartTime))
+
+			//Since we filled the data from the json file, set the data has been initialize.
+			bIsDataInitialize = true;
+			return true;
+		}
+
+		FString FJsonFetchAnimationQueriesCmd::JsonAnimationQueriesResultParser::ToJson() const
+		{
+			TSharedPtr<FJsonObject> ResultObject = MakeShared<FJsonObject>();
+
+			for (const TPair<FString, FString>& HashToFilename : HashToFilenames)
 			{
-				return false;
+				ResultObject->SetStringField(HashToFilename.Key, HashToFilename.Value);
 			}
-			if (!(*ActionDataObject)->TryGetNumberField(GetRangeEndTimeJsonKey(), RangeEndTime))
+
+			FString JsonResult;
+			TSharedRef< TJsonWriter< TCHAR, TPrettyJsonPrintPolicy<TCHAR> > > JsonWriter = TJsonWriterFactory< TCHAR, TPrettyJsonPrintPolicy<TCHAR> >::Create(&JsonResult);
+			if (!FJsonSerializer::Serialize(ResultObject.ToSharedRef(), JsonWriter))
 			{
+				//Error creating the json cmd string 
+				return FString();
+			}
+			return JsonResult;
+		}
+
+		bool FJsonFetchAnimationQueriesCmd::JsonAnimationQueriesResultParser::FromJson(const FString& JsonString)
+		{
+			HashToFilenames.Empty();
+
+			TSharedRef<TJsonReader<TCHAR>> Reader = FJsonStringReader::Create(JsonString);
+
+			TSharedPtr<FJsonObject> ResultObject;
+			if (!FJsonSerializer::Deserialize(Reader, ResultObject) || !ResultObject.IsValid())
+			{
+				//Cannot read the json file
 				return false;
 			}
 
-			//Since we filled the data from the json file, set the data has been initialize.
-			bIsDataInitialize = true;
+			HashToFilenames.Reserve(ResultObject->Values.Num());
+
+			for (const TPair<FString, TSharedPtr<FJsonValue>>& Value : ResultObject->Values)
+			{
+				HashToFilenames.Add(Value.Key, Value.Value->AsString());
+			}
+
 			return true;
 		}
 	} //ns Interchange
diff --git a/Engine/Plugins/Interchange/Runtime/Source/Dispatcher/Public/InterchangeDispatcherTask.h b/Engine/Plugins/Interchange/Runtime/Source/Dispatcher/Public/InterchangeDispatcherTask.h
--- a/Engine/Plugins/Interchange/Runtime/Source/Dispatcher/Public/InterchangeDispatcherTask.h
+++ b/Engine/Plugins/Interchange/Runtime/Source/Dispatcher/Public/InterchangeDispatcherTask.h
@@ -9,6 +9,7 @@ namespace UE
 	namespace Interchange
 	{
 		struct FTask;
+		struct FAnimationPayloadQuery;
 	}
 }
 DECLARE_DELEGATE_OneParam(FInterchangeDispatcherTaskCompleted, int32 TaskIndex);
@@ -333,69 +334,68 @@ namespace UE
 			FTransform MeshGlobalTransform = FTransform::Identity;
 		};
 
-		//Animation transform payload require transform to be bake by the translator
-		//The anim sequence API is not yet using curve to describe bone track animation
-		class INTERCHANGEDISPATCHER_API FJsonFetchAnimationBakeTransformPayloadCmd : public FJsonFetchPayloadCmd
+		class INTERCHANGEDISPATCHER_API FJsonFetchAnimationQueriesCmd : public FJsonFetchPayloadCmd
 		{
 		public:
-			FJsonFetchAnimationBakeTransformPayloadCmd()
+			FJsonFetchAnimationQueriesCmd()
 			{
 				check(!bIsDataInitialize);
 			}
 
-			FJsonFetchAnimationBakeTransformPayloadCmd(const FString& InTranslatorID
-				, const FString& InPayloadKey
-				, double InBakeFrequency
-				, double InRangeStartTime
-				, double InRangeEndTime)
-				: FJsonFetchPayloadCmd(InTranslatorID, InPayloadKey)
-				, BakeFrequency(InBakeFrequency)
-				, RangeStartTime(InRangeStartTime)
-				, RangeEndTime(InRangeEndTime)
-			{}
+			FJsonFetchAnimationQueriesCmd(const FString& InTranslatorID, const FString& InQueriesJsonString) 
+				: FJsonFetchPayloadCmd(InTranslatorID, TEXT("AnimationQueries"))
+				, QueriesJsonString(InQueriesJsonString)
+			{
+			}
 
 			virtual FString ToJson() const override;
 			virtual bool FromJson(const FString& JsonString) override;
 
-			double GetBakeFrequency() const
+			static FString GetQueriesJsonStringKey()
 			{
-				ensure(bIsDataInitialize);
-				return BakeFrequency;
-			}
-
-			static FString GetBakeFrequencyJsonKey()
-			{
-				static const FString Key = TEXT("BakeFrequency");
+				static const FString Key = TEXT("QueriesJsonString");
 				return Key;
 			}
 
-			double GetRangeStartTime() const
+			FString GetQueriesJsonString() const
 			{
-				ensure(bIsDataInitialize);
-				return RangeStartTime;
+				return QueriesJsonString;
 			}
 
-			static FString GetRangeStartTimeJsonKey()
+			/**
+			 * Use this class helper to create the cmd result json string and to read it
+			 */
+			class INTERCHANGEDISPATCHER_API JsonAnimationQueriesResultParser
 			{
-				static const FString Key = TEXT("RangeStartTime");
-				return Key;
-			}
+			public:
+				TMap<FString, FString> GetResultFilename() const
+				{
+					return HashToFilenames;
+				}
+				void SetHashToFilenames(const TMap<FString, FString>& InHashToFilenames)
+				{
+					HashToFilenames = InHashToFilenames;
+				}
 
-			double GetRangeEndTime() const
-			{
-				ensure(bIsDataInitialize);
-				return RangeEndTime;
-			}
+				FString ToJson() const;
+				bool FromJson(const FString& JsonString);
+
+				static FString GetHashToFilenamesKey()
+				{
+					const FString Key = TEXT("HashToFilenames");
+					return Key;
+				}
+
+				const TMap<FString, FString>& GetHashToFilenames() const
+				{
+					return HashToFilenames;
+				}
+			private:
+				TMap<FString, FString> HashToFilenames;
+			};
 
-			static FString GetRangeEndTimeJsonKey()
-			{
-				static const FString Key = TEXT("RangeEndTime");
-				return Key;
-			}
 		protected:
-			double BakeFrequency = 30.0;
-			double RangeStartTime = 0.0;
-			double RangeEndTime = 1.0/BakeFrequency;
+			FString QueriesJsonString;
 		};
 	} //ns Interchange
 }//ns UE
diff --git a/Engine/Plugins/Interchange/Runtime/Source/Import/Private/Animation/InterchangeAnimSequenceFactory.cpp b/Engine/Plugins/Interchange/Runtime/Source/Import/Private/Animation/InterchangeAnimSequenceFactory.cpp
--- a/Engine/Plugins/Interchange/Runtime/Source/Import/Private/Animation/InterchangeAnimSequenceFactory.cpp
+++ b/Engine/Plugins/Interchange/Runtime/Source/Import/Private/Animation/InterchangeAnimSequenceFactory.cpp
@@ -746,7 +746,8 @@ namespace UE::Interchange::Private
 				{
 					TArray<FString> AttributeCurveNames;
 					AnimSequenceFactoryNode->GetAnimatedAttributeCurveNames(AttributeCurveNames);
-					TMap<FString, TFuture<TOptional<UE::Interchange::FAnimationPayloadData>>> AnimationCurvesPayloads;
+					TArray<FAnimationPayloadQuery> PayloadQueries;
+					TMap<TPair<FString, FString>, FString> CurveNames; // <<NodeUid,PayloadUID>, Name>
 					for (const FString& NodeUid : SkeletonNodes)
 					{
 						if (const UInterchangeSceneNode* SkeletonSceneNode = Cast<UInterchangeSceneNode>(NodeContainer->GetNode(NodeUid)))
@@ -767,28 +768,21 @@ namespace UE::Interchange::Private
 								//This goes slightly against the intent of the Type / FInterchangeAnimationPayLoadKey usage as we set the Type here, outside of the Translator
 								// this is due to the nature of the Attribute curves.
 								// Could be potentially reworked so that with the AttributeInfos we store the Type as well.
-								AnimationCurvesPayloads.Add(CurveNamePayload.Key, AnimSequenceTranslatorPayloadInterface->GetAnimationPayloadData(FInterchangeAnimationPayLoadKey(CurveNamePayload.Value, EInterchangeAnimationPayLoadType::CURVE)));
+								PayloadQueries.Add(FAnimationPayloadQuery(NodeUid, FInterchangeAnimationPayLoadKey(CurveNamePayload.Value, EInterchangeAnimationPayLoadType::CURVE)));
+								CurveNames.Add(TPair<FString, FString>(NodeUid, CurveNamePayload.Value), CurveNamePayload.Key);
 							}
 						}
 					}
+					TArray<FAnimationPayloadData> Payloads = AnimSequenceTranslatorPayloadInterface->GetAnimationPayloadData(PayloadQueries);
 
-					for (TPair<FString, TFuture<TOptional<UE::Interchange::FAnimationPayloadData>>>& CurveNameAndPayload : AnimationCurvesPayloads)
+					for (FAnimationPayloadData& CurvePayloadData : Payloads)
 					{
-						const FString& CurveName = CurveNameAndPayload.Key;
-						TOptional<UE::Interchange::FAnimationPayloadData> AnimationCurvePayload = CurveNameAndPayload.Value.Get();
-						if (!AnimationCurvePayload.IsSet())
-						{
-							UInterchangeResultWarning_Generic* Message = Factory->AddMessage<UInterchangeResultWarning_Generic>();
-							Message->DestinationAssetName = AssetName;
-							Message->AssetType = UAnimSequence::StaticClass();
-							Message->Text = FText::Format(NSLOCTEXT("UInterchangeAnimSequenceFactory", "AnimationPayloadInvalidCurve", "Invalid animation material curve payload key [{0}] in AnimSequence asset {1}"), FText::FromString(CurveName), FText::FromString(AssetName));
-							continue;
-						}
+						FString* CurveNamePtr = CurveNames.Find(TPair<FString, FString>(CurvePayloadData.SceneNodeUniqueID, CurvePayloadData.PayloadKey.UniqueId));
+						FString CurveName = CurveNamePtr ? *CurveNamePtr : TEXT("Unknown");
 
-						FAnimationPayloadData& CurvePayload = AnimationCurvePayload.GetValue();
 						if (bRemoveCurveRedundantKeys)
 						{
-							for (FRichCurve& RichCurve : CurvePayload.Curves)
+							for (FRichCurve& RichCurve : CurvePayloadData.Curves)
 							{
 								RichCurve.RemoveRedundantAutoTangentKeys(SMALL_NUMBER);
 							}
@@ -797,11 +791,11 @@ namespace UE::Interchange::Private
 						constexpr int32 CurveFlags = 0;
 						if (bMaterialDriveParameterOnCustomAttribute || IsCurveHookToMaterial(CurveName))
 						{
-							CreateMaterialCurve(AnimSequence, CurvePayload.Curves, CurveName, CurveFlags, bDoNotImportCurveWithZero, bAddCurveMetadataToSkeleton, bShouldTransact);
+							CreateMaterialCurve(AnimSequence, CurvePayloadData.Curves, CurveName, CurveFlags, bDoNotImportCurveWithZero, bAddCurveMetadataToSkeleton, bShouldTransact);
 						}
 						else
 						{
-							CreateAttributeCurve(AnimSequence, CurvePayload.Curves, CurveName, CurveFlags, bDoNotImportCurveWithZero, bAddCurveMetadataToSkeleton, bShouldTransact);
+							CreateAttributeCurve(AnimSequence, CurvePayloadData.Curves, CurveName, CurveFlags, bDoNotImportCurveWithZero, bAddCurveMetadataToSkeleton, bShouldTransact);
 						}
 					}
 				}
@@ -810,8 +804,12 @@ namespace UE::Interchange::Private
 				{
 					TArray<FString> AttributeStepCurveNames;
 					AnimSequenceFactoryNode->GetAnimatedAttributeStepCurveNames(AttributeStepCurveNames);
-					TMap<FString, TFuture<TOptional<UE::Interchange::FAnimationPayloadData>>> AnimationCurvesPayloads;
-					TMap<FString, TFuture<TOptional<UE::Interchange::FAnimationPayloadData>>> AnimationStepCurvesPayloads;
+					TMap<TPair<FString, FString>, FString> CurveNames; // <<NodeUid,PayloadUID>, Name>
+					TMap<TPair<FString, FString>, FString> StepCurveNames; // <<NodeUid,PayloadUID>, Name>
+
+					TArray<FAnimationPayloadQuery> CurvePayloadQueries;
+					TArray<FAnimationPayloadQuery> StepCurvePayloadQueries;
+
 					TMap<FString, FString> AnimationBoneNames;
 					for (const FString& NodeUid : SkeletonNodes)
 					{
@@ -842,16 +840,18 @@ namespace UE::Interchange::Private
 								//This goes slightly against the intent of the Type / FInterchangeAnimationPayLoadKey usage as we set the Type here, outside of the Translator
 								// this is due to the nature of the Attribute curves.
 								// Could be potentially reworked so that with the AttributeInfos we store the Type as well.
-								AnimationCurvesPayloads.Add(CurveNamePayload.Key, AnimSequenceTranslatorPayloadInterface->GetAnimationPayloadData(FInterchangeAnimationPayLoadKey(CurveNamePayload.Value, EInterchangeAnimationPayLoadType::CURVE)));
+								CurvePayloadQueries.Add(FAnimationPayloadQuery(NodeUid, FInterchangeAnimationPayLoadKey(CurveNamePayload.Value, EInterchangeAnimationPayLoadType::CURVE)));
 								AnimationBoneNames.Add(CurveNamePayload.Key, BoneName);
+								CurveNames.Add(TPair<FString, FString>(NodeUid, CurveNamePayload.Value), CurveNamePayload.Key);
 							}
 							for (const TPair<FString, FString>& StepCurveNamePayload : CurveNameStepCurvePayloads)
 							{
 								//This goes slightly against the intent of the Type / FInterchangeAnimationPayLoadKey usage as we set the Type here, outside of the Translator
 								// this is due to the nature of the Attribute curves.
 								// Could be potentially reworked so that with the AttributeInfos we store the Type as well.
-								AnimationStepCurvesPayloads.Add(StepCurveNamePayload.Key, AnimSequenceTranslatorPayloadInterface->GetAnimationPayloadData(FInterchangeAnimationPayLoadKey(StepCurveNamePayload.Value, EInterchangeAnimationPayLoadType::STEPCURVE)));
+								StepCurvePayloadQueries.Add(FAnimationPayloadQuery(NodeUid, FInterchangeAnimationPayLoadKey(StepCurveNamePayload.Value, EInterchangeAnimationPayLoadType::STEPCURVE)));
 								AnimationBoneNames.Add(StepCurveNamePayload.Key, BoneName);
+								StepCurveNames.Add(TPair<FString, FString>(NodeUid, StepCurveNamePayload.Value), StepCurveNamePayload.Key);
 							}
 						}
 					}
@@ -870,35 +870,24 @@ namespace UE::Interchange::Private
 						CreateAttributeStepCurve(AnimSequence, StepCurvePayload.StepCurves, CurveName, BoneName, CurveFlags, bDoNotImportCurveWithZero);
 					};
 
-					for (TPair<FString, TFuture<TOptional<UE::Interchange::FAnimationPayloadData>>>& CurveNameAndPayload : AnimationCurvesPayloads)
+					TArray<FAnimationPayloadData> CurvePayloads = AnimSequenceTranslatorPayloadInterface->GetAnimationPayloadData(CurvePayloadQueries);
+					for (FAnimationPayloadData& CurvePayloadData : CurvePayloads)
 					{
-						TOptional<UE::Interchange::FAnimationPayloadData> AnimationCurvePayload = CurveNameAndPayload.Value.Get();
-						if (!AnimationCurvePayload.IsSet())
-						{
-							UInterchangeResultWarning_Generic* Message = Factory->AddMessage<UInterchangeResultWarning_Generic>();
-							Message->DestinationAssetName = AssetName;
-							Message->AssetType = UAnimSequence::StaticClass();
-							Message->Text = FText::Format(NSLOCTEXT("UInterchangeAnimSequenceFactory", "AnimationPayloadInvalidCurvePayloadKey", "Invalid animation material curve payload key [{0}] in AnimSequence asset {1}"), FText::FromString(CurveNameAndPayload.Key), FText::FromString(AssetName));
-							continue;
-						}
-						FAnimationPayloadData& AnimationPayloadData = AnimationCurvePayload.GetValue();
-						AnimationPayloadData.CalculateDataFor(EInterchangeAnimationPayLoadType::STEPCURVE);
-						AddAttributeCurveToAnimSequence(AnimationPayloadData, CurveNameAndPayload.Key, AnimationBoneNames.FindChecked(CurveNameAndPayload.Key));
+						CurvePayloadData.CalculateDataFor(EInterchangeAnimationPayLoadType::STEPCURVE);
+
+						FString* CurveNamePtr = CurveNames.Find(TPair<FString, FString>(CurvePayloadData.SceneNodeUniqueID, CurvePayloadData.PayloadKey.UniqueId));
+						FString CurveName = CurveNamePtr ? *CurveNamePtr : TEXT("Unknown");
+
+						AddAttributeCurveToAnimSequence(CurvePayloadData, CurveName, AnimationBoneNames.FindChecked(CurveName));
 					}
 
-					for (TPair<FString, TFuture<TOptional<UE::Interchange::FAnimationPayloadData>>>& StepCurveNameAndPayload : AnimationStepCurvesPayloads)
+					TArray<FAnimationPayloadData> StepCurvePayloads = AnimSequenceTranslatorPayloadInterface->GetAnimationPayloadData(StepCurvePayloadQueries);
+					for (FAnimationPayloadData& StepCurvePayloadData : StepCurvePayloads)
 					{
-						TOptional<UE::Interchange::FAnimationPayloadData> AnimationStepCurvePayload = StepCurveNameAndPayload.Value.Get();
-						if (!AnimationStepCurvePayload.IsSet())
-						{
-							UInterchangeResultWarning_Generic* Message = Factory->AddMessage<UInterchangeResultWarning_Generic>();
-							Message->DestinationAssetName = AssetName;
-							Message->AssetType = UAnimSequence::StaticClass();
-							Message->Text = FText::Format(NSLOCTEXT("UInterchangeAnimSequenceFactory", "AnimationPayloadInvalidStepCurvePayloadKey", "Invalid animation curve payload key [{0}] in AnimSequence asset {1}"), FText::FromString(StepCurveNameAndPayload.Key), FText::FromString(AssetName));
-							continue;
-						}
+						FString* CurveNamePtr = StepCurveNames.Find(TPair<FString, FString>(StepCurvePayloadData.SceneNodeUniqueID, StepCurvePayloadData.PayloadKey.UniqueId));
+						FString CurveName = CurveNamePtr ? *CurveNamePtr : TEXT("Unknown");
 
-						AddAttributeCurveToAnimSequence(AnimationStepCurvePayload.GetValue(), StepCurveNameAndPayload.Key, AnimationBoneNames.FindChecked(StepCurveNameAndPayload.Key));
+						AddAttributeCurveToAnimSequence(StepCurvePayloadData, CurveName, AnimationBoneNames.FindChecked(CurveName));
 					}
 				}
 			}
@@ -1089,8 +1078,6 @@ UInterchangeFactoryBase::FImportAssetResult UInterchangeAnimSequenceFactory::Imp
 		USkeleton* Skeleton = AnimSequence->GetSkeleton();
 		check(Skeleton);
 
-		TMap<const UInterchangeSceneNode*, TFuture<TOptional<UE::Interchange::FAnimationPayloadData>>> AnimationPayloads;
-
 		//Get the sample rate, default to 30Hz in case the attribute is missing
 		double SampleRate = 30.0;
 		AnimSequenceFactoryNode->GetCustomImportBoneTracksSampleRate(SampleRate);
@@ -1105,42 +1092,41 @@ UInterchangeFactoryBase::FImportAssetResult UInterchangeAnimSequenceFactory::Imp
 		TMap<FString, FInterchangeAnimationPayLoadKey> PayloadKeys;
 		AnimSequenceFactoryNode->GetSceneNodeAnimationPayloadKeys(PayloadKeys);
 
+		TArray<UE::Interchange::FAnimationPayloadQuery> PayloadQueries;
+
 		for (const TTuple<FString, FInterchangeAnimationPayLoadKey>& SceneNodeUidAndPayloadKey : PayloadKeys)
 		{
-			if (const UInterchangeSceneNode* SkeletonSceneNode = Cast<UInterchangeSceneNode>(Arguments.NodeContainer->GetNode(SceneNodeUidAndPayloadKey.Key)))
-			{
-				AnimationPayloads.Add(SkeletonSceneNode, AnimSequenceTranslatorPayloadInterface->GetAnimationPayloadData(SceneNodeUidAndPayloadKey.Value, SampleRate, RangeStart, RangeEnd));
-			}
+			PayloadQueries.Add(UE::Interchange::FAnimationPayloadQuery(SceneNodeUidAndPayloadKey.Key, SceneNodeUidAndPayloadKey.Value, SampleRate, RangeStart, RangeEnd));
 		}
+
+		TArray<UE::Interchange::FAnimationPayloadData> AnimationPayloads = AnimSequenceTranslatorPayloadInterface->GetAnimationPayloadData(PayloadQueries);
 		
 		BoneTrackData.MergedRangeEnd = RangeEnd;
 		BoneTrackData.MergedRangeStart = RangeStart;
 		BoneTrackData.PreProcessedAnimationPayloads.Reset();
-		for (const TTuple< const UInterchangeSceneNode*, TFuture<TOptional<UE::Interchange::FAnimationPayloadData>>>& AnimationPayload : AnimationPayloads)
+		for (UE::Interchange::FAnimationPayloadData& AnimationPayload : AnimationPayloads)
 		{
-			const FName BoneName = FName(*(AnimationPayload.Key->GetDisplayLabel()));
-			const int32 BoneIndex = Skeleton->GetReferenceSkeleton().FindBoneIndex(BoneName);
-			if (BoneIndex == INDEX_NONE)
+			const UInterchangeSceneNode* SkeletonSceneNode = Cast<UInterchangeSceneNode>(Arguments.NodeContainer->GetNode(AnimationPayload.SceneNodeUniqueID));
+			if (!SkeletonSceneNode)
 			{
-				//Skip this bone, we did not found it in the skeleton
 				continue;
 			}
 
-			TOptional<UE::Interchange::FAnimationPayloadData> OptionalAnimationTransformPayload = AnimationPayload.Value.Get();
-			if (!OptionalAnimationTransformPayload.IsSet())
+			const FName BoneName = FName(*(SkeletonSceneNode->GetDisplayLabel()));
+			const int32 BoneIndex = Skeleton->GetReferenceSkeleton().FindBoneIndex(BoneName);
+			if (BoneIndex == INDEX_NONE)
 			{
-				FString PayloadKey = PayloadKeys[AnimationPayload.Key->GetUniqueID()].UniqueId;
-				UE_LOG(LogInterchangeImport, Warning, TEXT("Invalid animation transform payload key [%s] in AnimSequence asset %s."), *PayloadKey, *Arguments.AssetName);
+				//Skip this bone, we did not found it in the skeleton
 				continue;
 			}
 
-			UE::Interchange::FAnimationPayloadData& AnimationTransformPayload = OptionalAnimationTransformPayload.GetValue();
+			UE::Interchange::FAnimationPayloadData& AnimationTransformPayload = AnimationPayload;
 
 			if (AnimationTransformPayload.Type != EInterchangeAnimationPayLoadType::BAKED)
 			{
 				//Where Curve is null the LocalTransform should be used for the Baked Transform generation.
 				FTransform LocalTransform;
-				AnimationPayload.Key->GetCustomLocalTransform(LocalTransform);
+				SkeletonSceneNode->GetCustomLocalTransform(LocalTransform);
 
 				//Currently only Curve -> Baked conversion (for LevelSequence->AnimSequence conversion by ForceMeshType Skeletal use case)
 				//and Curve -> Step Curve conversion (for custom attributes)
@@ -1152,7 +1138,7 @@ UInterchangeFactoryBase::FImportAssetResult UInterchangeAnimSequenceFactory::Imp
 				}
 			}
 
-			BoneTrackData.PreProcessedAnimationPayloads.Add(AnimationPayload.Key, AnimationTransformPayload);
+			BoneTrackData.PreProcessedAnimationPayloads.Add(SkeletonSceneNode, AnimationTransformPayload);
 		}
 	}
 
@@ -1161,7 +1147,8 @@ UInterchangeFactoryBase::FImportAssetResult UInterchangeAnimSequenceFactory::Imp
 		TMap<FString, FInterchangeAnimationPayLoadKey> MorphTargetNodeAnimationPayloads;
 		AnimSequenceFactoryNode->GetMorphTargetNodeAnimationPayloadKeys(MorphTargetNodeAnimationPayloads);
 
-		TMap<FString, TFuture<TOptional<UE::Interchange::FAnimationPayloadData>>> AnimationCurvesPayloads;
+		TArray<UE::Interchange::FAnimationPayloadQuery> PayloadQueries;
+		TMap<FString, TFuture<TOptional<UE::Interchange::FAnimationPayloadData>>> MorphTargetCurveWeightInstanceAnimationPayloads;
 
 		for (const TPair<FString, FInterchangeAnimationPayLoadKey>& MorphTargetNodeUidAnimationPayload : MorphTargetNodeAnimationPayloads)
 		{
@@ -1174,10 +1161,10 @@ UInterchangeFactoryBase::FImportAssetResult UInterchangeAnimSequenceFactory::Imp
 			{
 				if (MorphTargetNodeUidAnimationPayload.Value.Type == EInterchangeAnimationPayLoadType::MORPHTARGETCURVEWEIGHTINSTANCE)
 				{
-					AnimationCurvesPayloads.Add(PayloadKey, Async(EAsyncExecution::TaskGraph, [&MorphTargetNodeUidAnimationPayload]
+					MorphTargetCurveWeightInstanceAnimationPayloads.Add(PayloadKey, Async(EAsyncExecution::TaskGraph, [&MorphTargetNodeUidAnimationPayload]
 						{
 							TOptional<UE::Interchange::FAnimationPayloadData> Result;
-							UE::Interchange::FAnimationPayloadData AnimationPayLoadData(MorphTargetNodeUidAnimationPayload.Value.Type);
+							UE::Interchange::FAnimationPayloadData AnimationPayLoadData(MorphTargetNodeUidAnimationPayload.Key, MorphTargetNodeUidAnimationPayload.Value);
 
 							TArray<FString> PayLoadKeys;
 							MorphTargetNodeUidAnimationPayload.Value.UniqueId.ParseIntoArray(PayLoadKeys, TEXT(":"));
@@ -1201,13 +1188,21 @@ UInterchangeFactoryBase::FImportAssetResult UInterchangeAnimSequenceFactory::Imp
 				}
 				else
 				{
-					AnimationCurvesPayloads.Add(PayloadKey, AnimSequenceTranslatorPayloadInterface->GetAnimationPayloadData(MorphTargetNodeUidAnimationPayload.Value));
+					PayloadQueries.Add(UE::Interchange::FAnimationPayloadQuery(MorphTargetNodeUidAnimationPayload.Key, MorphTargetNodeUidAnimationPayload.Value));
 				}
 				MorphTargetData.CurveNodeNamePerPayloadKey.Add(PayloadKey, MorphTargetNode->GetDisplayLabel());
 			}
 		}
 
-		for (TPair<FString, TFuture<TOptional<UE::Interchange::FAnimationPayloadData>>>& CurveNameAndPayload : AnimationCurvesPayloads)
+		TArray<UE::Interchange::FAnimationPayloadData> AnimationCurvesPayloads = AnimSequenceTranslatorPayloadInterface->GetAnimationPayloadData(PayloadQueries);
+		for (UE::Interchange::FAnimationPayloadData& AnimationCurvePayload : AnimationCurvesPayloads)
+		{
+			FString PayloadKey = AnimationCurvePayload.PayloadKey.UniqueId; //in this instance its the payload's uniqueid
+			MorphTargetData.CurvesPayloads.Add(PayloadKey, MoveTemp(AnimationCurvePayload));
+		}
+		AnimationCurvesPayloads.Empty();
+
+		for (TPair<FString, TFuture<TOptional<UE::Interchange::FAnimationPayloadData>>>& CurveNameAndPayload : MorphTargetCurveWeightInstanceAnimationPayloads)
 		{
 			TOptional<UE::Interchange::FAnimationPayloadData> AnimationCurvePayload = CurveNameAndPayload.Value.Get();
 			if (!AnimationCurvePayload.IsSet())
@@ -1218,7 +1213,7 @@ UInterchangeFactoryBase::FImportAssetResult UInterchangeAnimSequenceFactory::Imp
 			UE::Interchange::FAnimationPayloadData& CurvePayload = AnimationCurvePayload.GetValue();
 			MorphTargetData.CurvesPayloads.Add(CurveNameAndPayload.Key, MoveTemp(CurvePayload));
 		}
-		AnimationCurvesPayloads.Empty();
+		MorphTargetCurveWeightInstanceAnimationPayloads.Empty();
 	}
 
 	//Getting the file Hash will cache it into the source data
diff --git a/Engine/Plugins/Interchange/Runtime/Source/Import/Private/Animation/InterchangeLevelSequenceFactory.cpp b/Engine/Plugins/Interchange/Runtime/Source/Import/Private/Animation/InterchangeLevelSequenceFactory.cpp
--- a/Engine/Plugins/Interchange/Runtime/Source/Import/Private/Animation/InterchangeLevelSequenceFactory.cpp
+++ b/Engine/Plugins/Interchange/Runtime/Source/Import/Private/Animation/InterchangeLevelSequenceFactory.cpp
@@ -243,7 +243,9 @@ namespace UE::Interchange::Private
 	{
 		// Get targeted actor exists
 		AActor* Actor = GetActor(TransformTrackNode);
-		if (!Actor)
+		FString SceneNodeUid;
+
+		if (!Actor || !TransformTrackNode.GetCustomActorDependencyUid(SceneNodeUid))
 		{
 			UE_LOG(LogInterchangeImport, Warning, TEXT("Cannot find actor for animation track %s"), *TransformTrackNode.GetDisplayLabel());
 			return;
@@ -257,14 +259,16 @@ namespace UE::Interchange::Private
 			return;
 		}
 
-		TFuture<TOptional<UE::Interchange::FAnimationPayloadData>> Result = PayloadInterface.GetAnimationPayloadData(PayloadKey);
-		const TOptional<UE::Interchange::FAnimationPayloadData>& PayloadData = Result.Get();
-		if (!PayloadData.IsSet() || PayloadData->Curves.Num() != 9)
+		TArray<UE::Interchange::FAnimationPayloadData> PayloadDataArray = PayloadInterface.GetAnimationPayloadData({ FAnimationPayloadQuery(SceneNodeUid, PayloadKey) });
+		
+		if (PayloadDataArray.Num() != 1 || PayloadDataArray[0].Curves.Num() != 9)
 		{
 			UE_LOG(LogInterchangeImport, Warning, TEXT("No payload for animation track %s on actor %s"), *TransformTrackNode.GetDisplayLabel(), *Actor->GetActorLabel());
 			return;
 		}
 
+		UE::Interchange::FAnimationPayloadData* PayloadData = &PayloadDataArray[0];
+
 		FGuid ObjectBinding = BindActorToLevelSequence(Actor);
 
 		UMovieScene3DTransformTrack* TransformTrack = MovieScene->FindTrack<UMovieScene3DTransformTrack>(ObjectBinding);
@@ -446,7 +450,9 @@ namespace UE::Interchange::Private
 
 		// Get targeted actor exists
 		AActor* Actor = GetActor(AnimationTrackNode);
-		if (!Actor)
+		FString SceneNodeUid;
+
+		if (!Actor || !AnimationTrackNode.GetCustomActorDependencyUid(SceneNodeUid))
 		{
 			UE_LOG(LogInterchangeImport, Warning, TEXT("Cannot find actor for animation track %s"), *AnimationTrackNode.GetDisplayLabel());
 			return;
@@ -460,14 +466,16 @@ namespace UE::Interchange::Private
 			return;
 		}
 
-		TFuture<TOptional<UE::Interchange::FAnimationPayloadData>> Result = PayloadInterface.GetAnimationPayloadData(PayloadKey);
-		const TOptional<UE::Interchange::FAnimationPayloadData>& PayloadData = Result.Get();
-		if (!PayloadData.IsSet() || (PayloadData->StepCurves.IsEmpty() && PayloadData->Curves.IsEmpty()))
+		TArray<UE::Interchange::FAnimationPayloadData> PayloadDataArray = PayloadInterface.GetAnimationPayloadData({ FAnimationPayloadQuery(SceneNodeUid, PayloadKey) });
+		
+		if (PayloadDataArray.Num() != 1 || (PayloadDataArray[0].StepCurves.IsEmpty() && PayloadDataArray[0].Curves.IsEmpty()))
 		{
 			UE_LOG(LogInterchangeImport, Warning, TEXT("No payload for animation track %s on actor %s"), *AnimationTrackNode.GetDisplayLabel(), *Actor->GetActorLabel());
 			return;
 		}
 
+		UE::Interchange::FAnimationPayloadData* PayloadData = &PayloadDataArray[0];
+
 		FGuid ObjectBinding;
 		if(PropertyTrack == EInterchangePropertyTracks::Visibility)
 		{
diff --git a/Engine/Plugins/Interchange/Runtime/Source/Import/Private/Fbx/InterchangeFbxTranslator.cpp b/Engine/Plugins/Interchange/Runtime/Source/Import/Private/Fbx/InterchangeFbxTranslator.cpp
--- a/Engine/Plugins/Interchange/Runtime/Source/Import/Private/Fbx/InterchangeFbxTranslator.cpp
+++ b/Engine/Plugins/Interchange/Runtime/Source/Import/Private/Fbx/InterchangeFbxTranslator.cpp
@@ -459,39 +459,44 @@ TFuture<TOptional<UE::Interchange::FMeshPayloadData>> UInterchangeFbxTranslator:
 	return Promise->GetFuture();
 }
 
-TFuture<TOptional<UE::Interchange::FAnimationPayloadData>> UInterchangeFbxTranslator::GetAnimationPayloadData(const FInterchangeAnimationPayLoadKey& PayLoadKey, const double BakeFrequency, const double RangeStartSecond, const double RangeStopSecond) const
+TArray<UE::Interchange::FAnimationPayloadData> UInterchangeFbxTranslator::GetAnimationPayloadData(const TArray<UE::Interchange::FAnimationPayloadQuery>& PayloadQueries) const
 {
-	TSharedPtr<TPromise<TOptional<UE::Interchange::FAnimationPayloadData>>> Promise = MakeShared<TPromise<TOptional<UE::Interchange::FAnimationPayloadData>>>();
+	TArray<UE::Interchange::FAnimationPayloadData> PayloadDataResults;
+	
+	if (PayloadQueries.Num() == 0)
+	{
+		return PayloadDataResults;
+	}
 
-	auto OnPayloadReady = [Promise, PayLoadKey](const FString& AnimationPayloadFilename)
+	bool bBakedQueries = PayloadQueries[0].PayloadKey.Type == EInterchangeAnimationPayLoadType::BAKED;
+
+	auto OnPayloadReady = [](const UE::Interchange::FAnimationPayloadQuery& PayloadQuery, const FString& AnimationPayloadFile)
 		{
-			if (!ensure(FPaths::FileExists(AnimationPayloadFilename)))
+			if (!ensure(FPaths::FileExists(AnimationPayloadFile)))
 			{
 				// TODO log an error saying the payload file does not exist even if the get payload command succeeded
-				Promise->SetValue(TOptional<UE::Interchange::FAnimationPayloadData>());
-				return;
+				return TOptional<UE::Interchange::FAnimationPayloadData>();
 			}
 
-			UE::Interchange::FAnimationPayloadData AnimationTransformPayload(PayLoadKey.Type);
-
 			// All sub object should be gone with the reset
 			TArray64<uint8> Buffer;
-			FFileHelper::LoadFileToArray(Buffer, *AnimationPayloadFilename);
+			FFileHelper::LoadFileToArray(Buffer, *AnimationPayloadFile);
 			uint8* FileData = Buffer.GetData();
 			int64 FileDataSize = Buffer.Num();
 			if (FileDataSize < 1)
 			{
 				// Nothing to load from this file
-				Promise->SetValue(TOptional<UE::Interchange::FAnimationPayloadData>());
-				return;
+				return TOptional<UE::Interchange::FAnimationPayloadData>();
 			}
 
+			UE::Interchange::FAnimationPayloadData AnimationTransformPayload(PayloadQuery.SceneNodeUniqueID, PayloadQuery.PayloadKey);
+
 			// Buffer keeps the ownership of the data, the large memory reader is use to serialize the TMap
 			FLargeMemoryReader Ar(FileData, FileDataSize);
 
-			switch (PayLoadKey.Type)
+			switch (PayloadQuery.PayloadKey.Type)
 			{
-			case EInterchangeAnimationPayLoadType::CURVE:
+				case EInterchangeAnimationPayLoadType::CURVE:
 				{
 					TArray<FInterchangeCurve> InterchangeCurves;
 					Ar << InterchangeCurves;
@@ -503,7 +508,7 @@ TFuture<TOptional<UE::Interchange::FAnimationPayloadData>> UInterchangeFbxTransl
 					}
 				}
 				break;
-			case EInterchangeAnimationPayLoadType::MORPHTARGETCURVE:
+				case EInterchangeAnimationPayLoadType::MORPHTARGETCURVE:
 				{
 					TArray<FInterchangeCurve> InterchangeCurves;
 					Ar << InterchangeCurves;
@@ -517,94 +522,183 @@ TFuture<TOptional<UE::Interchange::FAnimationPayloadData>> UInterchangeFbxTransl
 					}
 				}
 				break;
-			case EInterchangeAnimationPayLoadType::STEPCURVE:
+				case EInterchangeAnimationPayLoadType::STEPCURVE:
 				{
 					Ar << AnimationTransformPayload.StepCurves;
 				}
 				break;
-			case EInterchangeAnimationPayLoadType::BAKED:
-				AnimationTransformPayload.SerializeBaked(Ar);
-				break;
-			case EInterchangeAnimationPayLoadType::NONE:
-			default:
-				break;
+				case EInterchangeAnimationPayLoadType::BAKED:
+					AnimationTransformPayload.SerializeBaked(Ar);
+					break;
+				case EInterchangeAnimationPayLoadType::NONE:
+				default:
+					break;
 			}
 
-			Promise->SetValue(MoveTemp(AnimationTransformPayload));
+			return TOptional<UE::Interchange::FAnimationPayloadData>(AnimationTransformPayload);
 		};
 
 	if (!bUseWorkerImport)
 	{
 #if WITH_EDITOR
-		FString AnimationPayloadFile;
-		if (PayLoadKey.Type == EInterchangeAnimationPayLoadType::BAKED)
+		//PayloadQueries are assumed to be of the same PayloadKey.Type.
+		if (bBakedQueries)
 		{
-			FString ResultPayloadsUniqueId = FbxParser.FetchAnimationBakeTransformPayload(PayLoadKey.UniqueId
-				, BakeFrequency
-				, RangeStartSecond
-				, RangeStopSecond
-				, ResultFolder);
-			AnimationPayloadFile = FbxParser.GetResultPayloadFilepath(ResultPayloadsUniqueId);
+			//Fetch the animation data for all the queries at the same time:
+			FbxParser.FetchAnimationBakeTransformPayloads(PayloadQueries, ResultFolder);
+
+			for (const UE::Interchange::FAnimationPayloadQuery& PayloadQuery : PayloadQueries)
+			{
+				TOptional<UE::Interchange::FAnimationPayloadData> OptionalPayloadData = OnPayloadReady(PayloadQuery, FbxParser.GetResultPayloadFilepath(PayloadQuery.GetHashString()));
+				if (OptionalPayloadData.IsSet())
+				{
+					PayloadDataResults.Add(OptionalPayloadData.GetValue());
+				}
+			}
 		}
 		else
 		{
-			FbxParser.FetchPayload(PayLoadKey.UniqueId, ResultFolder);
-			AnimationPayloadFile = FbxParser.GetResultPayloadFilepath(PayLoadKey.UniqueId);
+			for (const UE::Interchange::FAnimationPayloadQuery& PayloadQuery : PayloadQueries)
+			{
+				FbxParser.FetchPayload(PayloadQuery.PayloadKey.UniqueId, ResultFolder);
+
+				TOptional<UE::Interchange::FAnimationPayloadData> OptionalPayloadData = OnPayloadReady(PayloadQuery, FbxParser.GetResultPayloadFilepath(PayloadQuery.PayloadKey.UniqueId));
+				if (OptionalPayloadData.IsSet())
+				{
+					PayloadDataResults.Add(OptionalPayloadData.GetValue());
+				}
+			}
 		}
-		OnPayloadReady(AnimationPayloadFile);
 #endif //WITH_EDITOR
 	}
 	else
 	{
 		if (!Dispatcher.IsValid())
 		{
-			Promise->SetValue(TOptional<UE::Interchange::FAnimationPayloadData>());
-			return Promise->GetFuture();
+			return PayloadDataResults;
 		}
 
-		// Create a json command to read the fbx file
-		FString JsonCommand = 
-			(PayLoadKey.Type == EInterchangeAnimationPayLoadType::BAKED)
-			? CreateFetchAnimationBakeTransformPayloadFbxCommand(PayLoadKey.UniqueId, BakeFrequency, RangeStartSecond, RangeStopSecond)
-			: CreateFetchPayloadFbxCommand(PayLoadKey.UniqueId);
+		//Init Promises and Futures
+		//Promises so Dispatcher can set the Values
+		//Futures so we can run the Asynch tasks and then acquire the Values.
+		TArray<TFuture<TOptional<UE::Interchange::FAnimationPayloadData>>> AnimationPayloadDataFutures;
+		TArray<TSharedPtr<TPromise<TOptional<UE::Interchange::FAnimationPayloadData>>>> AnimationPayloadDataPromises;
+		AnimationPayloadDataFutures.Reserve(PayloadQueries.Num());
+		for (size_t PayloadQueryIndex = 0; PayloadQueryIndex < PayloadQueries.Num(); PayloadQueryIndex++)
+		{
+			TSharedPtr<TPromise<TOptional<UE::Interchange::FAnimationPayloadData>>> Promise = MakeShared<TPromise<TOptional<UE::Interchange::FAnimationPayloadData>>>();
+			AnimationPayloadDataFutures.Add(Promise->GetFuture());
+			AnimationPayloadDataPromises.Add(Promise);
+		}
 
-		const int32 CreatedTaskIndex = Dispatcher->AddTask(JsonCommand, FInterchangeDispatcherTaskCompleted::CreateLambda([this, Promise, PayLoadKey, OnPayloadReadyClosure = MoveTemp(OnPayloadReady)](const int32 TaskIndex)
-			{
-				TRACE_CPUPROFILER_EVENT_SCOPE(UInterchangeFbxTranslator::GetAnimationCurvePayloadData::Dispatcher->AddTaskDone)
-				UE::Interchange::ETaskState TaskState;
-				FString JsonResult;
-				TArray<FString> JsonMessages;
-				Dispatcher->GetTaskState(TaskIndex, TaskState, JsonResult, JsonMessages);
+		if (bBakedQueries)
+		{
+			FString JsonCommand = CreateFetchAnimationBakeTransformPayloadFbxCommand(PayloadQueries);
 
-				// Parse the Json messages into UInterchangeResults
-				for (const FString& JsonMessage : JsonMessages)
+			const int32 CreatedTaskIndex = Dispatcher->AddTask(JsonCommand, FInterchangeDispatcherTaskCompleted::CreateLambda([this, PayloadQueries, OnPayloadReadyClosure = MoveTemp(OnPayloadReady), &AnimationPayloadDataPromises](const int32 TaskIndex)
 				{
-					UE::Interchange::Private::ApplyTranslatorMessage(this, JsonMessage);
-				}
+					TRACE_CPUPROFILER_EVENT_SCOPE(UInterchangeFbxTranslator::GetAnimationCurvePayloadData::Dispatcher->AddTaskDone)
+					
+					UE::Interchange::ETaskState TaskState;
+					FString JsonResult;
+					TArray<FString> JsonMessages;
+					Dispatcher->GetTaskState(TaskIndex, TaskState, JsonResult, JsonMessages);
 
-				if (TaskState != UE::Interchange::ETaskState::ProcessOk)
-				{
-					Promise->SetValue(TOptional<UE::Interchange::FAnimationPayloadData>());
-					return;
-				}
+					// Parse the Json messages into UInterchangeResults
+					for (const FString& JsonMessage : JsonMessages)
+					{
+						UE::Interchange::Private::ApplyTranslatorMessage(this, JsonMessage);
+					}
 
-				// Grab the result file and fill the BaseNodeContainer
-				UE::Interchange::FJsonFetchPayloadCmd::JsonResultParser ResultParser;
-				ResultParser.FromJson(JsonResult);
-				FString AnimationTransformPayloadFilename = ResultParser.GetResultFilename();
-				OnPayloadReadyClosure(AnimationTransformPayloadFilename);
-			}));
+					if (TaskState != UE::Interchange::ETaskState::ProcessOk)
+					{
+						//Process failed, set Promises up with OptionalValue.
+						for (size_t PayloadQueryIndex = 0; PayloadQueryIndex < PayloadQueries.Num(); PayloadQueryIndex++)
+						{
+							TSharedPtr<TPromise<TOptional<UE::Interchange::FAnimationPayloadData>>>& Promise = AnimationPayloadDataPromises[PayloadQueryIndex];
+							Promise->SetValue(TOptional<UE::Interchange::FAnimationPayloadData>());
+						}
+						return;
+					}
 
-		// The task was not added to the dispatcher
-		if (CreatedTaskIndex == INDEX_NONE)
+					// Grab the result file and fill the BaseNodeContainer
+					UE::Interchange::FJsonFetchAnimationQueriesCmd::JsonAnimationQueriesResultParser ResultParser;
+					ResultParser.FromJson(JsonResult);
+					const TMap<FString, FString>& HashToFilenames = ResultParser.GetHashToFilenames();
+
+					for (size_t PayloadQueryIndex = 0; PayloadQueryIndex < PayloadQueries.Num(); PayloadQueryIndex++)
+					{
+						TSharedPtr<TPromise<TOptional<UE::Interchange::FAnimationPayloadData>>>& Promise = AnimationPayloadDataPromises[PayloadQueryIndex];
+
+						const UE::Interchange::FAnimationPayloadQuery& PayloadQuery = PayloadQueries[PayloadQueryIndex];
+						if (HashToFilenames.Contains(PayloadQuery.GetHashString()))
+						{
+							TOptional<UE::Interchange::FAnimationPayloadData> OptionalPayloadData = OnPayloadReadyClosure(PayloadQuery, HashToFilenames[PayloadQuery.GetHashString()]);
+
+							Promise->SetValue(OptionalPayloadData);
+						}
+						else 
+						{
+							Promise->SetValue(TOptional<UE::Interchange::FAnimationPayloadData>());
+						}
+					}
+				}));
+		}
+		else
 		{
-			Promise->SetValue(TOptional<UE::Interchange::FAnimationPayloadData>{});
+			for (size_t PayloadQueryIndex = 0; PayloadQueryIndex < PayloadQueries.Num(); PayloadQueryIndex++)
+			{
+				const UE::Interchange::FAnimationPayloadQuery& PayloadQuery = PayloadQueries[PayloadQueryIndex];
+				TSharedPtr<TPromise<TOptional<UE::Interchange::FAnimationPayloadData>>>& Promise = AnimationPayloadDataPromises[PayloadQueryIndex];
+
+				FString JsonCommand = CreateFetchPayloadFbxCommand(PayloadQuery.PayloadKey.UniqueId);
+
+				const int32 CreatedTaskIndex = Dispatcher->AddTask(JsonCommand, FInterchangeDispatcherTaskCompleted::CreateLambda([this, PayloadQuery, OnPayloadReadyClosure = MoveTemp(OnPayloadReady), &Promise](const int32 TaskIndex)
+					{
+						TRACE_CPUPROFILER_EVENT_SCOPE(UInterchangeFbxTranslator::GetAnimationCurvePayloadData::Dispatcher->AddTaskDone)
+							UE::Interchange::ETaskState TaskState;
+						FString JsonResult;
+						TArray<FString> JsonMessages;
+						Dispatcher->GetTaskState(TaskIndex, TaskState, JsonResult, JsonMessages);
+
+						// Parse the Json messages into UInterchangeResults
+						for (const FString& JsonMessage : JsonMessages)
+						{
+							UE::Interchange::Private::ApplyTranslatorMessage(this, JsonMessage);
+						}
+
+						if (TaskState != UE::Interchange::ETaskState::ProcessOk)
+						{
+							Promise->SetValue(TOptional<UE::Interchange::FAnimationPayloadData>());
+							return;
+						}
+
+						// Grab the result file and fill the BaseNodeContainer
+						UE::Interchange::FJsonFetchPayloadCmd::JsonResultParser ResultParser;
+						ResultParser.FromJson(JsonResult);
+						FString AnimationTransformPayloadFilename = ResultParser.GetResultFilename();
+						TOptional<UE::Interchange::FAnimationPayloadData> OptionalPayloadData = OnPayloadReadyClosure(PayloadQuery, AnimationTransformPayloadFilename);
+						Promise->SetValue(OptionalPayloadData);
+					}));
+			}
 		}
-	}
 
-	return Promise->GetFuture();
+		for (TFuture<TOptional<UE::Interchange::FAnimationPayloadData>>& AnimationPayloadFuture : AnimationPayloadDataFutures)
+		{
+			TOptional<UE::Interchange::FAnimationPayloadData> OptionalPayloadData = AnimationPayloadFuture.Get();
+			if (!OptionalPayloadData.IsSet())
+			{
+				continue;
+			}
+			PayloadDataResults.Add(OptionalPayloadData.GetValue());
+		}
+
+	}
+	
+	return PayloadDataResults;
 }
 
+
 FString UInterchangeFbxTranslator::CreateLoadFbxFileCommand(const FString& FbxFilePath, const bool bConvertScene, const bool bForceFrontXAxis, const bool bConvertSceneUnit, const bool bKeepFbxNamespace) const
 {
 	UE::Interchange::FJsonLoadSourceCmd LoadSourceCommand(TEXT("FBX"), FbxFilePath, bConvertScene, bForceFrontXAxis, bConvertSceneUnit, bKeepFbxNamespace);
@@ -623,9 +717,9 @@ FString UInterchangeFbxTranslator::CreateFetchMeshPayloadFbxCommand(const FStrin
 	return PayloadCommand.ToJson();
 }
 
-FString UInterchangeFbxTranslator::CreateFetchAnimationBakeTransformPayloadFbxCommand(const FString& FbxPayloadKey, const double BakeFrequency, const double RangeStartTime, const double RangeEndTime) const
+FString UInterchangeFbxTranslator::CreateFetchAnimationBakeTransformPayloadFbxCommand(const TArray<UE::Interchange::FAnimationPayloadQuery>& PayloadQueries) const
 {
-	UE::Interchange::FJsonFetchAnimationBakeTransformPayloadCmd PayloadCommand(TEXT("FBX"), FbxPayloadKey, BakeFrequency, RangeStartTime, RangeEndTime);
+	UE::Interchange::FJsonFetchAnimationQueriesCmd PayloadCommand(TEXT("FBX"), UE::Interchange::FAnimationPayloadQuery::ToJson(PayloadQueries));
 	return PayloadCommand.ToJson();
 }
 
diff --git a/Engine/Plugins/Interchange/Runtime/Source/Import/Private/Gltf/InterchangeGltfTranslator.cpp b/Engine/Plugins/Interchange/Runtime/Source/Import/Private/Gltf/InterchangeGltfTranslator.cpp
--- a/Engine/Plugins/Interchange/Runtime/Source/Import/Private/Gltf/InterchangeGltfTranslator.cpp
+++ b/Engine/Plugins/Interchange/Runtime/Source/Import/Private/Gltf/InterchangeGltfTranslator.cpp
@@ -903,33 +903,33 @@ TOptional< UE::Interchange::FImportImage > UInterchangeGLTFTranslator::GetTextur
 	return TexturePayloadData;
 }
 
-TFuture<TOptional<UE::Interchange::FAnimationPayloadData>> UInterchangeGLTFTranslator::GetAnimationPayloadData(const FInterchangeAnimationPayLoadKey& PayLoadKey, const double BakeFrequency, const double RangeStartSecond, const double RangeStopSecond) const
+TFuture<TOptional<UE::Interchange::FAnimationPayloadData>> UInterchangeGLTFTranslator::GetAnimationPayloadData(const UE::Interchange::FAnimationPayloadQuery& PayloadQuery) const
 {
-	return Async(EAsyncExecution::TaskGraph, [this, PayLoadKey, BakeFrequency, RangeStartSecond, RangeStopSecond]
+	return Async(EAsyncExecution::TaskGraph, [this, PayloadQuery]
 		{
 
 			TOptional<UE::Interchange::FAnimationPayloadData> Result;
-			UE::Interchange::FAnimationPayloadData AnimationPayLoadData(PayLoadKey.Type);
+			UE::Interchange::FAnimationPayloadData AnimationPayLoadData(PayloadQuery.SceneNodeUniqueID, PayloadQuery.PayloadKey);
 
-			switch (PayLoadKey.Type)
+			switch (PayloadQuery.PayloadKey.Type)
 			{
 			case EInterchangeAnimationPayLoadType::CURVE:
-				if (UE::Interchange::Gltf::Private::GetTransformAnimationPayloadData(PayLoadKey.UniqueId, GltfAsset, AnimationPayLoadData))
+				if (UE::Interchange::Gltf::Private::GetTransformAnimationPayloadData(PayloadQuery.PayloadKey.UniqueId, GltfAsset, AnimationPayLoadData))
 				{
 					Result.Emplace(AnimationPayLoadData);
 				}
 				break;
 			case EInterchangeAnimationPayLoadType::MORPHTARGETCURVE:
-				if (UE::Interchange::Gltf::Private::GetMorphTargetAnimationPayloadData(PayLoadKey.UniqueId, GltfAsset, AnimationPayLoadData))
+				if (UE::Interchange::Gltf::Private::GetMorphTargetAnimationPayloadData(PayloadQuery.PayloadKey.UniqueId, GltfAsset, AnimationPayLoadData))
 				{
 					Result.Emplace(AnimationPayLoadData);
 				}
 				break;
 			case EInterchangeAnimationPayLoadType::BAKED:
-				AnimationPayLoadData.BakeFrequency = BakeFrequency;
-				AnimationPayLoadData.RangeStartTime = RangeStartSecond;
-				AnimationPayLoadData.RangeEndTime = RangeStopSecond;
-				if (UE::Interchange::Gltf::Private::GetBakedAnimationTransformPayloadData(PayLoadKey.UniqueId, GltfAsset, AnimationPayLoadData))
+				AnimationPayLoadData.BakeFrequency = PayloadQuery.TimeDescription.BakeFrequency;
+				AnimationPayLoadData.RangeStartTime = PayloadQuery.TimeDescription.RangeStartSecond;
+				AnimationPayLoadData.RangeEndTime = PayloadQuery.TimeDescription.RangeStopSecond;
+				if (UE::Interchange::Gltf::Private::GetBakedAnimationTransformPayloadData(PayloadQuery.PayloadKey.UniqueId, GltfAsset, AnimationPayLoadData))
 				{
 					Result.Emplace(AnimationPayLoadData);
 				}
@@ -945,6 +945,28 @@ TFuture<TOptional<UE::Interchange::FAnimationPayloadData>> UInterchangeGLTFTrans
 	);
 }
 
+TArray<UE::Interchange::FAnimationPayloadData> UInterchangeGLTFTranslator::GetAnimationPayloadData(const TArray<UE::Interchange::FAnimationPayloadQuery>& PayloadQueries) const
+{
+	TArray<TFuture<TOptional<UE::Interchange::FAnimationPayloadData>>> AnimationPayloadFutures;
+	for (const UE::Interchange::FAnimationPayloadQuery& PayloadQuery: PayloadQueries)
+	{
+		AnimationPayloadFutures.Add(GetAnimationPayloadData(PayloadQuery));
+	}
+
+	TArray<UE::Interchange::FAnimationPayloadData> AnimationPayloads;
+	for (TFuture<TOptional<UE::Interchange::FAnimationPayloadData>>& AnimationPayloadFuture : AnimationPayloadFutures)
+	{
+		TOptional<UE::Interchange::FAnimationPayloadData> OptionalPayloadData = AnimationPayloadFuture.Get();
+		if (!OptionalPayloadData.IsSet())
+		{
+			continue;
+		}
+		AnimationPayloads.Add(OptionalPayloadData.GetValue());
+	}
+
+	return AnimationPayloads;
+}
+
 void UInterchangeGLTFTranslator::SetTextureSRGB(UInterchangeBaseNodeContainer& NodeContainer, const GLTF::FTextureMap& TextureMap, bool bSRGB) const
 {
 	if (GltfAsset.Textures.IsValidIndex(TextureMap.TextureIndex))
diff --git a/Engine/Plugins/Interchange/Runtime/Source/Import/Public/Animation/InterchangeAnimationPayloadInterface.h b/Engine/Plugins/Interchange/Runtime/Source/Import/Public/Animation/InterchangeAnimationPayloadInterface.h
--- a/Engine/Plugins/Interchange/Runtime/Source/Import/Public/Animation/InterchangeAnimationPayloadInterface.h
+++ b/Engine/Plugins/Interchange/Runtime/Source/Import/Public/Animation/InterchangeAnimationPayloadInterface.h
@@ -30,12 +30,11 @@ public:
 	 * Get animation payload data for the specified payload key.
 	 * It return an array of FRichCurve (Rich curve are float curve we can interpolate) or array of "Step" curve or an array of Baked Transformations, depending on Type
 	 *
-	 * @param PayLoadKey.UniqueId - The key to retrieve the a particular payload contain into the specified source data.
-	 * @param PayLoadKey.Type - The type of animation data to be retrieved (Type is provided by the Translator upon initialization of PayLoadKeys)
-	 * @return - The resulting PayloadData as a TFuture point by the PayloadKey. The TOptional will not be set if there is an error retrieving the payload.
+	 * @param PayloadQueries - A PayloadQuery contains all the necessary data for a Query to be processed (Including SceneNodeUID, PayloadKey, TimeDescription).
+	 * @return - The resulting PayloadData.
 	 * 
 	 */
-	virtual TFuture<TOptional<UE::Interchange::FAnimationPayloadData>> GetAnimationPayloadData(const FInterchangeAnimationPayLoadKey& PayLoadKey, const double BakeFrequency = 0, const double RangeStartSecond = 0, const double RangeStopSecond = 0) const = 0;
+	virtual TArray<UE::Interchange::FAnimationPayloadData> GetAnimationPayloadData(const TArray<UE::Interchange::FAnimationPayloadQuery>& PayloadQuery) const = 0;
 };
 
 
diff --git a/Engine/Plugins/Interchange/Runtime/Source/Import/Public/Fbx/InterchangeFbxTranslator.h b/Engine/Plugins/Interchange/Runtime/Source/Import/Public/Fbx/InterchangeFbxTranslator.h
--- a/Engine/Plugins/Interchange/Runtime/Source/Import/Public/Fbx/InterchangeFbxTranslator.h
+++ b/Engine/Plugins/Interchange/Runtime/Source/Import/Public/Fbx/InterchangeFbxTranslator.h
@@ -103,7 +103,7 @@ public:
 
 	//////////////////////////////////////////////////////////////////////////
 	/* IInterchangeAnimationPayloadInterface Begin */
-	virtual TFuture<TOptional<UE::Interchange::FAnimationPayloadData>> GetAnimationPayloadData(const FInterchangeAnimationPayLoadKey& PayLoadKey, const double BakeFrequency = 0, const double RangeStartSecond = 0, const double RangeStopSecond = 0) const override;
+	virtual TArray<UE::Interchange::FAnimationPayloadData> GetAnimationPayloadData(const TArray<UE::Interchange::FAnimationPayloadQuery>& PayloadQueries) const override;
 	/* IInterchangeAnimationPayloadInterface End */
 private:
 	FString CreateLoadFbxFileCommand(const FString& FbxFilePath, const bool bConvertScene, const bool bForceFrontXAxis, const bool bConvertSceneUnit, const bool bKeepFbxNamespace) const;
@@ -112,7 +112,7 @@ private:
 
 	FString CreateFetchPayloadFbxCommand(const FString& FbxPayloadKey) const;
 
-	FString CreateFetchAnimationBakeTransformPayloadFbxCommand(const FString& FbxPayloadKey, const double BakeFrequency, const double RangeStartTime, const double RangeEndTime) const;
+	FString CreateFetchAnimationBakeTransformPayloadFbxCommand(const TArray<UE::Interchange::FAnimationPayloadQuery>& PayloadQueries) const;
 	
 	//Dispatcher is mutable since it is create during the Translate operation
 	//We do not want to allocate the dispatcher and start the InterchangeWorker process
diff --git a/Engine/Plugins/Interchange/Runtime/Source/Import/Public/Gltf/InterchangeGltfTranslator.h b/Engine/Plugins/Interchange/Runtime/Source/Import/Public/Gltf/InterchangeGltfTranslator.h
--- a/Engine/Plugins/Interchange/Runtime/Source/Import/Public/Gltf/InterchangeGltfTranslator.h
+++ b/Engine/Plugins/Interchange/Runtime/Source/Import/Public/Gltf/InterchangeGltfTranslator.h
@@ -53,7 +53,8 @@ public:
 	/* IInterchangeTexturePayloadInterface End */
 
 	/* IInterchangeAnimationPayloadInterface Begin */
-	virtual TFuture<TOptional<UE::Interchange::FAnimationPayloadData>> GetAnimationPayloadData(const FInterchangeAnimationPayLoadKey& PayLoadKey, const double BakeFrequency = 0, const double RangeStartSecond = 0, const double RangeStopSecond = 0) const override;
+	TFuture<TOptional<UE::Interchange::FAnimationPayloadData>> GetAnimationPayloadData(const UE::Interchange::FAnimationPayloadQuery& PayloadQuery) const;
+	virtual TArray<UE::Interchange::FAnimationPayloadData> GetAnimationPayloadData(const TArray<UE::Interchange::FAnimationPayloadQuery>& PayloadQueries) const override;
 	/* IInterchangeAnimationPayloadInterface End */
 
 	/* IInterchangeVariantSetPayloadInterface Begin */
diff --git a/Engine/Plugins/Interchange/Runtime/Source/Parsers/CommonParser/InterchangeCommonParser.Build.cs b/Engine/Plugins/Interchange/Runtime/Source/Parsers/CommonParser/InterchangeCommonParser.Build.cs
--- a/Engine/Plugins/Interchange/Runtime/Source/Parsers/CommonParser/InterchangeCommonParser.Build.cs
+++ b/Engine/Plugins/Interchange/Runtime/Source/Parsers/CommonParser/InterchangeCommonParser.Build.cs
@@ -23,6 +23,7 @@ namespace UnrealBuildTool.Rules
 				PrivateDependencyModuleNames.AddRange(
 					new string[] {
 						"Engine",
+						"Json",
 					}
 				);
 			}
diff --git a/Engine/Plugins/Interchange/Runtime/Source/Parsers/CommonParser/Private/InterchangeCommonAnimationPayload.cpp b/Engine/Plugins/Interchange/Runtime/Source/Parsers/CommonParser/Private/InterchangeCommonAnimationPayload.cpp
--- a/Engine/Plugins/Interchange/Runtime/Source/Parsers/CommonParser/Private/InterchangeCommonAnimationPayload.cpp
+++ b/Engine/Plugins/Interchange/Runtime/Source/Parsers/CommonParser/Private/InterchangeCommonAnimationPayload.cpp
@@ -4,10 +4,37 @@
 
 #include "CoreMinimal.h"
 
+#include "InterchangeHelper.h"
+
+#include "Dom/JsonValue.h"
+#include "Dom/JsonObject.h"
+#include "Serialization/JsonReader.h"
+#include "Serialization/JsonWriter.h"
+#include "Serialization/JsonSerializer.h"
+
 #include UE_INLINE_GENERATED_CPP_BY_NAME(InterchangeCommonAnimationPayload)
 
 namespace UE::Interchange
 {
+	FString Private::HashString(const FString& String)
+	{
+		TArray<uint8> TempBytes;
+		TempBytes.Reserve(64);
+		//The archive is flagged as persistent so that machines of different endianness produce identical binary results.
+		FMemoryWriter Ar(TempBytes, /*bIsPersistent=*/ true);
+		//Hack because serialization do not support const
+		FString& NonConst = *const_cast<FString*>(&String);
+		Ar << NonConst;
+		FSHA1 Sha;
+		Sha.Update(TempBytes.GetData(), TempBytes.Num() * TempBytes.GetTypeSize());
+		Sha.Final();
+		// Retrieve the hash and use it to construct a pseudo-GUID.
+		uint32 Hash[5];
+		Sha.GetHash((uint8*)Hash);
+		FGuid Guid = FGuid(Hash[0] ^ Hash[4], Hash[1], Hash[2], Hash[3]);
+		return Guid.ToString(EGuidFormats::Base36Encoded);
+	}
+
 	void FAnimationPayloadData::SerializeBaked(FArchive& Ar)
 	{
 		Ar << BakeFrequency;
@@ -114,6 +141,108 @@ namespace UE::Interchange
 		}
 #endif
 	}
+
+	FString FAnimationPayloadQuery::GetHashString() const
+	{
+		if (!HashStringCache.IsSet())
+		{
+			FString ResultPayloadUniqueId = PayloadKey.UniqueId + FString::FromInt(TimeDescription.GetHash());
+			HashStringCache = Private::HashString(ResultPayloadUniqueId);
+		}
+
+		return HashStringCache.GetValue();
+	}
+
+	FString FAnimationPayloadQuery::ToJson() const
+	{
+		TSharedPtr<FJsonObject> QueryObject = MakeShared<FJsonObject>();
+
+		QueryObject->SetStringField(TEXT("SceneNodeUniqueID"), SceneNodeUniqueID);
+		
+		QueryObject->SetStringField(TEXT("PayloadKey.UniqueID"), PayloadKey.UniqueId);
+		QueryObject->SetNumberField(TEXT("PayloadKey.Type"), (uint8)PayloadKey.Type);
+
+		QueryObject->SetNumberField(TEXT("TimeDescription.BakeFrequency"), TimeDescription.BakeFrequency);
+		QueryObject->SetNumberField(TEXT("TimeDescription.RangeStartSecond"), TimeDescription.RangeStartSecond);
+		QueryObject->SetNumberField(TEXT("TimeDescription.RangeStopSecond"), TimeDescription.RangeStopSecond);
+		
+		FString JsonString;
+		TSharedRef<TJsonWriter<TCHAR, TPrettyJsonPrintPolicy<TCHAR>>> JsonWriter = TJsonWriterFactory<TCHAR, TPrettyJsonPrintPolicy<TCHAR>>::Create(&JsonString);
+		if (!FJsonSerializer::Serialize(QueryObject.ToSharedRef(), JsonWriter))
+		{
+			// Error creating the json string 
+			return FString();
+		}
+
+		return JsonString;
+	}
+
+	void FAnimationPayloadQuery::FromJson(const FString& JsonString)
+	{
+		TSharedRef<TJsonReader<TCHAR>> Reader = FJsonStringReader::Create(JsonString);
+
+		TSharedPtr<FJsonObject> QueryObject;
+		if (!FJsonSerializer::Deserialize(Reader, QueryObject) || !QueryObject.IsValid())
+		{
+			// Cannot read the json file
+			return;
+		}
+
+		QueryObject->TryGetStringField(TEXT("SceneNodeUniqueID"), SceneNodeUniqueID);
+
+		QueryObject->TryGetStringField(TEXT("PayloadKey.UniqueId"), PayloadKey.UniqueId);
+		uint8 Type;
+		if (QueryObject->TryGetNumberField(TEXT("PayloadKey.Type"), Type))
+		{
+			PayloadKey.Type = EInterchangeAnimationPayLoadType(Type);
+		}
+
+		QueryObject->TryGetNumberField(TEXT("TimeDescription.BakeFrequency"), TimeDescription.BakeFrequency);
+		QueryObject->TryGetNumberField(TEXT("TimeDescription.RangeStartSecond"), TimeDescription.RangeStartSecond);
+		QueryObject->TryGetNumberField(TEXT("TimeDescription.RangeStopSecond"), TimeDescription.RangeStopSecond);
+	}
+
+	FString FAnimationPayloadQuery::ToJson(const TArray<FAnimationPayloadQuery>& Queries)
+	{
+		TArray<TSharedPtr<FJsonValue>> QueriesJsonArray;
+		for (const FAnimationPayloadQuery& Query : Queries)
+		{
+			
+			QueriesJsonArray.Add(MakeShared<FJsonValueString>(Query.ToJson()));
+		}
+
+		FString JsonString;
+		TSharedRef<TJsonWriter<TCHAR, TPrettyJsonPrintPolicy<TCHAR>>> JsonWriter = TJsonWriterFactory<TCHAR, TPrettyJsonPrintPolicy<TCHAR>>::Create(&JsonString);
+		if (!FJsonSerializer::Serialize(QueriesJsonArray, JsonWriter))
+		{
+			// Error creating the json string 
+			return FString();
+		}
+		return JsonString;
+	}
+
+	void FAnimationPayloadQuery::FromJson(const FString& JsonString, TArray<FAnimationPayloadQuery>& Queries)
+	{
+		TSharedRef<TJsonReader<TCHAR>> Reader = FJsonStringReader::Create(JsonString);
+
+		TArray<TSharedPtr<FJsonValue>> JsonValueArray;
+		if (!FJsonSerializer::Deserialize(Reader, JsonValueArray))
+		{
+			// Cannot read the json file
+			return;
+		}
+
+		Queries.Reserve(JsonValueArray.Num());
+
+		for (size_t QueryIndex = 0; QueryIndex < JsonValueArray.Num(); QueryIndex++)
+		{
+			FString QueryJsonString = JsonValueArray[QueryIndex]->AsString();
+
+			FAnimationPayloadQuery Query;
+			Query.FromJson(QueryJsonString);
+			Queries.Add(Query);
+		}
+	}
 }
 
 #if WITH_ENGINE
@@ -288,4 +417,4 @@ void FInterchangeStepCurve::InternalRemoveRedundantKey(TArray<ValueType>& Values
 
 template void FInterchangeStepCurve::InternalRemoveRedundantKey<float>(TArray<float>& Values, TFunctionRef<bool(const float& ValueA, const float& ValueB)> CompareFunction);
 template void FInterchangeStepCurve::InternalRemoveRedundantKey<int32>(TArray<int32>& Values, TFunctionRef<bool(const int32& ValueA, const int32& ValueB)> CompareFunction);
-template void FInterchangeStepCurve::InternalRemoveRedundantKey<FString>(TArray<FString>& Values, TFunctionRef<bool(const FString& ValueA, const FString& ValueB)> CompareFunction);
+template void FInterchangeStepCurve::InternalRemoveRedundantKey<FString>(TArray<FString>& Values, TFunctionRef<bool(const FString& ValueA, const FString& ValueB)> CompareFunction);
\ No newline at end of file
diff --git a/Engine/Plugins/Interchange/Runtime/Source/Parsers/CommonParser/Public/InterchangeCommonAnimationPayload.h b/Engine/Plugins/Interchange/Runtime/Source/Parsers/CommonParser/Public/InterchangeCommonAnimationPayload.h
--- a/Engine/Plugins/Interchange/Runtime/Source/Parsers/CommonParser/Public/InterchangeCommonAnimationPayload.h
+++ b/Engine/Plugins/Interchange/Runtime/Source/Parsers/CommonParser/Public/InterchangeCommonAnimationPayload.h
@@ -17,8 +17,15 @@
 
 namespace UE::Interchange
 {
+	namespace Private
+	{
+		INTERCHANGECOMMONPARSER_API FString HashString(const FString& String);
+	};
+
 	struct INTERCHANGECOMMONPARSER_API FAnimationPayloadData
 	{
+		FString SceneNodeUniqueID;
+
 #if WITH_ENGINE
 		//CURVE
 		TArray<FRichCurve> Curves;
@@ -41,12 +48,14 @@ namespace UE::Interchange
 		TArray<FTransform> Transforms;
 
 		//TYPE
+		FInterchangeAnimationPayLoadKey PayloadKey;
 		EInterchangeAnimationPayLoadType Type = EInterchangeAnimationPayLoadType::NONE; //Original
 		EInterchangeAnimationPayLoadType AdditionalSupportedType = EInterchangeAnimationPayLoadType::NONE;
 
 		//
-		FAnimationPayloadData(const EInterchangeAnimationPayLoadType& InType)
-			: Type(InType)
+		FAnimationPayloadData(const FString& InSceneNodeUID, const FInterchangeAnimationPayLoadKey& InPayloadKey)
+			: SceneNodeUniqueID(InSceneNodeUID)
+			, PayloadKey(InPayloadKey)
 		{
 		}
 
@@ -55,6 +64,64 @@ namespace UE::Interchange
 		//Conversions:
 		void CalculateDataFor(const EInterchangeAnimationPayLoadType& ToType, const FTransform& DefaultTransform = FTransform());
 	};
+	
+	struct INTERCHANGECOMMONPARSER_API FAnimationTimeDescription
+	{
+		double BakeFrequency;
+		double RangeStartSecond;
+		double RangeStopSecond;
+
+		FAnimationTimeDescription()
+			: BakeFrequency(0)
+			, RangeStartSecond(0)
+			, RangeStopSecond(0)
+		{
+		}
+
+		FAnimationTimeDescription(double InBakeFequency, double InRangeStartSecond, double InRangeStopSecond)
+			: BakeFrequency(InBakeFequency)
+			, RangeStartSecond(InRangeStartSecond)
+			, RangeStopSecond(InRangeStopSecond)
+		{
+
+		}
+
+		uint32 GetHash() const
+		{
+			return HashCombine(HashCombine(GetTypeHash(BakeFrequency), GetTypeHash(RangeStartSecond)), GetTypeHash(RangeStopSecond));
+		}
+	};
+
+	struct INTERCHANGECOMMONPARSER_API FAnimationPayloadQuery
+	{
+		FString SceneNodeUniqueID;
+		FInterchangeAnimationPayLoadKey PayloadKey;
+		FAnimationTimeDescription TimeDescription;
+
+	private:
+		FAnimationPayloadQuery()
+		{
+		}
+
+	public:
+		FAnimationPayloadQuery(const FString& InSceneNodeUniqueID, const FInterchangeAnimationPayLoadKey& InPayloadKey, double InBakeFequency = 0, double InRangeStartSecond = 0, double InRangeStopSecond = 0)
+			: SceneNodeUniqueID(InSceneNodeUniqueID)
+			, PayloadKey(InPayloadKey)
+			, TimeDescription(InBakeFequency, InRangeStartSecond, InRangeStopSecond)
+		{
+
+		}
+
+		FString GetHashString() const;
+
+		FString ToJson() const;
+		void FromJson(const FString& JsonString);
+
+		static FString ToJson(const TArray<FAnimationPayloadQuery>& Queries);
+		static void FromJson(const FString& JsonString, TArray<FAnimationPayloadQuery>& Queries);
+	private:
+		mutable TOptional<FString> HashStringCache;
+	};
 }
 
 /** If using Cubic, this enum describes how the tangents should be controlled. */
diff --git a/Engine/Plugins/Interchange/Runtime/Source/Parsers/Fbx/Private/FbxAPI.cpp b/Engine/Plugins/Interchange/Runtime/Source/Parsers/Fbx/Private/FbxAPI.cpp
--- a/Engine/Plugins/Interchange/Runtime/Source/Parsers/Fbx/Private/FbxAPI.cpp
+++ b/Engine/Plugins/Interchange/Runtime/Source/Parsers/Fbx/Private/FbxAPI.cpp
@@ -19,6 +19,10 @@
 #include "Nodes/InterchangeBaseNodeContainer.h"
 #include "Nodes/InterchangeSourceNode.h"
 #include "Misc/SecureHash.h"
+#include "InterchangeCommonAnimationPayload.h"
+#include "Serialization/LargeMemoryWriter.h"
+
+#include "FbxAnimation.h"
 
 #define LOCTEXT_NAMESPACE "InterchangeFbxParser"
 
@@ -244,21 +248,31 @@ namespace UE
 			}
 #endif
 
-			bool FFbxParser::FetchAnimationBakeTransformPayload(const FString& PayloadKey, const double BakeFrequency, const double RangeStartTime, const double RangeEndTime, const FString& PayloadFilepath)
+			bool FFbxParser::FetchAnimationBakeTransformPayload(const TArray<UE::Interchange::FAnimationPayloadQuery>& PayloadQueries, const FString& ResultFolder, FCriticalSection* ResultPayloadsCriticalSection, TAtomic<int64>& UniqueIdCounter, TMap<FString, FString>& ResultPayloads/*PayloadUniqueID to FilePath*/)
 			{
-				if (!PayloadContexts.Contains(PayloadKey))
+				TMap<uint32, TArray<const UE::Interchange::FAnimationPayloadQuery*>> PayloadQueriesGrouped;
+
+				for (const UE::Interchange::FAnimationPayloadQuery& PayloadQuery : PayloadQueries)
 				{
-					UInterchangeResultError_Generic* Message = AddMessage<UInterchangeResultError_Generic>();
-					Message->Text = LOCTEXT("CannotRetrievePayload", "Cannot retrieve payload; payload key doesn't have any context.");
-					return false;
+					TArray<const UE::Interchange::FAnimationPayloadQuery*>& PayloadQueriesForHash = PayloadQueriesGrouped.FindOrAdd(PayloadQuery.TimeDescription.GetHash());
+					PayloadQueriesForHash.Add(&PayloadQuery);
 				}
 
+				TArray<FText> OutErrorMessages;
+
+				bool bResult = true;
+				for (const TPair<uint32, TArray<const UE::Interchange::FAnimationPayloadQuery*>>& Group : PayloadQueriesGrouped)
 				{
-					//Critical section to force payload to be fetch one by one with no concurrency.
-					FScopeLock Lock(&PayloadCriticalSection);
-					TSharedPtr<FPayloadContextBase>& PayloadContext = PayloadContexts.FindChecked(PayloadKey);
-					return PayloadContext->FetchAnimationBakeTransformPayloadToFile(*this, BakeFrequency, RangeStartTime, RangeEndTime, PayloadFilepath);
+					bResult = FFbxAnimation::FetchAnimationBakeTransformPayload(*this, GetSDKScene(), PayloadContexts, Group.Value, ResultFolder, ResultPayloadsCriticalSection, UniqueIdCounter, ResultPayloads, OutErrorMessages) && bResult;
 				}
+
+				for (const FText& ErrorMessage : OutErrorMessages)
+				{
+					UInterchangeResultError_Generic* Message = AddMessage<UInterchangeResultError_Generic>();
+					Message->Text = ErrorMessage;
+				}
+
+				return bResult;
 			}
 
 			void ManageNamespace(const bool bKeepFbxNamespace, FString& ObjectName, FbxObject* Object)
diff --git a/Engine/Plugins/Interchange/Runtime/Source/Parsers/Fbx/Private/FbxAPI.h b/Engine/Plugins/Interchange/Runtime/Source/Parsers/Fbx/Private/FbxAPI.h
--- a/Engine/Plugins/Interchange/Runtime/Source/Parsers/Fbx/Private/FbxAPI.h
+++ b/Engine/Plugins/Interchange/Runtime/Source/Parsers/Fbx/Private/FbxAPI.h
@@ -15,6 +15,8 @@ namespace UE
 {
 	namespace Interchange
 	{
+		struct FAnimationPayloadQuery;
+		struct FAnimationPayloadData;
 		namespace Private
 		{
 			class FPayloadContextBase;
@@ -75,8 +77,11 @@ namespace UE
 				bool FetchMeshPayloadData(const FString& PayloadKey, const FTransform& MeshGlobalTransform, FMeshPayloadData& OutMeshPayloadData);
 #endif
 
-				/* Extract the fbx data from the sdk into our node container */
-				bool FetchAnimationBakeTransformPayload(const FString& PayloadKey, const double BakeFrequency, const double RangeStartTime, const double RangeEndTime, const FString& PayloadFilepath);
+				/* Extract the fbx data from the sdk into our node container
+				* @Param PayloadQueries - Will be grouped based on their TimeDescription Hashes (so that we acquire the same timings in one iteration, avoiding cache rebuilds)
+				*/
+				bool FetchAnimationBakeTransformPayload(const TArray<UE::Interchange::FAnimationPayloadQuery>& PayloadQueries, const FString& ResultFolder, FCriticalSection* ResultPayloadsCriticalSection, TAtomic<int64>& UniqueIdCounter, TMap<FString, FString>& ResultPayloads/*PayloadUniqueID to FilePath*/);
+				
 				/**
 				 * This function is used to add the given message object directly into the results for this operation.
 				 */
diff --git a/Engine/Plugins/Interchange/Runtime/Source/Parsers/Fbx/Private/FbxAnimation.cpp b/Engine/Plugins/Interchange/Runtime/Source/Parsers/Fbx/Private/FbxAnimation.cpp
--- a/Engine/Plugins/Interchange/Runtime/Source/Parsers/Fbx/Private/FbxAnimation.cpp
+++ b/Engine/Plugins/Interchange/Runtime/Source/Parsers/Fbx/Private/FbxAnimation.cpp
@@ -21,6 +21,7 @@
 #include "StaticMeshAttributes.h"
 #include "InterchangeAnimationTrackSetNode.h"
 #include "Nodes/InterchangeBaseNodeContainer.h"
+#include "Async/ParallelFor.h"
 
 #define LOCTEXT_NAMESPACE "InterchangeFbxMesh"
 
@@ -355,98 +356,6 @@ namespace UE::Interchange::Private
 		DestinationCurve.StringKeyValues = StepCurveValues;
 	}
 
-	bool ImportBakeTransforms(FNodeTransformFetchPayloadData& FetchPayloadData, FAnimationPayloadData& AnimationBakeTransformPayloadData, FFbxParser& Parser)
-	{
-		if (!ensure(!FMath::IsNearlyZero(AnimationBakeTransformPayloadData.BakeFrequency)))
-		{
-			UInterchangeResultError_Generic* Message = Parser.AddMessage<UInterchangeResultError_Generic>();
-			Message->InterchangeKey = Parser.GetFbxHelper()->GetFbxNodeHierarchyName(FetchPayloadData.Node);
-			Message->Text = LOCTEXT("BakeFrequencyZero", "Cannot fetch FBX animation bake transforms payload because the bake frequency is zero.");
-			return false;
-		}
-		
-		FbxTime StartTime;
-		StartTime.SetSecondDouble(AnimationBakeTransformPayloadData.RangeStartTime);
-		FbxTime EndTime;
-		EndTime.SetSecondDouble(AnimationBakeTransformPayloadData.RangeEndTime);
-		if (!ensure(AnimationBakeTransformPayloadData.RangeEndTime > AnimationBakeTransformPayloadData.RangeStartTime))
-		{
-			UInterchangeResultError_Generic* Message = Parser.AddMessage<UInterchangeResultError_Generic>();
-			Message->InterchangeKey = Parser.GetFbxHelper()->GetFbxNodeHierarchyName(FetchPayloadData.Node);
-			Message->Text = LOCTEXT("InvalidRange", "Cannot fetch FBX animation bake transforms payload because the bake range is invalid.");
-			return false;
-		}
-
-		const double TimeStepSecond = 1.0 / AnimationBakeTransformPayloadData.BakeFrequency;
-		FbxTime TimeStep = 0;
-		TimeStep.SetSecondDouble(TimeStepSecond);
-
-		const double SequenceLength = FMath::Max<double>(AnimationBakeTransformPayloadData.RangeEndTime - AnimationBakeTransformPayloadData.RangeStartTime, TimeStepSecond);
-		const int32 NumFrame = FMath::RoundToInt32(SequenceLength * AnimationBakeTransformPayloadData.BakeFrequency);
-		int32 BakeKeyCount = NumFrame + 1;
-
-		ensure(NumFrame >= 0);
-
-		//Add a threshold when we compare if we have reach the end of the animation
-		const FbxTime TimeComparisonThreshold = (UE_DOUBLE_KINDA_SMALL_NUMBER * static_cast<double>(FBXSDK_TC_SECOND));
-		AnimationBakeTransformPayloadData.Transforms.Empty(NumFrame);
-
-		Parser.GetSDKScene()->SetCurrentAnimationStack(FetchPayloadData.CurrentAnimStack);
-
-		bool bNanErrorLogged = false;
-		auto LogNanError = [&bNanErrorLogged, &FetchPayloadData, &Parser]()
-			{
-				if (bNanErrorLogged)
-				{
-					return;
-				}
-				UInterchangeResultError_Generic* Message = Parser.AddMessage<UInterchangeResultError_Generic>();
-				Message->InterchangeKey = Parser.GetFbxHelper()->GetFbxNodeHierarchyName(FetchPayloadData.Node);
-				Message->Text = LOCTEXT("BoneTransformNan", "Error when fetching FBX animation bake transforms payload, some transform contain NAN.");
-				bNanErrorLogged = true;
-			};
-
-		FbxTime CurrentTime = StartTime;
-		for (size_t FrameIndex = 0; FrameIndex < BakeKeyCount; FrameIndex++, CurrentTime+=TimeStep)
-		{
-			FTransform LocalTransform;
-			FbxNode* ParentNode = FetchPayloadData.Node->GetParent();
-			if (ParentNode)
-			{
-				FbxAMatrix NodeTransform = FetchPayloadData.Node->EvaluateGlobalTransform(CurrentTime);
-				FTransform GlobalTransform = UE::Interchange::Private::FFbxConvert::ConvertTransform<FTransform, FVector, FQuat>(NodeTransform);
-				if (GlobalTransform.ContainsNaN())
-				{
-					LogNanError();
-					GlobalTransform.SetIdentity();
-				}
-
-				FbxAMatrix ParentTransform = ParentNode->EvaluateGlobalTransform(CurrentTime);
-				FTransform ParentGlobalTransform = UE::Interchange::Private::FFbxConvert::ConvertTransform<FTransform, FVector, FQuat>(ParentTransform);
-				if (ParentGlobalTransform.ContainsNaN())
-				{
-					LogNanError();
-					ParentGlobalTransform.SetIdentity();
-				}
-				LocalTransform = GlobalTransform.GetRelativeTransform(ParentGlobalTransform);
-			}
-			else
-			{
-				FbxAMatrix& LocalMatrix = FetchPayloadData.Node->EvaluateLocalTransform(CurrentTime);
-				FbxVector4 NewLocalT = LocalMatrix.GetT();
-				FbxVector4 NewLocalS = LocalMatrix.GetS();
-				FbxQuaternion NewLocalQ = LocalMatrix.GetQ();
-
-				LocalTransform.SetTranslation(UE::Interchange::Private::FFbxConvert::ConvertPos<FVector>(NewLocalT));
-				LocalTransform.SetScale3D(UE::Interchange::Private::FFbxConvert::ConvertScale<FVector>(NewLocalS));
-				LocalTransform.SetRotation(UE::Interchange::Private::FFbxConvert::ConvertRotToQuat<FQuat>(NewLocalQ));
-			}
-			AnimationBakeTransformPayloadData.Transforms.Add(LocalTransform);
-		}
-
-		return true;
-	}
-
 	struct FGetFbxTransformCurvesParameters
 	{
 		FGetFbxTransformCurvesParameters(FbxScene* InSDKScene, FbxNode* InNode)
@@ -539,7 +448,7 @@ namespace UE::Interchange::Private
 		return false;
 	}
 
-	bool FAnimationPayloadContext::FetchAnimationBakeTransformPayloadToFile(FFbxParser& Parser, const double BakeFrequency, const double RangeStartTime, const double RangeEndTime, const FString& PayloadFilepath)
+	bool FAnimationPayloadContext::FetchAnimationBakeTransformPayloadForTime(FFbxParser& Parser, const FbxTime Currenttime, FTransform& OutLocalTransform)
 	{
 		if (!ensure(NodeTransformFetchPayloadData.IsSet()))
 		{
@@ -557,30 +466,65 @@ namespace UE::Interchange::Private
 			return false;
 		}
 
-		bool bBakeTransform = false;
+		bool bNanErrorLogged = false;
+		auto LogNanError = [&bNanErrorLogged, &FetchPayloadData, &Parser]()
+			{
+				if (bNanErrorLogged)
+				{
+					return;
+				}
+				UInterchangeResultError_Generic* Message = Parser.AddMessage<UInterchangeResultError_Generic>();
+				Message->InterchangeKey = Parser.GetFbxHelper()->GetFbxNodeHierarchyName(FetchPayloadData.Node);
+				Message->Text = LOCTEXT("BoneTransformNan", "Error when fetching FBX animation bake transforms payload, some transform contain NAN.");
+				bNanErrorLogged = true;
+			};
+
+		FbxNode* ParentNode = FetchPayloadData.Node->GetParent();
+		if (ParentNode)
+		{
+			FbxAMatrix NodeTransform = FetchPayloadData.Node->EvaluateGlobalTransform(Currenttime);
+			FTransform GlobalTransform = UE::Interchange::Private::FFbxConvert::ConvertTransform<FTransform, FVector, FQuat>(NodeTransform);
+			if (GlobalTransform.ContainsNaN())
+			{
+				LogNanError();
+				GlobalTransform.SetIdentity();
+			}
 
-		//At the Moment the Interchange Worker does not transfer the PayLoadType
-		//FBX parser does not need it however because the existing systems identify the appropriate fetch mechanism required in a different manner
-		//For that reason we pass PayLoadType here instead of receiving it and passing it through.
-		//FAnimationPayloadData is used here only for the Baked properties which then gets Serialized via SerializeBaked() function
-		//Resulting in the FAnimatinoPayloadData.Type to be inconsequential here.
-		FAnimationPayloadData AnimationBakeTransformPayloadData(EInterchangeAnimationPayLoadType::BAKED);
-		AnimationBakeTransformPayloadData.BakeFrequency = BakeFrequency;
-		AnimationBakeTransformPayloadData.RangeStartTime = RangeStartTime;
-		AnimationBakeTransformPayloadData.RangeEndTime = RangeEndTime;
+			FbxAMatrix ParentTransform = ParentNode->EvaluateGlobalTransform(Currenttime);
+			FTransform ParentGlobalTransform = UE::Interchange::Private::FFbxConvert::ConvertTransform<FTransform, FVector, FQuat>(ParentTransform);
+			if (ParentGlobalTransform.ContainsNaN())
+			{
+				LogNanError();
+				ParentGlobalTransform.SetIdentity();
+			}
 
-		ImportBakeTransforms(FetchPayloadData, AnimationBakeTransformPayloadData, Parser);
+			OutLocalTransform = GlobalTransform.GetRelativeTransform(ParentGlobalTransform);
+		}
+		else
 		{
-			FLargeMemoryWriter Ar;
-			AnimationBakeTransformPayloadData.SerializeBaked(Ar);
-			uint8* ArchiveData = Ar.GetData();
-			int64 ArchiveSize = Ar.TotalSize();
-			TArray64<uint8> Buffer(ArchiveData, ArchiveSize);
-			FFileHelper::SaveArrayToFile(Buffer, *PayloadFilepath);
+			FbxAMatrix& LocalMatrix = FetchPayloadData.Node->EvaluateLocalTransform(Currenttime);
+			FbxVector4 NewLocalT = LocalMatrix.GetT();
+			FbxVector4 NewLocalS = LocalMatrix.GetS();
+			FbxQuaternion NewLocalQ = LocalMatrix.GetQ();
+
+			OutLocalTransform.SetTranslation(UE::Interchange::Private::FFbxConvert::ConvertPos<FVector>(NewLocalT));
+			OutLocalTransform.SetScale3D(UE::Interchange::Private::FFbxConvert::ConvertScale<FVector>(NewLocalS));
+			OutLocalTransform.SetRotation(UE::Interchange::Private::FFbxConvert::ConvertRotToQuat<FQuat>(NewLocalQ));
 		}
+
 		return true;
 	}
 
+	FbxAnimStack* FAnimationPayloadContext::GetAnimStack()
+	{
+		if (NodeTransformFetchPayloadData.IsSet())
+		{
+			return NodeTransformFetchPayloadData->CurrentAnimStack;
+		}
+
+		return nullptr;
+	}
+
 	bool FAnimationPayloadContext::InternalFetchCurveNodePayloadToFile(FFbxParser& Parser, const FString& PayloadFilepath)
 	{
 		if (AttributeFetchPayloadData.IsSet())
@@ -977,6 +921,194 @@ namespace UE::Interchange::Private
 
 		return false;
 	}
+
+	//PayloadQueries arriving here should be of the same start/stop and frequency:
+	bool FFbxAnimation::FetchAnimationBakeTransformPayload(FFbxParser& Parser, FbxScene* SDKScene,
+		TMap<FString, TSharedPtr<FPayloadContextBase>>& PayloadContexts,
+		const TArray<const UE::Interchange::FAnimationPayloadQuery*>& PayloadQueries,
+		const FString& ResultFolder, FCriticalSection* ResultPayloadsCriticalSection,
+		TAtomic<int64>& UniqueIdCounter, TMap<FString, FString>& ResultPayloads,
+		TArray<FText>& OutErrorMessages)
+	{
+		//Helper Structs:
+		struct FPayloadQueryHelper
+		{
+			const UE::Interchange::FAnimationPayloadQuery* PayloadQuery;
+			TSharedPtr<FPayloadContextBase> PayloadContext;
+			FPayloadQueryHelper(const UE::Interchange::FAnimationPayloadQuery* InPayloadQuery, TSharedPtr<FPayloadContextBase> InPayloadContext)
+				: PayloadQuery(InPayloadQuery)
+				, PayloadContext(InPayloadContext)
+			{
+			}
+		};
+		struct FPayloadDataHelper
+		{
+			FString QueryHashString; //Used for payload file name
+			FAnimationPayloadData PayloadData;
+			FPayloadDataHelper(const FString& InHashString, const FString& InSceneNodeUID, const FInterchangeAnimationPayLoadKey& InPayloadKey)
+				: QueryHashString(InHashString)
+				, PayloadData(InSceneNodeUID, InPayloadKey)
+			{
+			}
+		};
+
+		//If no PayloadQueries then return
+		if (PayloadQueries.Num() == 0)
+		{
+			return true;
+		}
+
+		//Get Timings and Number of Frames for the current PayloadQueries
+		//		Note: PayloadQueries arriving here(FetchAnimationBakeTransformPayloadInternal) should be of the same start/stop and frequency, so we can grab the first one's:
+		FAnimationTimeDescription TimeDescription = PayloadQueries[0]->TimeDescription;
+		if (!ensure(!FMath::IsNearlyZero(TimeDescription.BakeFrequency)))
+		{
+			OutErrorMessages.Add(LOCTEXT("BakeFrequencyZero", "Cannot fetch FBX animation bake transforms payload because the bake frequency is zero."));
+			return false;
+		}
+		FbxTime StartTime;
+		StartTime.SetSecondDouble(TimeDescription.RangeStartSecond);
+		FbxTime EndTime;
+		EndTime.SetSecondDouble(TimeDescription.RangeStopSecond);
+		if (!ensure(TimeDescription.RangeStopSecond > TimeDescription.RangeStartSecond))
+		{
+			OutErrorMessages.Add(LOCTEXT("InvalidRange", "Cannot fetch FBX animation bake transforms payload because the bake range is invalid."));
+			return false;
+		}
+		const double TimeStepSecond = 1.0 / TimeDescription.BakeFrequency;
+		FbxTime TimeStep = 0;
+		TimeStep.SetSecondDouble(TimeStepSecond);
+
+		const double SequenceLength = FMath::Max<double>(TimeDescription.RangeStopSecond - TimeDescription.RangeStartSecond, TimeStepSecond);
+		const int32 NumFrame = FMath::RoundToInt32(SequenceLength * TimeDescription.BakeFrequency);
+		int32 BakeKeyCount = NumFrame + 1;
+
+		ensure(NumFrame >= 0);
+
+		//Acquire PayloadContexts and AnimationStacks for the PayloadQueries (also group the FPayloadQueryHelper per AnimationStacks):
+		//		PayloadQueryHelpers maps AnimStacks to FPaloayQueryHelper (which contains a PayloadQuery and the PayloadContext that belongs to it)
+		TMap<FbxAnimStack*, TArray<FPayloadQueryHelper>> PayloadQueryHelpers;
+		{
+			for (const UE::Interchange::FAnimationPayloadQuery* PayloadQuery : PayloadQueries)
+			{
+				if (!PayloadContexts.Contains(PayloadQuery->PayloadKey.UniqueId))
+				{
+					OutErrorMessages.Add(FText::Format(LOCTEXT("CannotRetrievePayload", "Cannot retrieve payload; payload key['{0}'] doesn't have any context."), FText::FromString(PayloadQuery->PayloadKey.UniqueId)));
+					continue;
+				}
+
+				TSharedPtr<FPayloadContextBase> PayloadContext = PayloadContexts.FindChecked(PayloadQuery->PayloadKey.UniqueId);
+
+				FbxAnimStack* AnimStack = PayloadContext->GetAnimStack();
+				if (!AnimStack)
+				{
+					continue;
+				}
+
+				FString ResultPayloadUniqueId = PayloadQuery->GetHashString();
+
+				//If we already have extract this mesh, no need to extract again
+				if (ResultPayloads.Contains(ResultPayloadUniqueId))
+				{
+					continue;
+				}
+
+				TArray<FPayloadQueryHelper>& PayloadHelpersPerAnimStack = PayloadQueryHelpers.FindOrAdd(AnimStack);
+				PayloadHelpersPerAnimStack.Add(FPayloadQueryHelper(PayloadQuery, PayloadContexts.FindChecked(PayloadQuery->PayloadKey.UniqueId)));
+			}
+		}
+
+		//Iterate on the AnimStacked groups:
+		for (TPair<FbxAnimStack*, TArray<FPayloadQueryHelper>>& AnimStackPayloadQueryHelpersPair : PayloadQueryHelpers)
+		{
+			SDKScene->SetCurrentAnimationStack(AnimStackPayloadQueryHelpersPair.Key);
+
+			TArray<FPayloadQueryHelper>& PayloadQueryHelpersForAnimStack = AnimStackPayloadQueryHelpersPair.Value;
+
+			//Initialize PayloadDataHelpers (contains the HashString for the Query and the FAnimationPayloadData
+			TArray<FPayloadDataHelper> PayloadDataHelpersForAnimStack;
+			PayloadDataHelpersForAnimStack.Empty(PayloadQueryHelpersForAnimStack.Num());
+			for (FPayloadQueryHelper& PayloadQueryHelper : PayloadQueryHelpersForAnimStack)
+			{
+				FPayloadDataHelper PayloadDataHelper(PayloadQueryHelper.PayloadQuery->GetHashString(), PayloadQueryHelper.PayloadQuery->SceneNodeUniqueID, PayloadQueryHelper.PayloadQuery->PayloadKey);
+
+				PayloadDataHelper.PayloadData.BakeFrequency = PayloadQueryHelper.PayloadQuery->TimeDescription.BakeFrequency;
+				PayloadDataHelper.PayloadData.RangeStartTime = PayloadQueryHelper.PayloadQuery->TimeDescription.RangeStartSecond;
+				PayloadDataHelper.PayloadData.RangeEndTime = PayloadQueryHelper.PayloadQuery->TimeDescription.RangeStopSecond;
+
+				PayloadDataHelper.PayloadData.Transforms.SetNum(BakeKeyCount);
+
+				PayloadDataHelpersForAnimStack.Add(PayloadDataHelper);
+			}
+
+			//Acquire Bone Transforms:
+			//Time iteration
+			FbxTime CurrentTime = StartTime;
+			for (size_t FrameIndex = 0; FrameIndex < BakeKeyCount; FrameIndex++, CurrentTime+=TimeStep)
+			{
+
+				//Bone iteration:
+				for (size_t PayloadDataIndex = 0; PayloadDataIndex < PayloadDataHelpersForAnimStack.Num(); PayloadDataIndex++)
+				{
+					FPayloadQueryHelper& PayloadQueryHelper = PayloadQueryHelpersForAnimStack[PayloadDataIndex];
+					FPayloadDataHelper& PayloadDataHelper = PayloadDataHelpersForAnimStack[PayloadDataIndex];
+
+					FTransform Transform;
+					if (PayloadQueryHelper.PayloadContext->FetchAnimationBakeTransformPayloadForTime(Parser, CurrentTime, Transform))
+					{
+						PayloadDataHelper.PayloadData.Transforms[FrameIndex] = Transform;
+					}
+					else
+					{
+						if (FrameIndex == 0)
+						{
+							PayloadDataHelper.PayloadData.Transforms[FrameIndex] = FTransform::Identity;
+						}
+						else
+						{
+							PayloadDataHelper.PayloadData.Transforms[FrameIndex] = PayloadDataHelper.PayloadData.Transforms[FrameIndex - 1];
+						}
+					}
+				}
+			}
+
+			//Write out results:
+			ParallelFor(PayloadDataHelpersForAnimStack.Num(), [&PayloadDataHelpersForAnimStack, &ResultPayloads, &ResultPayloadsCriticalSection, &ResultFolder, &UniqueIdCounter](int32 PayloadDataHelperIndex)
+				{
+					FPayloadDataHelper& PayloadDataHelper = PayloadDataHelpersForAnimStack[PayloadDataHelperIndex];
+					FString QueryHashString = PayloadDataHelper.QueryHashString; //PayloadQuery's GetHashString()
+
+					//If we already have extract this mesh, no need to extract again
+					if (ResultPayloads.Contains(QueryHashString))
+					{
+						return;
+					}
+
+					FString PayloadFilepathCopy;
+
+					{
+						FScopeLock Lock(ResultPayloadsCriticalSection);
+						FString& PayloadFilepath = ResultPayloads.FindOrAdd(QueryHashString);
+						//To avoid file path with too many character, we hash the payloadKey so we have a deterministic length for the file path.
+						PayloadFilepath = ResultFolder + TEXT("/") + QueryHashString + FString::FromInt(UniqueIdCounter.IncrementExchange()) + TEXT(".payload");
+
+						//Copy the map filename key because we are multithreaded and the TMap can be reallocated
+						PayloadFilepathCopy = PayloadFilepath;
+					}
+
+					FString PayloadFilepath = PayloadFilepathCopy;
+
+					FLargeMemoryWriter Ar;
+					PayloadDataHelper.PayloadData.SerializeBaked(Ar);
+					uint8* ArchiveData = Ar.GetData();
+					int64 ArchiveSize = Ar.TotalSize();
+					TArray64<uint8> Buffer(ArchiveData, ArchiveSize);
+					FFileHelper::SaveArrayToFile(Buffer, *PayloadFilepath);
+				});
+		}
+
+		return true;
+	}
 }//ns UE::Interchange::Private
 
 #undef LOCTEXT_NAMESPACE
diff --git a/Engine/Plugins/Interchange/Runtime/Source/Parsers/Fbx/Private/FbxAnimation.h b/Engine/Plugins/Interchange/Runtime/Source/Parsers/Fbx/Private/FbxAnimation.h
--- a/Engine/Plugins/Interchange/Runtime/Source/Parsers/Fbx/Private/FbxAnimation.h
+++ b/Engine/Plugins/Interchange/Runtime/Source/Parsers/Fbx/Private/FbxAnimation.h
@@ -14,6 +14,11 @@ class UInterchangeMeshNode;
 class UInterchangeSceneNode;
 class UInterchangeSkeletalAnimationTrackNode;
 
+namespace UE::Interchange
+{
+	struct FAnimationPayloadQuery;
+}
+
 namespace UE::Interchange::Private
 {
 	struct FMorphTargetAnimationBuildingData;
@@ -76,7 +81,8 @@ namespace UE::Interchange::Private
 		virtual ~FAnimationPayloadContext() {}
 		virtual FString GetPayloadType() const override { return TEXT("TransformAnimation-PayloadContext"); }
 		virtual bool FetchPayloadToFile(FFbxParser& Parser, const FString& PayloadFilepath) override;
-		virtual bool FetchAnimationBakeTransformPayloadToFile(FFbxParser& Parser, const double BakeFrequency, const double RangeStartTime, const double RangeEndTime, const FString& PayloadFilepath) override;
+		virtual bool FetchAnimationBakeTransformPayloadForTime(FFbxParser& Parser, const FbxTime Currenttime, FTransform& Out) override;
+		virtual FbxAnimStack* GetAnimStack() override;
 		
 		TOptional<FNodeTransformFetchPayloadData> NodeTransformFetchPayloadData;
 		TOptional<FAttributeFetchPayloadData> AttributeFetchPayloadData;
@@ -127,6 +133,13 @@ namespace UE::Interchange::Private
 
 		static bool IsFbxPropertyTypeSupported(EFbxType PropertyType);
 		static bool IsFbxPropertyTypeDecimal(EFbxType PropertyType);
-
+		
+		//PayloadQueries arriving here should be of the same start/stop and frequency:
+		static bool FetchAnimationBakeTransformPayload(FFbxParser& Parser, FbxScene* SDKScene,
+			TMap<FString, TSharedPtr<FPayloadContextBase>>& PayloadContexts,
+			const TArray<const UE::Interchange::FAnimationPayloadQuery*>& PayloadQueries, 
+			const FString& ResultFolder, FCriticalSection* ResultPayloadsCriticalSection, 
+			TAtomic<int64>& UniqueIdCounter, TMap<FString, FString>& ResultPayloads,
+			TArray<FText>& OutErrorMessages);
 	};
 }//ns UE::Interchange::Private
diff --git a/Engine/Plugins/Interchange/Runtime/Source/Parsers/Fbx/Private/FbxHelper.h b/Engine/Plugins/Interchange/Runtime/Source/Parsers/Fbx/Private/FbxHelper.h
--- a/Engine/Plugins/Interchange/Runtime/Source/Parsers/Fbx/Private/FbxHelper.h
+++ b/Engine/Plugins/Interchange/Runtime/Source/Parsers/Fbx/Private/FbxHelper.h
@@ -28,7 +28,9 @@ namespace UE::Interchange
 #if WITH_ENGINE
 			virtual bool FetchMeshPayload(FFbxParser& Parser, const FTransform& MeshGlobalTransform, FMeshPayloadData& OutMeshPayloadData) { return false; };
 #endif
-			virtual bool FetchAnimationBakeTransformPayloadToFile(FFbxParser& Parser, const double BakeFrequency, const double RangeStartTime, const double RangeEndTime, const FString& PayloadFilepath) { return false; }
+			virtual bool FetchAnimationBakeTransformPayloadForTime(FFbxParser& Parser, const FbxTime Currenttime, FTransform& Out) { return false; }
+
+			virtual FbxAnimStack* GetAnimStack() { return nullptr; }
 		};
 
 		struct FFbxHelper
diff --git a/Engine/Plugins/Interchange/Runtime/Source/Parsers/Fbx/Private/InterchangeFbxParser.cpp b/Engine/Plugins/Interchange/Runtime/Source/Parsers/Fbx/Private/InterchangeFbxParser.cpp
--- a/Engine/Plugins/Interchange/Runtime/Source/Parsers/Fbx/Private/InterchangeFbxParser.cpp
+++ b/Engine/Plugins/Interchange/Runtime/Source/Parsers/Fbx/Private/InterchangeFbxParser.cpp
@@ -12,32 +12,13 @@
 #include "UObject/StrongObjectPtr.h"
 #include "UObject/UObjectGlobals.h"
 #include "UObject/Package.h"
+#include "InterchangeHelper.h"
+#include "InterchangeCommonAnimationPayload.h"
 
 #define LOCTEXT_NAMESPACE "InterchangeFbxParser"
 
 namespace UE::Interchange
 {
-	namespace Private
-	{
-		FString HashString(const FString& String)
-		{
-			TArray<uint8> TempBytes;
-			TempBytes.Reserve(64);
-			//The archive is flagged as persistent so that machines of different endianness produce identical binary results.
-			FMemoryWriter Ar(TempBytes, /*bIsPersistent=*/ true);
-			//Hack because serialization do not support const
-			FString& NonConst = *const_cast<FString*>(&String);
-			Ar << NonConst;
-			FSHA1 Sha;
-			Sha.Update(TempBytes.GetData(), TempBytes.Num() * TempBytes.GetTypeSize());
-			Sha.Final();
-			// Retrieve the hash and use it to construct a pseudo-GUID.
-			uint32 Hash[5];
-			Sha.GetHash((uint8*)Hash);
-			FGuid Guid = FGuid(Hash[0] ^ Hash[4], Hash[1], Hash[2], Hash[3]);
-			return Guid.ToString(EGuidFormats::Base36Encoded);
-		}
-	}
 
 	FInterchangeFbxParser::FInterchangeFbxParser()
 	{
@@ -223,49 +204,6 @@ namespace UE::Interchange
 	}
 #endif
 
-	FString FInterchangeFbxParser::FetchAnimationBakeTransformPayload(const FString& PayloadKey, const double BakeFrequency, const double RangeStartTime, const double RangeEndTime, const FString& ResultFolder)
-	{
-		ResultsContainer->Empty();
-		FString PayloadFilepathCopy;
-		FString ResultPayloadUniqueId = PayloadKey + FString::FromInt(static_cast<int32>(BakeFrequency * 1000.0)) + FString::FromInt(static_cast<int32>(RangeStartTime * 1000.0)) + FString::FromInt(static_cast<int32>(RangeEndTime * 1000.0));
-
-		if (!ensure(FbxParserPrivate.IsValid()))
-		{
-			if (UInterchangeResultError_Generic* Error = AddMessage<UInterchangeResultError_Generic>())
-			{
-				Error->SourceAssetName = PayloadKey;
-				Error->Text = LOCTEXT("CantFetchAnimationBakeTransformPayload_ParserInvalid", "FInterchangeFbxParser::FetchAnimationBakeTransformPayload: Cannot fetch the animation bake transform payload. The internal fbx parser is invalid.");
-			}
-			return ResultPayloadUniqueId;
-		}
-		
-		//If we already have extract this mesh, no need to extract again
-		if (ResultPayloads.Contains(ResultPayloadUniqueId))
-		{
-			return ResultPayloadUniqueId;
-		}
-
-		{
-			FScopeLock Lock(&ResultPayloadsCriticalSection);
-			FString& PayloadFilepath = ResultPayloads.FindOrAdd(ResultPayloadUniqueId);
-			//To avoid file path with too many character, we hash the payloadKey so we have a deterministic length for the file path.
-			FString PayloadKeyHash = Private::HashString(PayloadKey);
-			PayloadFilepath = ResultFolder + TEXT("/") + PayloadKeyHash + FString::FromInt(UniqueIdCounter.IncrementExchange()) + TEXT(".payload");
-
-			//Copy the map filename key because we are multithreaded and the TMap can be reallocated
-			PayloadFilepathCopy = PayloadFilepath;
-		}
-		if (!FbxParserPrivate->FetchAnimationBakeTransformPayload(PayloadKey, BakeFrequency, RangeStartTime, RangeEndTime, PayloadFilepathCopy))
-		{
-			if (UInterchangeResultError_Generic* Error = AddMessage<UInterchangeResultError_Generic>())
-			{
-				Error->SourceAssetName = SourceFilename;
-				Error->Text = LOCTEXT("CantFetchPayload", "Cannot fetch FBX payload data.");
-			}
-		}
-		return ResultPayloadUniqueId;
-	}
-
 	TArray<FString> FInterchangeFbxParser::GetJsonLoadMessages() const
 	{
 		TArray<FString> JsonResults;
@@ -287,6 +225,29 @@ namespace UE::Interchange
 		return ResultsContainer.Get();
 	}
 
+	void FInterchangeFbxParser::FetchAnimationBakeTransformPayloads(const TArray<UE::Interchange::FAnimationPayloadQuery>& PayloadQueries, const FString& ResultFolder)
+	{
+		FbxParserPrivate->FetchAnimationBakeTransformPayload(PayloadQueries, ResultFolder, &ResultPayloadsCriticalSection, UniqueIdCounter, ResultPayloads);
+	}
+
+	//Used by DispatchWorker
+	TMap<FString, FString> FInterchangeFbxParser::FetchAnimationBakeTransformPayloads(const FString& PayloadQueriesJsonString, const FString& ResultFolder)
+	{
+		TMap<FString, FString> HashToFilePaths;
+
+		TArray<UE::Interchange::FAnimationPayloadQuery> PayloadQueries;
+		UE::Interchange::FAnimationPayloadQuery::FromJson(PayloadQueriesJsonString, PayloadQueries);
+
+		FbxParserPrivate->FetchAnimationBakeTransformPayload(PayloadQueries, ResultFolder, &ResultPayloadsCriticalSection, UniqueIdCounter, ResultPayloads);
+
+		//Acquire FilePaths:
+		for (const UE::Interchange::FAnimationPayloadQuery& PayloadQuery : PayloadQueries)
+		{
+			HashToFilePaths.Add(PayloadQuery.GetHashString(), GetResultPayloadFilepath(PayloadQuery.GetHashString()));
+		}
+
+		return HashToFilePaths;
+	}
 }//ns UE::Interchange
 
 #undef LOCTEXT_NAMESPACE
diff --git a/Engine/Plugins/Interchange/Runtime/Source/Parsers/Fbx/Public/InterchangeFbxParser.h b/Engine/Plugins/Interchange/Runtime/Source/Parsers/Fbx/Public/InterchangeFbxParser.h
--- a/Engine/Plugins/Interchange/Runtime/Source/Parsers/Fbx/Public/InterchangeFbxParser.h
+++ b/Engine/Plugins/Interchange/Runtime/Source/Parsers/Fbx/Public/InterchangeFbxParser.h
@@ -16,6 +16,7 @@ namespace UE
 {
 	namespace Interchange
 	{
+		struct FAnimationPayloadQuery;
 		namespace Private
 		{
 			class FFbxParser;
@@ -75,15 +76,10 @@ namespace UE
 #endif
 			
 			/**
-			 * Extract bake transform animation payload data from the fbx, the key tell the translator what payload the client ask
-			 * @param - PayloadKey is the key that describe the payload data to extract from the fbx file
-			 * @param - BakeFrequency is the Hz at which we should bake the transform
-			 * @param - RangeStartTime is the start time of the bake
-			 * @param - RangeEndTime is the end time of the bake
-			 * @param - ResultFolder is the folder where we must put any result file
-			 * @return - Return the 'ResultPayloads' key unique id. We cannot use only the payload key because the bake parameter can be different.
+			 * Extract bake transform animation payload data from the fbx
 			 */
-			FString FetchAnimationBakeTransformPayload(const FString& PayloadKey, const double BakeFrequency, const double RangeStartTime, const double RangeEndTime, const FString& ResultFolder);
+			void FetchAnimationBakeTransformPayloads(const TArray<UE::Interchange::FAnimationPayloadQuery>& PayloadQueries, const FString& ResultFolder);
+			TMap<FString, FString> FetchAnimationBakeTransformPayloads(const FString& PayloadQueriesJsonString, const FString& ResultFolder);
 
 			FString GetResultFilepath() const { return ResultFilepath; }
 			FString GetResultPayloadFilepath(const FString& PayloadKey) const
diff --git a/Engine/Source/Programs/InterchangeWorker/Private/InterchangeWorkerImpl.cpp b/Engine/Source/Programs/InterchangeWorker/Private/InterchangeWorkerImpl.cpp
--- a/Engine/Source/Programs/InterchangeWorker/Private/InterchangeWorkerImpl.cpp
+++ b/Engine/Source/Programs/InterchangeWorker/Private/InterchangeWorkerImpl.cpp
@@ -222,7 +222,7 @@ void FInterchangeWorkerImpl::ProcessCommand(const TSharedPtr<UE::Interchange::IC
 	FString JSonResult;
 	TArray<FString> JSonMessages;
 	FJsonLoadSourceCmd LoadSourceCommand;
-	FJsonFetchAnimationBakeTransformPayloadCmd FetchAnimationBakeTransform;
+	FJsonFetchAnimationQueriesCmd FetchAnimationQueries;
 	FJsonFetchMeshPayloadCmd FetchMeshPayloadCommand;
 	FJsonFetchPayloadCmd FetchPayloadCommand;
 	//Any command FromJson function return true if the Json descibe the command
@@ -235,13 +235,13 @@ void FInterchangeWorkerImpl::ProcessCommand(const TSharedPtr<UE::Interchange::IC
 			ProcessResult = LoadFbxFile(LoadSourceCommand, JSonResult, JSonMessages);
 		}
 	}
-	else if (FetchAnimationBakeTransform.FromJson(JsonToProcess))
+	else if (FetchAnimationQueries.FromJson(JsonToProcess))
 	{
 		//Load file command
-		if (FetchAnimationBakeTransform.GetTranslatorID().Equals(TEXT("FBX"), ESearchCase::IgnoreCase))
+		if (FetchAnimationQueries.GetTranslatorID().Equals(TEXT("FBX"), ESearchCase::IgnoreCase))
 		{
 			//We want to load an FBX file
-			ProcessResult = FetchFbxPayload(FetchAnimationBakeTransform, JSonResult, JSonMessages);
+			ProcessResult = FetchFbxPayload(FetchAnimationQueries, JSonResult, JSonMessages);
 		}
 	}
 	else if (FetchMeshPayloadCommand.FromJson(JsonToProcess))
@@ -328,20 +328,17 @@ ETaskState FInterchangeWorkerImpl::FetchFbxPayload(const FJsonFetchMeshPayloadCm
 	return ResultState;
 }
 
-ETaskState FInterchangeWorkerImpl::FetchFbxPayload(const FJsonFetchAnimationBakeTransformPayloadCmd& FetchAnimationBakeTransformPayloadCommand, FString& OutJSonResult, TArray<FString>& OutJSonMessages)
+ETaskState FInterchangeWorkerImpl::FetchFbxPayload(const FJsonFetchAnimationQueriesCmd& FetchAnimationBakeTransformPayloadCommand, FString& OutJSonResult, TArray<FString>& OutJSonMessages)
 {
 	ETaskState ResultState = ETaskState::Unknown;
-	FString PayloadKey = FetchAnimationBakeTransformPayloadCommand.GetPayloadKey();
 
-	FString ResultPayloadsUniqueId = FbxParser.FetchAnimationBakeTransformPayload(PayloadKey
-		, FetchAnimationBakeTransformPayloadCommand.GetBakeFrequency()
-		, FetchAnimationBakeTransformPayloadCommand.GetRangeStartTime()
-		, FetchAnimationBakeTransformPayloadCommand.GetRangeEndTime()
-		, ResultFolder);
-	FJsonLoadSourceCmd::JsonResultParser ResultParser;
-	ResultParser.SetResultFilename(FbxParser.GetResultPayloadFilepath(ResultPayloadsUniqueId));
+	TMap<FString, FString> ResultPayloadFilePaths = FbxParser.FetchAnimationBakeTransformPayloads(FetchAnimationBakeTransformPayloadCommand.GetQueriesJsonString(), ResultFolder);
+	
+	FJsonFetchAnimationQueriesCmd::JsonAnimationQueriesResultParser ResultParser;
+	ResultParser.SetHashToFilenames(ResultPayloadFilePaths);
 	OutJSonMessages = FbxParser.GetJsonLoadMessages();
 	OutJSonResult = ResultParser.ToJson();
 	ResultState = ETaskState::ProcessOk;
+
 	return ResultState;
 }
\ No newline at end of file
diff --git a/Engine/Source/Programs/InterchangeWorker/Private/InterchangeWorkerImpl.h b/Engine/Source/Programs/InterchangeWorker/Private/InterchangeWorkerImpl.h
--- a/Engine/Source/Programs/InterchangeWorker/Private/InterchangeWorkerImpl.h
+++ b/Engine/Source/Programs/InterchangeWorker/Private/InterchangeWorkerImpl.h
@@ -28,7 +28,7 @@ private:
 	UE::Interchange::ETaskState LoadFbxFile(const UE::Interchange::FJsonLoadSourceCmd& LoadSourceCommand, FString& OutJSonResult, TArray<FString>& OutJSonMessages);
 	UE::Interchange::ETaskState FetchFbxPayload(const UE::Interchange::FJsonFetchPayloadCmd& FetchPayloadCommand, FString& OutJSonResult, TArray<FString>& OutJSonMessages);
 	UE::Interchange::ETaskState FetchFbxPayload(const UE::Interchange::FJsonFetchMeshPayloadCmd& FetchMeshPayloadCommand, FString& OutJSonResult, TArray<FString>& OutJSonMessages);
-	UE::Interchange::ETaskState FetchFbxPayload(const UE::Interchange::FJsonFetchAnimationBakeTransformPayloadCmd& FetchAnimationBakeTransformPayloadCommand, FString& OutJSonResult, TArray<FString>& OutJSonMessages);
+	UE::Interchange::ETaskState FetchFbxPayload(const UE::Interchange::FJsonFetchAnimationQueriesCmd& FetchAnimationBakeTransformPayloadCommand, FString& OutJSonResult, TArray<FString>& OutJSonMessages);
 private:
 	UE::Interchange::FNetworkClientNode NetworkInterface;
 	UE::Interchange::FCommandQueue CommandIO;
