MaterialX:
- Fix <normalmap> node by adding support of interfacename

diff --git a/Engine/Plugins/Interchange/Runtime/Source/Import/Private/MaterialX/MaterialXUtils/MaterialXSurfaceShaderAbstract.cpp b/Engine/Plugins/Interchange/Runtime/Source/Import/Private/MaterialX/MaterialXUtils/MaterialXSurfaceShaderAbstract.cpp
--- a/Engine/Plugins/Interchange/Runtime/Source/Import/Private/MaterialX/MaterialXUtils/MaterialXSurfaceShaderAbstract.cpp
+++ b/Engine/Plugins/Interchange/Runtime/Source/Import/Private/MaterialX/MaterialXUtils/MaterialXSurfaceShaderAbstract.cpp
@@ -1222,8 +1222,19 @@ void FMaterialXSurfaceShaderAbstract::ConnectNormalMapInputToOutput(const FConne
 {
 	using namespace UE::Interchange::Materials::Standard::Nodes;
 
-	// Only create a FunctionCall if there's a scale, otherwise just like dot
-	if(mx::InputPtr Input = Connect.UpstreamNode->getInput("scale"); Input && Input->hasValueString() && (mx::fromValueString<float>(Input->getValueString()) != 1.f))
+	// Only create a FunctionCall if there's a scale and the value is not 1, otherwise just like dot
+	mx::InputPtr Input = Connect.UpstreamNode->getInput("scale");
+	bool bIsNotEqualOne = true;
+	if(Input && Input->hasValueString())
+	{
+		bIsNotEqualOne = (mx::fromValueString<float>(Input->getValueString()) != 1.f);
+	}
+	else if(Input && Input->hasInterfaceName())
+	{
+		bIsNotEqualOne = (mx::fromValueString<float>(Input->getInterfaceInput()->getValueString()) != 1.f);
+	}
+
+	if(Input && bIsNotEqualOne)
 	{
 		UInterchangeShaderNode* FlattenNormalNode = CreateFunctionCallShaderNode(Connect.UpstreamNode->getName().c_str(), TEXT("/Engine/Functions/Engine_MaterialFunctions01/Texturing/FlattenNormal.FlattenNormal"));
 
