MaterialX:
- Remove namespace for StandardSurface default values

diff --git a/Engine/Plugins/Interchange/Runtime/Source/Common/Public/MaterialX/InterchangeMaterialXDefinitions.h b/Engine/Plugins/Interchange/Runtime/Source/Common/Public/MaterialX/InterchangeMaterialXDefinitions.h
--- a/Engine/Plugins/Interchange/Runtime/Source/Common/Public/MaterialX/InterchangeMaterialXDefinitions.h
+++ b/Engine/Plugins/Interchange/Runtime/Source/Common/Public/MaterialX/InterchangeMaterialXDefinitions.h
@@ -280,51 +280,44 @@ MATERIALX_NAMESPACE_BEGIN
 
 		namespace DefaultValue
 		{
-			namespace Float
-			{
-				static float Base = 1.f;
-				static float DiffuseRoughness = 0.f;
-				static float Metalness = 0.f;
-				static float Specular = 1.f;
-				static float SpecularRoughness = 0.2f;
-				static float SpecularIOR = 1.5f;
-				static float SpecularAnisotropy = 0.f;
-				static float SpecularRotation = 0.f;
-				static float Transmission = 0.f;
-				static float TransmissionDepth = 0.f;
-				static float TransmissionScatterAnisotropy = 0.f;
-				static float TransmissionDispersion = 0.f;
-				static float TransmissionExtraRoughness = 0.f;
-				static float Subsurface = 0.f;
-				static float SubsurfaceScale = 1.f;
-				static float SubsurfaceAnisotropy = 0.f;
-				static float Sheen = 0.f;
-				static float SheenRoughness = 0.3f;
-				static float Coat = 0.f;
-				static float CoatRoughness = 0.1f;
-				static float CoatAnisotropy = 0.f;
-				static float CoatRotation = 0.f;
-				static float CoatIOR = 1.5f;
-				static float CoatAffectColor = 0.f;
-				static float CoatAffectRoughness = 0.f;
-				static float ThinFilmThickness = 0.f;
-				static float ThinFilmIOR = 1.5f;
-				static float Emission = 0.f;
-			}
-
-			namespace Color3
-			{
-				constexpr FLinearColor BaseColor{ 0.8, 0.8, 0.8 };
-				constexpr FLinearColor SpecularColor{ 1, 1, 1 };
-				constexpr FLinearColor TransmissionColor{ 1, 1, 1 };
-				constexpr FLinearColor TransmissionScatter{ 0, 0, 0 };
-				constexpr FLinearColor SubsurfaceColor{ 1, 1, 1 };
-				constexpr FLinearColor SubsurfaceRadius{ 1, 1, 1 };
-				constexpr FLinearColor SheenColor{ 1, 1, 1 };
-				constexpr FLinearColor CoatColor{ 1, 1, 1 };
-				constexpr FLinearColor EmissionColor{ 1, 1, 1 };
-				constexpr FLinearColor Opacity{ 1, 1, 1 };
-			}
+			static constexpr float Base = 1.f;
+			constexpr FLinearColor BaseColor{ 0.8, 0.8, 0.8 };
+			static constexpr float DiffuseRoughness = 0.f;
+			static constexpr float Metalness = 0.f;
+			static constexpr float Specular = 1.f;
+			constexpr FLinearColor SpecularColor{ 1, 1, 1 };
+			static constexpr float SpecularRoughness = 0.2f;
+			static constexpr float SpecularIOR = 1.5f;
+			static constexpr float SpecularAnisotropy = 0.f;
+			static constexpr float SpecularRotation = 0.f;
+			static constexpr float Transmission = 0.f;
+			static constexpr float TransmissionDepth = 0.f;
+			static constexpr float TransmissionScatterAnisotropy = 0.f;
+			static constexpr float TransmissionDispersion = 0.f;
+			static constexpr float TransmissionExtraRoughness = 0.f;
+			static constexpr float Subsurface = 0.f;
+			constexpr FLinearColor SubsurfaceColor{ 1, 1, 1 };
+			constexpr FLinearColor SubsurfaceRadius{ 1, 1, 1 };
+			static constexpr float SubsurfaceScale = 1.f;
+			static constexpr float SubsurfaceAnisotropy = 0.f;
+			static constexpr float Sheen = 0.f;
+			constexpr FLinearColor SheenColor{ 1, 1, 1 };
+			static constexpr float SheenRoughness = 0.3f;
+			static constexpr float Coat = 0.f;
+			constexpr FLinearColor CoatColor{ 1, 1, 1 };
+			static constexpr float CoatRoughness = 0.1f;
+			static constexpr float CoatAnisotropy = 0.f;
+			static constexpr float CoatRotation = 0.f;
+			static constexpr float CoatIOR = 1.5f;
+			static constexpr float CoatAffectColor = 0.f;
+			static constexpr float CoatAffectRoughness = 0.f;
+			static constexpr float ThinFilmThickness = 0.f;
+			static constexpr float ThinFilmIOR = 1.5f;
+			constexpr FLinearColor TransmissionColor{ 1, 1, 1 };
+			constexpr FLinearColor TransmissionScatter{ 0, 0, 0 };
+			static constexpr float Emission = 0.f;
+			constexpr FLinearColor EmissionColor{ 1, 1, 1 };
+			constexpr FLinearColor Opacity{ 1, 1, 1 };
 		}
 	}
 
diff --git a/Engine/Plugins/Interchange/Runtime/Source/Import/Private/MaterialX/MaterialXUtils/MaterialXStandardSurfaceShader.cpp b/Engine/Plugins/Interchange/Runtime/Source/Import/Private/MaterialX/MaterialXUtils/MaterialXStandardSurfaceShader.cpp
--- a/Engine/Plugins/Interchange/Runtime/Source/Import/Private/MaterialX/MaterialXUtils/MaterialXStandardSurfaceShader.cpp
+++ b/Engine/Plugins/Interchange/Runtime/Source/Import/Private/MaterialX/MaterialXUtils/MaterialXStandardSurfaceShader.cpp
@@ -52,76 +52,76 @@ void FMaterialXStandardSurfaceShader::ConnectToStandardSurface(UInterchangeFunct
 
 	// Inputs
 	//Base
-	ConnectNodeOutputToInput(Input::Base, StandardSurfaceShaderNode, StandardSurface::Parameters::Base.ToString(), DefaultValue::Float::Base);
+	ConnectNodeOutputToInput(Input::Base, StandardSurfaceShaderNode, StandardSurface::Parameters::Base.ToString(), DefaultValue::Base);
 
 	//Base Color
-	ConnectNodeOutputToInput(Input::BaseColor, StandardSurfaceShaderNode, StandardSurface::Parameters::BaseColor.ToString(), DefaultValue::Color3::BaseColor);
+	ConnectNodeOutputToInput(Input::BaseColor, StandardSurfaceShaderNode, StandardSurface::Parameters::BaseColor.ToString(), DefaultValue::BaseColor);
 
 	//Diffuse Roughness
-	ConnectNodeOutputToInput(Input::DiffuseRoughness, StandardSurfaceShaderNode, StandardSurface::Parameters::DiffuseRoughness.ToString(), DefaultValue::Float::DiffuseRoughness);
+	ConnectNodeOutputToInput(Input::DiffuseRoughness, StandardSurfaceShaderNode, StandardSurface::Parameters::DiffuseRoughness.ToString(), DefaultValue::DiffuseRoughness);
 
 	//Specular
-	ConnectNodeOutputToInput(Input::Specular, StandardSurfaceShaderNode, StandardSurface::Parameters::Specular.ToString(), DefaultValue::Float::Specular);
+	ConnectNodeOutputToInput(Input::Specular, StandardSurfaceShaderNode, StandardSurface::Parameters::Specular.ToString(), DefaultValue::Specular);
 
 	//Specular Roughness
-	ConnectNodeOutputToInput(Input::SpecularRoughness, StandardSurfaceShaderNode, StandardSurface::Parameters::SpecularRoughness.ToString(), DefaultValue::Float::SpecularRoughness);
+	ConnectNodeOutputToInput(Input::SpecularRoughness, StandardSurfaceShaderNode, StandardSurface::Parameters::SpecularRoughness.ToString(), DefaultValue::SpecularRoughness);
 
 	//Specular IOR
-	ConnectNodeOutputToInput(Input::SpecularIOR, StandardSurfaceShaderNode, StandardSurface::Parameters::SpecularIOR.ToString(), DefaultValue::Float::SpecularIOR);
+	ConnectNodeOutputToInput(Input::SpecularIOR, StandardSurfaceShaderNode, StandardSurface::Parameters::SpecularIOR.ToString(), DefaultValue::SpecularIOR);
 
 	//Specular Anisotropy
-	ConnectNodeOutputToInput(Input::SpecularAnisotropy, StandardSurfaceShaderNode, StandardSurface::Parameters::SpecularAnisotropy.ToString(), DefaultValue::Float::SpecularAnisotropy);
+	ConnectNodeOutputToInput(Input::SpecularAnisotropy, StandardSurfaceShaderNode, StandardSurface::Parameters::SpecularAnisotropy.ToString(), DefaultValue::SpecularAnisotropy);
 
 	//Specular Rotation
-	ConnectNodeOutputToInput(Input::SpecularRotation, StandardSurfaceShaderNode, StandardSurface::Parameters::SpecularRotation.ToString(), DefaultValue::Float::SpecularRotation);
+	ConnectNodeOutputToInput(Input::SpecularRotation, StandardSurfaceShaderNode, StandardSurface::Parameters::SpecularRotation.ToString(), DefaultValue::SpecularRotation);
 
 	//Metallic
-	ConnectNodeOutputToInput(Input::Metalness, StandardSurfaceShaderNode, StandardSurface::Parameters::Metalness.ToString(), DefaultValue::Float::Metalness);
+	ConnectNodeOutputToInput(Input::Metalness, StandardSurfaceShaderNode, StandardSurface::Parameters::Metalness.ToString(), DefaultValue::Metalness);
 
 	//Subsurface
-	ConnectNodeOutputToInput(Input::Subsurface, StandardSurfaceShaderNode, StandardSurface::Parameters::Subsurface.ToString(), DefaultValue::Float::Subsurface);
+	ConnectNodeOutputToInput(Input::Subsurface, StandardSurfaceShaderNode, StandardSurface::Parameters::Subsurface.ToString(), DefaultValue::Subsurface);
 
 	//Subsurface Color
-	ConnectNodeOutputToInput(Input::SubsurfaceColor, StandardSurfaceShaderNode, StandardSurface::Parameters::SubsurfaceColor.ToString(), DefaultValue::Color3::SubsurfaceColor);
+	ConnectNodeOutputToInput(Input::SubsurfaceColor, StandardSurfaceShaderNode, StandardSurface::Parameters::SubsurfaceColor.ToString(), DefaultValue::SubsurfaceColor);
 
 	//Subsurface Radius
-	ConnectNodeOutputToInput(Input::SubsurfaceRadius, StandardSurfaceShaderNode, StandardSurface::Parameters::SubsurfaceRadius.ToString(), DefaultValue::Color3::SubsurfaceRadius);
+	ConnectNodeOutputToInput(Input::SubsurfaceRadius, StandardSurfaceShaderNode, StandardSurface::Parameters::SubsurfaceRadius.ToString(), DefaultValue::SubsurfaceRadius);
 
 	//Subsurface Scale
-	ConnectNodeOutputToInput(Input::SubsurfaceScale, StandardSurfaceShaderNode, StandardSurface::Parameters::SubsurfaceScale.ToString(), DefaultValue::Float::SubsurfaceScale);
+	ConnectNodeOutputToInput(Input::SubsurfaceScale, StandardSurfaceShaderNode, StandardSurface::Parameters::SubsurfaceScale.ToString(), DefaultValue::SubsurfaceScale);
 
 	//Sheen
-	ConnectNodeOutputToInput(Input::Sheen, StandardSurfaceShaderNode, StandardSurface::Parameters::Sheen.ToString(), DefaultValue::Float::Sheen);
+	ConnectNodeOutputToInput(Input::Sheen, StandardSurfaceShaderNode, StandardSurface::Parameters::Sheen.ToString(), DefaultValue::Sheen);
 
 	//Sheen Color
-	ConnectNodeOutputToInput(Input::SheenColor, StandardSurfaceShaderNode, StandardSurface::Parameters::SheenColor.ToString(), DefaultValue::Color3::SheenColor);
+	ConnectNodeOutputToInput(Input::SheenColor, StandardSurfaceShaderNode, StandardSurface::Parameters::SheenColor.ToString(), DefaultValue::SheenColor);
 
 	//Sheen Roughness
-	ConnectNodeOutputToInput(Input::SheenRoughness, StandardSurfaceShaderNode, StandardSurface::Parameters::SheenRoughness.ToString(), DefaultValue::Float::SheenRoughness);
+	ConnectNodeOutputToInput(Input::SheenRoughness, StandardSurfaceShaderNode, StandardSurface::Parameters::SheenRoughness.ToString(), DefaultValue::SheenRoughness);
 
 	//Coat
-	ConnectNodeOutputToInput(Input::Coat, StandardSurfaceShaderNode, StandardSurface::Parameters::Coat.ToString(), DefaultValue::Float::Coat);
+	ConnectNodeOutputToInput(Input::Coat, StandardSurfaceShaderNode, StandardSurface::Parameters::Coat.ToString(), DefaultValue::Coat);
 
 	//Coat Color
-	ConnectNodeOutputToInput(Input::CoatColor, StandardSurfaceShaderNode, StandardSurface::Parameters::CoatColor.ToString(), DefaultValue::Color3::CoatColor);
+	ConnectNodeOutputToInput(Input::CoatColor, StandardSurfaceShaderNode, StandardSurface::Parameters::CoatColor.ToString(), DefaultValue::CoatColor);
 
 	//Coat Roughness
-	ConnectNodeOutputToInput(Input::CoatRoughness, StandardSurfaceShaderNode, StandardSurface::Parameters::CoatRoughness.ToString(), DefaultValue::Float::CoatRoughness);
+	ConnectNodeOutputToInput(Input::CoatRoughness, StandardSurfaceShaderNode, StandardSurface::Parameters::CoatRoughness.ToString(), DefaultValue::CoatRoughness);
 
 	//Coat Normal: No need to take the default input if there is no CoatNormal input
 	ConnectNodeOutputToInput(Input::CoatNormal, StandardSurfaceShaderNode, StandardSurface::Parameters::CoatNormal.ToString(), nullptr, bInputInTangentSpace);
 
 	//Thin Film Thickness
-	ConnectNodeOutputToInput(Input::ThinFilmThickness, StandardSurfaceShaderNode, StandardSurface::Parameters::ThinFilmThickness.ToString(), DefaultValue::Float::ThinFilmThickness);
+	ConnectNodeOutputToInput(Input::ThinFilmThickness, StandardSurfaceShaderNode, StandardSurface::Parameters::ThinFilmThickness.ToString(), DefaultValue::ThinFilmThickness);
 
 	//Thin Film IOR
-	ConnectNodeOutputToInput(Input::ThinFilmIOR, StandardSurfaceShaderNode, StandardSurface::Parameters::ThinFilmIOR.ToString(), DefaultValue::Float::ThinFilmIOR);
+	ConnectNodeOutputToInput(Input::ThinFilmIOR, StandardSurfaceShaderNode, StandardSurface::Parameters::ThinFilmIOR.ToString(), DefaultValue::ThinFilmIOR);
 
 	//Emission
-	ConnectNodeOutputToInput(Input::Emission, StandardSurfaceShaderNode, StandardSurface::Parameters::Emission.ToString(), DefaultValue::Float::Emission);
+	ConnectNodeOutputToInput(Input::Emission, StandardSurfaceShaderNode, StandardSurface::Parameters::Emission.ToString(), DefaultValue::Emission);
 
 	//Emission Color
-	ConnectNodeOutputToInput(Input::EmissionColor, StandardSurfaceShaderNode, StandardSurface::Parameters::EmissionColor.ToString(), DefaultValue::Color3::EmissionColor);
+	ConnectNodeOutputToInput(Input::EmissionColor, StandardSurfaceShaderNode, StandardSurface::Parameters::EmissionColor.ToString(), DefaultValue::EmissionColor);
 
 	//Normal: No need to take the default input if there is no Normal input
 	ConnectNodeOutputToInput(Input::Normal, StandardSurfaceShaderNode, StandardSurface::Parameters::Normal.ToString(), nullptr, bInputInTangentSpace);
@@ -130,28 +130,28 @@ void FMaterialXStandardSurfaceShader::ConnectToStandardSurface(UInterchangeFunct
 	ConnectNodeOutputToInput(Input::Tangent, StandardSurfaceShaderNode, StandardSurface::Parameters::Tangent.ToString(), nullptr, bInputInTangentSpace);
 
 	//Transmission
-	ConnectNodeOutputToInput(Input::Transmission, StandardSurfaceShaderNode, StandardSurface::Parameters::Transmission.ToString(), DefaultValue::Float::Transmission);
+	ConnectNodeOutputToInput(Input::Transmission, StandardSurfaceShaderNode, StandardSurface::Parameters::Transmission.ToString(), DefaultValue::Transmission);
 
 	//Transmission Color
-	ConnectNodeOutputToInput(Input::TransmissionColor, StandardSurfaceShaderNode, StandardSurface::Parameters::TransmissionColor.ToString(), DefaultValue::Color3::TransmissionColor);
+	ConnectNodeOutputToInput(Input::TransmissionColor, StandardSurfaceShaderNode, StandardSurface::Parameters::TransmissionColor.ToString(), DefaultValue::TransmissionColor);
 
 	//Transmission Depth
-	ConnectNodeOutputToInput(Input::TransmissionDepth, StandardSurfaceShaderNode, StandardSurface::Parameters::TransmissionDepth.ToString(), DefaultValue::Float::TransmissionDepth);
+	ConnectNodeOutputToInput(Input::TransmissionDepth, StandardSurfaceShaderNode, StandardSurface::Parameters::TransmissionDepth.ToString(), DefaultValue::TransmissionDepth);
 
 	//Transmission Scatter
-	ConnectNodeOutputToInput(Input::TransmissionScatter, StandardSurfaceShaderNode, StandardSurface::Parameters::TransmissionScatter.ToString(), DefaultValue::Color3::TransmissionScatter);
+	ConnectNodeOutputToInput(Input::TransmissionScatter, StandardSurfaceShaderNode, StandardSurface::Parameters::TransmissionScatter.ToString(), DefaultValue::TransmissionScatter);
 
 	//Transmission Scatter Anisotropy
-	ConnectNodeOutputToInput(Input::TransmissionScatterAnisotropy, StandardSurfaceShaderNode, StandardSurface::Parameters::TransmissionScatterAnisotropy.ToString(), DefaultValue::Float::TransmissionScatterAnisotropy);
+	ConnectNodeOutputToInput(Input::TransmissionScatterAnisotropy, StandardSurfaceShaderNode, StandardSurface::Parameters::TransmissionScatterAnisotropy.ToString(), DefaultValue::TransmissionScatterAnisotropy);
 
 	//Transmission Dispersion
-	ConnectNodeOutputToInput(Input::TransmissionDispersion, StandardSurfaceShaderNode, StandardSurface::Parameters::TransmissionDispersion.ToString(), DefaultValue::Float::TransmissionDispersion);
+	ConnectNodeOutputToInput(Input::TransmissionDispersion, StandardSurfaceShaderNode, StandardSurface::Parameters::TransmissionDispersion.ToString(), DefaultValue::TransmissionDispersion);
 
 	//Transmission Extra Roughness
-	ConnectNodeOutputToInput(Input::TransmissionExtraRoughness, StandardSurfaceShaderNode, StandardSurface::Parameters::TransmissionExtraRoughness.ToString(), DefaultValue::Float::TransmissionExtraRoughness);
+	ConnectNodeOutputToInput(Input::TransmissionExtraRoughness, StandardSurfaceShaderNode, StandardSurface::Parameters::TransmissionExtraRoughness.ToString(), DefaultValue::TransmissionExtraRoughness);
 
 	//Opacity
-	ConnectNodeOutputToInput(Input::Opacity, StandardSurfaceShaderNode, StandardSurface::Parameters::Opacity.ToString(), DefaultValue::Color3::Opacity);
+	ConnectNodeOutputToInput(Input::Opacity, StandardSurfaceShaderNode, StandardSurface::Parameters::Opacity.ToString(), DefaultValue::Opacity);
 
 	// Outputs
 	UInterchangeShaderPortsAPI::ConnectOuputToInputByName(ShaderGraphNode, PBRMR::Parameters::BaseColor.ToString(), StandardSurfaceShaderNode->GetUniqueID(), TEXT("Base Color"));
@@ -206,76 +206,76 @@ void FMaterialXStandardSurfaceShader::ConnectToSubstrateStandardSurface(UInterch
 
 	// Inputs
 	//Base
-	ConnectNodeOutputToInput(Input::Base, StandardSurfaceShaderNode, StandardSurface::Parameters::Base.ToString(), DefaultValue::Float::Base);
+	ConnectNodeOutputToInput(Input::Base, StandardSurfaceShaderNode, StandardSurface::Parameters::Base.ToString(), DefaultValue::Base);
 
 	//Base Color
-	ConnectNodeOutputToInput(Input::BaseColor, StandardSurfaceShaderNode, StandardSurface::Parameters::BaseColor.ToString(), DefaultValue::Color3::BaseColor);
+	ConnectNodeOutputToInput(Input::BaseColor, StandardSurfaceShaderNode, StandardSurface::Parameters::BaseColor.ToString(), DefaultValue::BaseColor);
 
 	//Diffuse Roughness
-	ConnectNodeOutputToInput(Input::DiffuseRoughness, StandardSurfaceShaderNode, StandardSurface::Parameters::DiffuseRoughness.ToString(), DefaultValue::Float::DiffuseRoughness);
+	ConnectNodeOutputToInput(Input::DiffuseRoughness, StandardSurfaceShaderNode, StandardSurface::Parameters::DiffuseRoughness.ToString(), DefaultValue::DiffuseRoughness);
 
 	//Metallic
-	ConnectNodeOutputToInput(Input::Metalness, StandardSurfaceShaderNode, StandardSurface::Parameters::Metalness.ToString(), DefaultValue::Float::Metalness);
+	ConnectNodeOutputToInput(Input::Metalness, StandardSurfaceShaderNode, StandardSurface::Parameters::Metalness.ToString(), DefaultValue::Metalness);
 
 	//Specular
-	ConnectNodeOutputToInput(Input::Specular, StandardSurfaceShaderNode, StandardSurface::Parameters::Specular.ToString(), DefaultValue::Float::Specular);
+	ConnectNodeOutputToInput(Input::Specular, StandardSurfaceShaderNode, StandardSurface::Parameters::Specular.ToString(), DefaultValue::Specular);
 
 	//Specular Color
-	ConnectNodeOutputToInput(Input::SpecularColor, StandardSurfaceShaderNode, StandardSurface::Parameters::SpecularColor.ToString(), DefaultValue::Color3::SpecularColor);
+	ConnectNodeOutputToInput(Input::SpecularColor, StandardSurfaceShaderNode, StandardSurface::Parameters::SpecularColor.ToString(), DefaultValue::SpecularColor);
 
 	//Specular Roughness
-	ConnectNodeOutputToInput(Input::SpecularRoughness, StandardSurfaceShaderNode, StandardSurface::Parameters::SpecularRoughness.ToString(), DefaultValue::Float::SpecularRoughness);
+	ConnectNodeOutputToInput(Input::SpecularRoughness, StandardSurfaceShaderNode, StandardSurface::Parameters::SpecularRoughness.ToString(), DefaultValue::SpecularRoughness);
 
 	//Specular IOR
-	ConnectNodeOutputToInput(Input::SpecularIOR, StandardSurfaceShaderNode, StandardSurface::Parameters::SpecularIOR.ToString(), DefaultValue::Float::SpecularIOR);
+	ConnectNodeOutputToInput(Input::SpecularIOR, StandardSurfaceShaderNode, StandardSurface::Parameters::SpecularIOR.ToString(), DefaultValue::SpecularIOR);
 
 	//Specular Anisotropy
-	ConnectNodeOutputToInput(Input::SpecularAnisotropy, StandardSurfaceShaderNode, StandardSurface::Parameters::SpecularAnisotropy.ToString(), DefaultValue::Float::SpecularAnisotropy);
+	ConnectNodeOutputToInput(Input::SpecularAnisotropy, StandardSurfaceShaderNode, StandardSurface::Parameters::SpecularAnisotropy.ToString(), DefaultValue::SpecularAnisotropy);
 
 	//Specular Rotation
-	ConnectNodeOutputToInput(Input::SpecularRotation, StandardSurfaceShaderNode, StandardSurface::Parameters::SpecularRotation.ToString(), DefaultValue::Float::SpecularRotation);
+	ConnectNodeOutputToInput(Input::SpecularRotation, StandardSurfaceShaderNode, StandardSurface::Parameters::SpecularRotation.ToString(), DefaultValue::SpecularRotation);
 
 	//Sheen
-	ConnectNodeOutputToInput(Input::Sheen, StandardSurfaceShaderNode, StandardSurface::Parameters::Sheen.ToString(), DefaultValue::Float::Sheen);
+	ConnectNodeOutputToInput(Input::Sheen, StandardSurfaceShaderNode, StandardSurface::Parameters::Sheen.ToString(), DefaultValue::Sheen);
 
 	//Sheen Color
-	ConnectNodeOutputToInput(Input::SheenColor, StandardSurfaceShaderNode, StandardSurface::Parameters::SheenColor.ToString(), DefaultValue::Color3::SheenColor);
+	ConnectNodeOutputToInput(Input::SheenColor, StandardSurfaceShaderNode, StandardSurface::Parameters::SheenColor.ToString(), DefaultValue::SheenColor);
 
 	//Sheen Roughness
-	ConnectNodeOutputToInput(Input::SheenRoughness, StandardSurfaceShaderNode, StandardSurface::Parameters::SheenRoughness.ToString(), DefaultValue::Float::SheenRoughness);
+	ConnectNodeOutputToInput(Input::SheenRoughness, StandardSurfaceShaderNode, StandardSurface::Parameters::SheenRoughness.ToString(), DefaultValue::SheenRoughness);
 
 	//Coat
-	ConnectNodeOutputToInput(Input::Coat, StandardSurfaceShaderNode, StandardSurface::Parameters::Coat.ToString(), DefaultValue::Float::Coat);
+	ConnectNodeOutputToInput(Input::Coat, StandardSurfaceShaderNode, StandardSurface::Parameters::Coat.ToString(), DefaultValue::Coat);
 
 	//Coat Color
-	ConnectNodeOutputToInput(Input::CoatColor, StandardSurfaceShaderNode, StandardSurface::Parameters::CoatColor.ToString(), DefaultValue::Color3::CoatColor);
+	ConnectNodeOutputToInput(Input::CoatColor, StandardSurfaceShaderNode, StandardSurface::Parameters::CoatColor.ToString(), DefaultValue::CoatColor);
 
 	//Coat Roughness
-	ConnectNodeOutputToInput(Input::CoatRoughness, StandardSurfaceShaderNode, StandardSurface::Parameters::CoatRoughness.ToString(), DefaultValue::Float::CoatRoughness);
+	ConnectNodeOutputToInput(Input::CoatRoughness, StandardSurfaceShaderNode, StandardSurface::Parameters::CoatRoughness.ToString(), DefaultValue::CoatRoughness);
 
 	//Coat IOR
-	ConnectNodeOutputToInput(Input::CoatIOR, StandardSurfaceShaderNode, StandardSurface::Parameters::CoatIOR.ToString(), DefaultValue::Float::CoatIOR);
+	ConnectNodeOutputToInput(Input::CoatIOR, StandardSurfaceShaderNode, StandardSurface::Parameters::CoatIOR.ToString(), DefaultValue::CoatIOR);
 
 	//Coat Anisotropy
-	ConnectNodeOutputToInput(Input::CoatAnisotropy, StandardSurfaceShaderNode, StandardSurface::Parameters::CoatAnisotropy.ToString(), DefaultValue::Float::CoatAnisotropy);
+	ConnectNodeOutputToInput(Input::CoatAnisotropy, StandardSurfaceShaderNode, StandardSurface::Parameters::CoatAnisotropy.ToString(), DefaultValue::CoatAnisotropy);
 
 	//Coat Rotation
-	ConnectNodeOutputToInput(Input::CoatRotation, StandardSurfaceShaderNode, StandardSurface::Parameters::CoatRotation.ToString(), DefaultValue::Float::CoatAnisotropy);
+	ConnectNodeOutputToInput(Input::CoatRotation, StandardSurfaceShaderNode, StandardSurface::Parameters::CoatRotation.ToString(), DefaultValue::CoatAnisotropy);
 
 	//Coat Normal: No need to take the default input if there is no CoatNormal input
 	ConnectNodeOutputToInput(Input::CoatNormal, StandardSurfaceShaderNode, StandardSurface::Parameters::CoatNormal.ToString(), nullptr, bInputInTangentSpace);
 
 	//Thin Film Thickness
-	ConnectNodeOutputToInput(Input::ThinFilmThickness, StandardSurfaceShaderNode, StandardSurface::Parameters::ThinFilmThickness.ToString(), DefaultValue::Float::ThinFilmThickness);
+	ConnectNodeOutputToInput(Input::ThinFilmThickness, StandardSurfaceShaderNode, StandardSurface::Parameters::ThinFilmThickness.ToString(), DefaultValue::ThinFilmThickness);
 
 	//Thin Film IOR
-	ConnectNodeOutputToInput(Input::ThinFilmIOR, StandardSurfaceShaderNode, StandardSurface::Parameters::ThinFilmIOR.ToString(), DefaultValue::Float::ThinFilmIOR);
+	ConnectNodeOutputToInput(Input::ThinFilmIOR, StandardSurfaceShaderNode, StandardSurface::Parameters::ThinFilmIOR.ToString(), DefaultValue::ThinFilmIOR);
 
 	//Emission
-	ConnectNodeOutputToInput(Input::Emission, StandardSurfaceShaderNode, StandardSurface::Parameters::Emission.ToString(), DefaultValue::Float::Emission);
+	ConnectNodeOutputToInput(Input::Emission, StandardSurfaceShaderNode, StandardSurface::Parameters::Emission.ToString(), DefaultValue::Emission);
 
 	//Emission Color
-	ConnectNodeOutputToInput(Input::EmissionColor, StandardSurfaceShaderNode, StandardSurface::Parameters::EmissionColor.ToString(), DefaultValue::Color3::EmissionColor);
+	ConnectNodeOutputToInput(Input::EmissionColor, StandardSurfaceShaderNode, StandardSurface::Parameters::EmissionColor.ToString(), DefaultValue::EmissionColor);
 
 	//Normal: No need to take the default input if there is no Normal input
 	ConnectNodeOutputToInput(Input::Normal, StandardSurfaceShaderNode, StandardSurface::Parameters::Normal.ToString(), nullptr, bInputInTangentSpace);
@@ -284,22 +284,22 @@ void FMaterialXStandardSurfaceShader::ConnectToSubstrateStandardSurface(UInterch
 	ConnectNodeOutputToInput(Input::Tangent, StandardSurfaceShaderNode, StandardSurface::Parameters::Tangent.ToString(), nullptr, bInputInTangentSpace);
 
 	//Opacity
-	ConnectNodeOutputToInput(Input::Opacity, StandardSurfaceShaderNode, StandardSurface::Parameters::Opacity.ToString(), DefaultValue::Color3::Opacity);
+	ConnectNodeOutputToInput(Input::Opacity, StandardSurfaceShaderNode, StandardSurface::Parameters::Opacity.ToString(), DefaultValue::Opacity);
 
 	//Transmission
-	ConnectNodeOutputToInput(Input::Transmission, StandardSurfaceShaderNode, StandardSurface::Parameters::Transmission.ToString(), DefaultValue::Float::Transmission);
+	ConnectNodeOutputToInput(Input::Transmission, StandardSurfaceShaderNode, StandardSurface::Parameters::Transmission.ToString(), DefaultValue::Transmission);
 
 	//Transmission Color
-	ConnectNodeOutputToInput(Input::TransmissionColor, StandardSurfaceShaderNode, StandardSurface::Parameters::TransmissionColor.ToString(), DefaultValue::Color3::TransmissionColor);
+	ConnectNodeOutputToInput(Input::TransmissionColor, StandardSurfaceShaderNode, StandardSurface::Parameters::TransmissionColor.ToString(), DefaultValue::TransmissionColor);
 
 	//Transmission Depth
-	ConnectNodeOutputToInput(Input::TransmissionDepth, StandardSurfaceShaderNode, StandardSurface::Parameters::TransmissionDepth.ToString(), DefaultValue::Float::TransmissionDepth);
+	ConnectNodeOutputToInput(Input::TransmissionDepth, StandardSurfaceShaderNode, StandardSurface::Parameters::TransmissionDepth.ToString(), DefaultValue::TransmissionDepth);
 
 	//Transmission Scatter
-	ConnectNodeOutputToInput(Input::TransmissionScatter, StandardSurfaceShaderNode, StandardSurface::Parameters::TransmissionScatter.ToString(), DefaultValue::Color3::TransmissionScatter);
+	ConnectNodeOutputToInput(Input::TransmissionScatter, StandardSurfaceShaderNode, StandardSurface::Parameters::TransmissionScatter.ToString(), DefaultValue::TransmissionScatter);
 
 	//Transmission Extra Roughness
-	ConnectNodeOutputToInput(Input::TransmissionExtraRoughness, StandardSurfaceShaderNode, StandardSurface::Parameters::TransmissionExtraRoughness.ToString(), DefaultValue::Float::TransmissionExtraRoughness);
+	ConnectNodeOutputToInput(Input::TransmissionExtraRoughness, StandardSurfaceShaderNode, StandardSurface::Parameters::TransmissionExtraRoughness.ToString(), DefaultValue::TransmissionExtraRoughness);
 
 	// Outputs
 	if(UInterchangeShaderPortsAPI::HasInput(StandardSurfaceShaderNode, StandardSurface::Parameters::Transmission))
@@ -311,19 +311,19 @@ void FMaterialXStandardSurfaceShader::ConnectToSubstrateStandardSurface(UInterch
 	else
 	{
 		//Subsurface
-		ConnectNodeOutputToInput(Input::Subsurface, StandardSurfaceShaderNode, StandardSurface::Parameters::Subsurface.ToString(), DefaultValue::Float::Subsurface);
+		ConnectNodeOutputToInput(Input::Subsurface, StandardSurfaceShaderNode, StandardSurface::Parameters::Subsurface.ToString(), DefaultValue::Subsurface);
 
 		//Subsurface Color
-		ConnectNodeOutputToInput(Input::SubsurfaceColor, StandardSurfaceShaderNode, StandardSurface::Parameters::SubsurfaceColor.ToString(), DefaultValue::Color3::SubsurfaceColor);
+		ConnectNodeOutputToInput(Input::SubsurfaceColor, StandardSurfaceShaderNode, StandardSurface::Parameters::SubsurfaceColor.ToString(), DefaultValue::SubsurfaceColor);
 
 		//Subsurface Radius
-		ConnectNodeOutputToInput(Input::SubsurfaceRadius, StandardSurfaceShaderNode, StandardSurface::Parameters::SubsurfaceRadius.ToString(), DefaultValue::Color3::SubsurfaceRadius);
+		ConnectNodeOutputToInput(Input::SubsurfaceRadius, StandardSurfaceShaderNode, StandardSurface::Parameters::SubsurfaceRadius.ToString(), DefaultValue::SubsurfaceRadius);
 
 		//Subsurface Scale
-		ConnectNodeOutputToInput(Input::SubsurfaceScale, StandardSurfaceShaderNode, StandardSurface::Parameters::SubsurfaceScale.ToString(), DefaultValue::Float::SubsurfaceScale);
+		ConnectNodeOutputToInput(Input::SubsurfaceScale, StandardSurfaceShaderNode, StandardSurface::Parameters::SubsurfaceScale.ToString(), DefaultValue::SubsurfaceScale);
 
 		//Subsurface Anisotropy
-		ConnectNodeOutputToInput(Input::SubsurfaceAnisotropy, StandardSurfaceShaderNode, StandardSurface::Parameters::SubsurfaceAnisotropy.ToString(), DefaultValue::Float::SubsurfaceAnisotropy);
+		ConnectNodeOutputToInput(Input::SubsurfaceAnisotropy, StandardSurfaceShaderNode, StandardSurface::Parameters::SubsurfaceAnisotropy.ToString(), DefaultValue::SubsurfaceAnisotropy);
 
 		UInterchangeShaderPortsAPI::ConnectOuputToInputByName(ShaderGraphNode, SubstrateMaterial::Parameters::FrontMaterial.ToString(), StandardSurfaceShaderNode->GetUniqueID(), StandardSurface::SubstrateMaterial::Outputs::Opaque.ToString());
 		if(UInterchangeShaderPortsAPI::HasInput(StandardSurfaceShaderNode, StandardSurface::Parameters::Opacity))
