MaterialX:
- StandardSurface: adding support of thin_walled
- OpenPBR: adding support of geometry_thin_walled
- Add support of two sided materials, we only evaluate the Two Sided option only if the input has a value attribute, otherwise it doesn't really have a meaning since the two sided option has to be set in the detail panel
- Edit tests to count the newly thin_walled input

#virtualized

diff --git a/Engine/Plugins/Interchange/Runtime/Source/Common/Public/MaterialX/InterchangeMaterialXDefinitions.h b/Engine/Plugins/Interchange/Runtime/Source/Common/Public/MaterialX/InterchangeMaterialXDefinitions.h
--- a/Engine/Plugins/Interchange/Runtime/Source/Common/Public/MaterialX/InterchangeMaterialXDefinitions.h
+++ b/Engine/Plugins/Interchange/Runtime/Source/Common/Public/MaterialX/InterchangeMaterialXDefinitions.h
@@ -318,6 +318,7 @@ MATERIALX_NAMESPACE_BEGIN
 			static constexpr float Emission = 0.f;
 			constexpr FLinearColor EmissionColor{ 1, 1, 1 };
 			constexpr FLinearColor Opacity{ 1, 1, 1 };
+			constexpr bool ThinWalled = false;
 		}
 	}
 
diff --git a/Engine/Plugins/Interchange/Runtime/Source/Import/Private/MaterialX/MaterialXUtils/MaterialXOpenPBRSurfaceShader.cpp b/Engine/Plugins/Interchange/Runtime/Source/Import/Private/MaterialX/MaterialXUtils/MaterialXOpenPBRSurfaceShader.cpp
--- a/Engine/Plugins/Interchange/Runtime/Source/Import/Private/MaterialX/MaterialXUtils/MaterialXOpenPBRSurfaceShader.cpp
+++ b/Engine/Plugins/Interchange/Runtime/Source/Import/Private/MaterialX/MaterialXUtils/MaterialXOpenPBRSurfaceShader.cpp
@@ -78,7 +78,17 @@ void FMaterialXOpenPBRSurfaceShader::Translate(MaterialX::NodePtr OpenPBRSurface
 	ConnectNodeOutputToInput(Input::GeometryOpacity, OpenPBRSurfaceShaderNode, OpenPBRSurface::Parameters::GeometryOpacity.ToString(), DefaultValue::GeometryOpacity);
 	ConnectNodeOutputToInput(Input::GeometryTangent, OpenPBRSurfaceShaderNode, OpenPBRSurface::Parameters::GeometryTangent.ToString(), DefaultValue::GeometryTangent, bTangentSpace);
 	ConnectNodeOutputToInput(Input::GeometryThinWalled, OpenPBRSurfaceShaderNode, OpenPBRSurface::Parameters::GeometryThinWalled.ToString(), DefaultValue::GeometryThinWalled);
-	
+	{
+		MaterialX::DocumentPtr Document = SurfaceShaderNode->getDocument();
+		MaterialX::InputPtr Input = GetInput(SurfaceShaderNode, mx::OpenPBRSurface::Input::GeometryThinWalled);
+		if(Input->hasValue() && mx::fromValueString<bool>(Input->getValueString()) == true)
+		{
+			// weird that we also have to enable that to have a two sided material (seems to only have meaning for Translucent material)
+			ShaderGraphNode->SetCustomTwoSidedTransmission(true);
+			ShaderGraphNode->SetCustomTwoSided(true);
+		}
+	}
+
 	//Fuzz
 	ConnectNodeOutputToInput(Input::FuzzColor, OpenPBRSurfaceShaderNode, OpenPBRSurface::Parameters::FuzzColor.ToString(), DefaultValue::FuzzColor);
 	ConnectNodeOutputToInput(Input::FuzzRoughness, OpenPBRSurfaceShaderNode, OpenPBRSurface::Parameters::FuzzRoughness.ToString(), DefaultValue::FuzzRoughness);
diff --git a/Engine/Plugins/Interchange/Runtime/Source/Import/Private/MaterialX/MaterialXUtils/MaterialXStandardSurfaceShader.cpp b/Engine/Plugins/Interchange/Runtime/Source/Import/Private/MaterialX/MaterialXUtils/MaterialXStandardSurfaceShader.cpp
--- a/Engine/Plugins/Interchange/Runtime/Source/Import/Private/MaterialX/MaterialXUtils/MaterialXStandardSurfaceShader.cpp
+++ b/Engine/Plugins/Interchange/Runtime/Source/Import/Private/MaterialX/MaterialXUtils/MaterialXStandardSurfaceShader.cpp
@@ -153,6 +153,17 @@ void FMaterialXStandardSurfaceShader::ConnectToStandardSurface(UInterchangeFunct
 	//Opacity
 	ConnectNodeOutputToInput(Input::Opacity, StandardSurfaceShaderNode, StandardSurface::Parameters::Opacity.ToString(), DefaultValue::Opacity);
 
+	//Thin Walled: we create a shader graph for that input even if it has no real-meaning, we only check the value to enable the Two-Sided option
+	ConnectNodeOutputToInput(Input::ThinWalled, StandardSurfaceShaderNode, StandardSurface::Parameters::ThinWalled.ToString(), DefaultValue::ThinWalled);
+	{
+		MaterialX::DocumentPtr Document = SurfaceShaderNode->getDocument();
+		MaterialX::InputPtr Input = GetInput(SurfaceShaderNode, mx::StandardSurface::Input::ThinWalled);
+		if(Input->hasValue() && mx::fromValueString<bool>(Input->getValueString()) == true)
+		{
+			ShaderGraphNode->SetCustomTwoSided(true);
+		}
+	}
+
 	// Outputs
 	UInterchangeShaderPortsAPI::ConnectOuputToInputByName(ShaderGraphNode, PBRMR::Parameters::BaseColor.ToString(), StandardSurfaceShaderNode->GetUniqueID(), TEXT("Base Color"));
 	UInterchangeShaderPortsAPI::ConnectOuputToInputByName(ShaderGraphNode, PBRMR::Parameters::Metallic.ToString(), StandardSurfaceShaderNode->GetUniqueID(), PBRMR::Parameters::Metallic.ToString());
@@ -301,6 +312,19 @@ void FMaterialXStandardSurfaceShader::ConnectToSubstrateStandardSurface(UInterch
 	//Transmission Extra Roughness
 	ConnectNodeOutputToInput(Input::TransmissionExtraRoughness, StandardSurfaceShaderNode, StandardSurface::Parameters::TransmissionExtraRoughness.ToString(), DefaultValue::TransmissionExtraRoughness);
 
+	//Thin Walled: we create a shader graph for that input even if it has no real-meaning, we only check the value to enable the Two-Sided option
+	ConnectNodeOutputToInput(Input::ThinWalled, StandardSurfaceShaderNode, StandardSurface::Parameters::ThinWalled.ToString(), DefaultValue::ThinWalled);
+	{
+		MaterialX::DocumentPtr Document = SurfaceShaderNode->getDocument();
+		MaterialX::InputPtr Input = GetInput(SurfaceShaderNode, mx::StandardSurface::Input::ThinWalled);
+		if(Input->hasValue() && mx::fromValueString<bool>(Input->getValueString()) == true)
+		{
+			// weird that we also have to enable that to have a two sided material (seems to only have meaning for Translucent material)
+			ShaderGraphNode->SetCustomTwoSidedTransmission(true);
+			ShaderGraphNode->SetCustomTwoSided(true);
+		}
+	}
+
 	// Outputs
 	if(UInterchangeShaderPortsAPI::HasInput(StandardSurfaceShaderNode, StandardSurface::Parameters::Transmission))
 	{
diff --git a/Engine/Plugins/Interchange/Runtime/Source/Import/Public/MaterialX/MaterialXUtils/MaterialXSurfaceShaderAbstract.h b/Engine/Plugins/Interchange/Runtime/Source/Import/Public/MaterialX/MaterialXUtils/MaterialXSurfaceShaderAbstract.h
--- a/Engine/Plugins/Interchange/Runtime/Source/Import/Public/MaterialX/MaterialXUtils/MaterialXSurfaceShaderAbstract.h
+++ b/Engine/Plugins/Interchange/Runtime/Source/Import/Public/MaterialX/MaterialXUtils/MaterialXSurfaceShaderAbstract.h
@@ -147,6 +147,10 @@ protected:
 				{					
 					bIsConnected = AddVectorAttribute(Input, InputShaderName, ShaderNode, DefaultValue);
 				}
+				if constexpr(std::is_same_v<decltype(DefaultValue), bool>)
+				{
+					bIsConnected = AddBooleanAttribute(Input, InputShaderName, ShaderNode);
+				}
 			}
 		}
 
diff --git a/Engine/Plugins/Interchange/Runtime/Source/Pipelines/Private/InterchangeMaterialXPipeline.cpp b/Engine/Plugins/Interchange/Runtime/Source/Pipelines/Private/InterchangeMaterialXPipeline.cpp
--- a/Engine/Plugins/Interchange/Runtime/Source/Pipelines/Private/InterchangeMaterialXPipeline.cpp
+++ b/Engine/Plugins/Interchange/Runtime/Source/Pipelines/Private/InterchangeMaterialXPipeline.cpp
@@ -104,6 +104,7 @@ UMaterialXPipelineSettings::UMaterialXPipelineSettings()
 					UE::Interchange::Materials::OpenPBRSurface::Parameters::GeometryNormal,
 					UE::Interchange::Materials::OpenPBRSurface::Parameters::GeometryTangent,
 					UE::Interchange::Materials::OpenPBRSurface::Parameters::GeometryOpacity,
+					UE::Interchange::Materials::OpenPBRSurface::Parameters::GeometryThinWalled,
 				},
 				// OpenPBRSurface Outputs
 				!bIsSubstrateEnabled ?
@@ -171,6 +172,7 @@ UMaterialXPipelineSettings::UMaterialXPipelineSettings()
 						UE::Interchange::Materials::OpenPBRSurface::Parameters::GeometryNormal,
 						UE::Interchange::Materials::OpenPBRSurface::Parameters::GeometryTangent,
 						UE::Interchange::Materials::OpenPBRSurface::Parameters::GeometryOpacity,
+						UE::Interchange::Materials::OpenPBRSurface::Parameters::GeometryThinWalled,
 					},
 					// OpenPBRSurfaceTransmission Outputs
 					!bIsSubstrateEnabled ?
@@ -219,12 +221,14 @@ UMaterialXPipelineSettings::UMaterialXPipelineSettings()
 						UE::Interchange::Materials::StandardSurface::Parameters::CoatRoughness,
 						UE::Interchange::Materials::StandardSurface::Parameters::CoatNormal,
 						UE::Interchange::Materials::StandardSurface::Parameters::ThinFilmThickness,
+						UE::Interchange::Materials::StandardSurface::Parameters::ThinFilmIOR,
 						UE::Interchange::Materials::StandardSurface::Parameters::Emission,
 						UE::Interchange::Materials::StandardSurface::Parameters::EmissionColor,
 						UE::Interchange::Materials::StandardSurface::Parameters::Normal,
 						UE::Interchange::Materials::StandardSurface::Parameters::Tangent,
+						UE::Interchange::Materials::StandardSurface::Parameters::ThinWalled,
 						bIsSubstrateEnabled ? UE::Interchange::Materials::StandardSurface::Parameters::SpecularColor : FName{},
-						bIsSubstrateEnabled ? UE::Interchange::Materials::StandardSurface::Parameters::CoatIOR : FName{},
+						UE::Interchange::Materials::StandardSurface::Parameters::CoatIOR,
 						bIsSubstrateEnabled ? UE::Interchange::Materials::StandardSurface::Parameters::CoatAnisotropy : FName{},
 						bIsSubstrateEnabled ? UE::Interchange::Materials::StandardSurface::Parameters::CoatRotation : FName{},
 						bIsSubstrateEnabled ? UE::Interchange::Materials::StandardSurface::Parameters::ThinFilmIOR : FName{},
@@ -290,9 +294,10 @@ UMaterialXPipelineSettings::UMaterialXPipelineSettings()
 						UE::Interchange::Materials::StandardSurface::Parameters::CoatRoughness,
 						bIsSubstrateEnabled ? UE::Interchange::Materials::StandardSurface::Parameters::CoatAnisotropy : FName{},
 						bIsSubstrateEnabled ? UE::Interchange::Materials::StandardSurface::Parameters::CoatRotation : FName{},
-						bIsSubstrateEnabled ? UE::Interchange::Materials::StandardSurface::Parameters::CoatIOR : FName{},
+						UE::Interchange::Materials::StandardSurface::Parameters::CoatIOR,
 						UE::Interchange::Materials::StandardSurface::Parameters::CoatNormal,
 						UE::Interchange::Materials::StandardSurface::Parameters::ThinFilmThickness,
+						UE::Interchange::Materials::StandardSurface::Parameters::ThinFilmIOR,
 						bIsSubstrateEnabled ? UE::Interchange::Materials::StandardSurface::Parameters::ThinFilmIOR : FName{},
 						UE::Interchange::Materials::StandardSurface::Parameters::Emission,
 						UE::Interchange::Materials::StandardSurface::Parameters::EmissionColor,
