Interchange:
- Fix reimport of material instances

GenericMaterialPipeline:
- Enable material factory for a reimport of material instances (materials are left untouched since we don't support yet reimport)
- Apply a reimport strategy for editor properties

Material factory:
- Override inputs when there's a reimport

diff --git a/Engine/Plugins/Interchange/Runtime/Source/Import/Private/Material/InterchangeMaterialFactory.cpp b/Engine/Plugins/Interchange/Runtime/Source/Import/Private/Material/InterchangeMaterialFactory.cpp
--- a/Engine/Plugins/Interchange/Runtime/Source/Import/Private/Material/InterchangeMaterialFactory.cpp
+++ b/Engine/Plugins/Interchange/Runtime/Source/Import/Private/Material/InterchangeMaterialFactory.cpp
@@ -884,31 +884,27 @@ UInterchangeFactoryBase::FImportAssetResult UInterchangeMaterialFactory::BeginIm
 	{
 		if (UMaterialInstance* MaterialInstance = Cast<UMaterialInstance>(Material))
 		{
-			const EReimportStrategyFlags ReimportStrategyFlags = MaterialFactoryNode->GetReimportStrategyFlags();
+			EReimportStrategyFlags ReimportStrategyFlags = MaterialFactoryNode->GetReimportStrategyFlags();
 			bool bApplyPipelineProperties = !Arguments.ReimportObject || ReimportStrategyFlags != EReimportStrategyFlags::ApplyNoProperties;
 
 #if WITH_EDITORONLY_DATA
 			// For the time being, reimport policies are only enforced on UMaterialInstanceConstant
 			if (Arguments.ReimportObject && MaterialInstance->IsA<UMaterialInstanceConstant>())
 			{
-
-				if (ReimportStrategyFlags == EReimportStrategyFlags::ApplyEditorChangedProperties)
+				if(UInterchangeAssetImportData* AssetImportData = Cast<UInterchangeAssetImportData>(MaterialInstance->AssetImportData))
 				{
-					if (UInterchangeAssetImportData* AssetImportData = Cast<UInterchangeAssetImportData>(MaterialInstance->AssetImportData))
-					{
-						UInterchangeMaterialInstanceFactoryNode* PreviousNode = Cast<UInterchangeMaterialInstanceFactoryNode>(AssetImportData->GetStoredFactoryNode(AssetImportData->NodeUniqueID));
+					UInterchangeMaterialInstanceFactoryNode* PreviousNode = Cast<UInterchangeMaterialInstanceFactoryNode>(AssetImportData->GetStoredFactoryNode(AssetImportData->NodeUniqueID));
 
-						if (PreviousNode)
+					if(PreviousNode)
+					{
+						FString PreviousParentPath;
+						FString ParentPath;
+						if(PreviousNode->GetCustomParent(PreviousParentPath) && MaterialInstanceFactoryNode->GetCustomParent(ParentPath))
 						{
-							FString PreviousParentPath;
-							FString ParentPath;
-							if (PreviousNode->GetCustomParent(PreviousParentPath) && MaterialInstanceFactoryNode->GetCustomParent(ParentPath))
+							if(ParentPath == PreviousParentPath)
 							{
-								if (ParentPath == PreviousParentPath)
-								{
-									SetupReimportedMaterialInstance(*MaterialInstance, *Arguments.NodeContainer, *MaterialInstanceFactoryNode, *PreviousNode);
-									bApplyPipelineProperties = false;
-								}
+								SetupReimportedMaterialInstance(*MaterialInstance, *Arguments.NodeContainer, *MaterialInstanceFactoryNode, *PreviousNode);
+								bApplyPipelineProperties = false;
 							}
 						}
 					}
@@ -1827,7 +1823,7 @@ void UInterchangeMaterialFactory::SetupReimportedMaterialInstance(UMaterialInsta
 		const FString AttributKey = bIsAParameter ? UInterchangeShaderPortsAPI::MakeInputParameterKey(InputName) : UInterchangeShaderPortsAPI::MakeInputValueKey(InputName);
 
 		FGuid Uid;
-		switch (UInterchangeShaderPortsAPI::GetInputType(&FactoryNode, InputName))
+		switch (UInterchangeShaderPortsAPI::GetInputType(&FactoryNode, InputName, bIsAParameter))
 		{
 #if WITH_EDITORONLY_DATA
 		case UE::Interchange::EAttributeTypes::Bool:
diff --git a/Engine/Plugins/Interchange/Runtime/Source/Pipelines/Private/InterchangeGenericMaterialPipeline.cpp b/Engine/Plugins/Interchange/Runtime/Source/Pipelines/Private/InterchangeGenericMaterialPipeline.cpp
--- a/Engine/Plugins/Interchange/Runtime/Source/Pipelines/Private/InterchangeGenericMaterialPipeline.cpp
+++ b/Engine/Plugins/Interchange/Runtime/Source/Pipelines/Private/InterchangeGenericMaterialPipeline.cpp
@@ -833,13 +833,15 @@ void UInterchangeGenericMaterialPipeline::ExecutePipeline(UInterchangeBaseNodeCo
 				if (UMaterialInterface* ExistingMaterial = UE::Interchange::InterchangeGenericMaterialPipeline::Private::FindExistingMaterialFromSearchLocation(MaterialName, ContentBasePath, ClosureSearchLocation))
 				{
 					//Make sure we have the correct type of material (can be material instance) before setting the custom object reference.
-					if ((MaterialBaseFactoryNode->IsA<UInterchangeMaterialInstanceFactoryNode>() && ExistingMaterial->IsA<UMaterialInstance>())
-						|| (MaterialBaseFactoryNode->IsA<UInterchangeMaterialFactoryNode>() && ExistingMaterial->IsA<UMaterial>()))
+					const bool bIsMaterial = MaterialBaseFactoryNode->IsA<UInterchangeMaterialFactoryNode>() && ExistingMaterial->IsA<UMaterial>();
+					const bool bIsMaterialinstance = MaterialBaseFactoryNode->IsA<UInterchangeMaterialInstanceFactoryNode>() && ExistingMaterial->IsA<UMaterialInstance>();
+					if (bIsMaterial || bIsMaterialinstance)
 					{
 						MaterialBaseFactoryNode->SetCustomReferenceObject(ExistingMaterial);
-						//No need to import an existing material
-						MaterialBaseFactoryNode->SetCustomIsMaterialImportEnabled(false);
-						MaterialBaseFactoryNode->SetEnabled(false);
+						//Reimport can only be done on material instances
+						const bool bEnableReimport = !bIsMaterial && bIsMaterialinstance;
+						MaterialBaseFactoryNode->SetCustomIsMaterialImportEnabled(bEnableReimport);
+						MaterialBaseFactoryNode->SetEnabled(bEnableReimport);
 					}
 				}
 			}
