[Interchange] Fixing Animation type conversions.

#rnx

diff --git a/Engine/Plugins/Interchange/Runtime/Source/Import/Private/Animation/InterchangeAnimSequenceFactory.cpp b/Engine/Plugins/Interchange/Runtime/Source/Import/Private/Animation/InterchangeAnimSequenceFactory.cpp
--- a/Engine/Plugins/Interchange/Runtime/Source/Import/Private/Animation/InterchangeAnimSequenceFactory.cpp
+++ b/Engine/Plugins/Interchange/Runtime/Source/Import/Private/Animation/InterchangeAnimSequenceFactory.cpp
@@ -498,7 +498,7 @@ namespace UE::Interchange::Private
 					break;
 				}
 
-				if (AnimationTransformPayload.Type == EInterchangeAnimationPayLoadType::BAKED)
+				if (AnimationTransformPayload.PayloadKey.Type == EInterchangeAnimationPayLoadType::BAKED)
 				{
 					//Everything should match Key count, sample rate and range
 					if (!(ensure(FMath::IsNearlyEqual(AnimationTransformPayload.BakeFrequency, SampleRate, UE_DOUBLE_KINDA_SMALL_NUMBER)) &&
@@ -1122,7 +1122,7 @@ UInterchangeFactoryBase::FImportAssetResult UInterchangeAnimSequenceFactory::Imp
 
 			UE::Interchange::FAnimationPayloadData& AnimationTransformPayload = AnimationPayload;
 
-			if (AnimationTransformPayload.Type != EInterchangeAnimationPayLoadType::BAKED)
+			if (AnimationTransformPayload.PayloadKey.Type != EInterchangeAnimationPayLoadType::BAKED)
 			{
 				//Where Curve is null the LocalTransform should be used for the Baked Transform generation.
 				FTransform LocalTransform;
diff --git a/Engine/Plugins/Interchange/Runtime/Source/Parsers/CommonParser/Private/InterchangeCommonAnimationPayload.cpp b/Engine/Plugins/Interchange/Runtime/Source/Parsers/CommonParser/Private/InterchangeCommonAnimationPayload.cpp
--- a/Engine/Plugins/Interchange/Runtime/Source/Parsers/CommonParser/Private/InterchangeCommonAnimationPayload.cpp
+++ b/Engine/Plugins/Interchange/Runtime/Source/Parsers/CommonParser/Private/InterchangeCommonAnimationPayload.cpp
@@ -46,7 +46,7 @@ namespace UE::Interchange
 	void FAnimationPayloadData::CalculateDataFor(const EInterchangeAnimationPayLoadType& ToType, const FTransform& DefaultTransform)
 	{
 #if WITH_ENGINE
-		if (Type == EInterchangeAnimationPayLoadType::CURVE
+		if (PayloadKey.Type == EInterchangeAnimationPayLoadType::CURVE
 			&& ToType == EInterchangeAnimationPayLoadType::STEPCURVE)
 		{
 
@@ -68,7 +68,7 @@ namespace UE::Interchange
 			}
 			AdditionalSupportedType = ToType;
 		}
-		else if (Type == EInterchangeAnimationPayLoadType::CURVE
+		else if (PayloadKey.Type == EInterchangeAnimationPayLoadType::CURVE
 				 && ToType == EInterchangeAnimationPayLoadType::BAKED)
 		{
 			if (Curves.Num() != 9)
diff --git a/Engine/Plugins/Interchange/Runtime/Source/Parsers/CommonParser/Public/InterchangeCommonAnimationPayload.h b/Engine/Plugins/Interchange/Runtime/Source/Parsers/CommonParser/Public/InterchangeCommonAnimationPayload.h
--- a/Engine/Plugins/Interchange/Runtime/Source/Parsers/CommonParser/Public/InterchangeCommonAnimationPayload.h
+++ b/Engine/Plugins/Interchange/Runtime/Source/Parsers/CommonParser/Public/InterchangeCommonAnimationPayload.h
@@ -47,9 +47,8 @@ namespace UE::Interchange
 		double RangeEndTime = 1.0 / BakeFrequency;
 		TArray<FTransform> Transforms;
 
-		//TYPE
+		//PayloadKey related:
 		FInterchangeAnimationPayLoadKey PayloadKey;
-		EInterchangeAnimationPayLoadType Type = EInterchangeAnimationPayLoadType::NONE; //Original
 		EInterchangeAnimationPayLoadType AdditionalSupportedType = EInterchangeAnimationPayLoadType::NONE;
 
 		//
