TextureImportSettings AutoLimitDimension option to resize textures on import to limit texture source sizes

diff --git a/Engine/Plugins/Interchange/Runtime/Source/Import/Private/Texture/InterchangeTextureFactory.cpp b/Engine/Plugins/Interchange/Runtime/Source/Import/Private/Texture/InterchangeTextureFactory.cpp
--- a/Engine/Plugins/Interchange/Runtime/Source/Import/Private/Texture/InterchangeTextureFactory.cpp
+++ b/Engine/Plugins/Interchange/Runtime/Source/Import/Private/Texture/InterchangeTextureFactory.cpp
@@ -741,8 +741,260 @@ namespace UE::Interchange::Private::InterchangeTextureFactory
 	void SetupTextureSourceDataFromBulkData(UTexture* Texture, const FImportImage& Image, UE::Serialization::FEditorBulkData::FSharedBufferWithID&& BufferAndId, bool bIsReimport);
 
 #if WITH_EDITOR
-	void SetupTextureSourceDataFromBulkData_Editor(UTexture* Texture, const FImportImage& Image, UE::Serialization::FEditorBulkData::FSharedBufferWithID&& BufferAndId, bool bIsReimport)
+	
+	// "Nearest" is in log scale (eg. in ratio sense) not linear value difference
+	//	RoundToNearestPowerOfTwo32(0) == 1
+	static inline uint64 RoundToNearestPowerOfTwo32(const uint32 x)
 	{
+		uint64 up = FMath::RoundUpToPowerOfTwo64(x);
+		uint64 dn = up>>1;
+	
+		// input arg is 32 bit so squares fit in 64
+		return ( up*dn <= x*x ) ? up : dn;
+	}
+
+	// Image and BufferAndId will be read from and mutated if a resize is done
+	void ResizeImportImageIfNeeded(UTexture* Texture, FImportImage& Image, int64 NumSlices, UE::Serialization::FEditorBulkData::FSharedBufferWithID & BufferAndId)
+	{
+		// Image.RawData ignored; use BufferAndId instead for payload bytes
+		//	NumSlices not changed by import (no volumes here)
+		Image.RawData.Reset();
+
+		if ( Texture->VirtualTextureStreaming )
+		{
+			// may already be set for a reimport
+			return;
+		}
+
+		int64 LimitPixelCount = GetDefault<UTextureImportSettings>()->GetAutoLimitPixelCount();
+		if ( LimitPixelCount == 0 )
+		{
+			return;
+		}
+
+		int64 InitialPixelCount = (int64) Image.SizeX * Image.SizeY * NumSlices;
+		if ( InitialPixelCount <= LimitPixelCount )
+		{
+			return;
+		}
+
+		// ensure resizing to 1x1xNumSlices will get us under LimitPixelCount
+		int64 MinLimitPixelCount = 16384;
+		check( NumSlices <= MinLimitPixelCount );
+		if ( LimitPixelCount < MinLimitPixelCount )
+		{
+			UE_LOG(LogInterchangeImport, Warning, TEXT("Tiny LimitPixelCount %lld changed to %lld ."), LimitPixelCount, MinLimitPixelCount );
+			LimitPixelCount = MinLimitPixelCount;
+		}
+
+		if ( Texture->GetTextureClass() == ETextureClass::TwoD && NumSlices == 1 && GetDefault<UTextureImportSettings>()->IsImportAutoVTEnabled() )
+		{
+			// see ShouldTextureBeVirtualByAutoImportSize
+			Texture->VirtualTextureStreaming = true;
+			return;
+		}
+
+		// texture is not VT and exceeds pixel count limit, do resize
+		
+		const int64 BytesPerPixel = FTextureSource::GetBytesPerPixel(Image.Format);
+
+		if ( Image.NumMips > 1 )
+		{
+			// for existing-mips, handle by stepping down the existing mips first
+			//	existing mips must either be full mip chains or just one
+			
+			check( Image.RawDataCompressionFormat == TSCF_None );
+
+			int64 MipIndex = 0;
+			int64 NewSizeX = Image.SizeX;
+			int64 NewSizeY = Image.SizeY;
+			int64 BytesDiscarded = 0;
+			while( NewSizeX * NewSizeY * NumSlices > LimitPixelCount )
+			{
+				BytesDiscarded += NewSizeX * NewSizeY * NumSlices * BytesPerPixel;
+
+				MipIndex ++;				
+				NewSizeX = FMath::Max<int64>(NewSizeX>>1,1);
+				NewSizeY = FMath::Max<int64>(NewSizeY>>1,1);
+			}
+			
+			FSharedBuffer SrcBuffer = BufferAndId.GetPayload();
+			FMemoryView SrcView = SrcBuffer.GetView();
+
+			check( SrcView.GetSize() == Image.ComputeBufferSize() );
+
+			if ( MipIndex < Image.NumMips )
+			{
+				// just truncate the large mips off the front of the buffer
+				
+				UE_LOG(LogInterchangeImport, Display, TEXT("Image %s larger than GetAutoLimitPixelCount.  Dropping existing mips: %d x %d -> %d x %d (%d slices)"),
+					*Texture->GetName(),
+					Image.SizeX,Image.SizeY,
+					NewSizeX,NewSizeY,
+					NumSlices);
+
+				Image.SizeX = NewSizeX;
+				Image.SizeY = NewSizeY;
+				Image.NumMips -= MipIndex;
+
+				check( (int64)SrcView.GetSize() > BytesDiscarded );
+				SrcView = SrcView.RightChop( BytesDiscarded );
+				
+				check( SrcView.GetSize() == Image.ComputeBufferSize() );
+
+				BufferAndId = UE::Serialization::FEditorBulkData::FSharedBufferWithID( FSharedBuffer::Clone(SrcView) );
+
+				// done just by stepping down existing mips, no actual resizing needed
+				return;
+			}
+			else
+			{
+				// didn't have enough mips to step down
+				// just change to 1 mip only and resize from there
+				
+				UE_LOG(LogInterchangeImport, Display, TEXT("Image %s has existing mips but not a full set, discarding."),
+					*Texture->GetName() );
+
+				Image.NumMips = 1;
+				SrcView = SrcView.LeftChop( Image.SizeX * Image.SizeY * NumSlices * BytesPerPixel );
+				
+				check( SrcView.GetSize() == Image.ComputeBufferSize() );
+
+				BufferAndId = UE::Serialization::FEditorBulkData::FSharedBufferWithID( FSharedBuffer::Clone(SrcView) );
+
+				// continue to resizing
+			}
+		}
+
+		// since we are resizing anyway, may as well change to pow2 also
+		// note that we change both dimensions to pow2
+		//	 which can change the aspect ratio if the source was not square or pow2
+
+		int64 LargerDim  = FMath::Max<int64>(Image.SizeX,Image.SizeY);
+		int64 SmallerDim = FMath::Min<int64>(Image.SizeX,Image.SizeY);
+		double OriginalDimRatio = (double) SmallerDim / LargerDim;
+
+		LargerDim = FMath::RoundUpToPowerOfTwo64(LargerDim);
+
+		for(;;)
+		{
+			LargerDim >>= 1;
+			SmallerDim = RoundToNearestPowerOfTwo32( FMath::RoundToInt32( LargerDim * OriginalDimRatio) );
+
+			if ( LargerDim * SmallerDim * NumSlices <= LimitPixelCount )
+			{
+				break;
+			}
+		}
+
+		int64 NewSizeX,NewSizeY;
+		if ( Image.SizeX >= Image.SizeY )
+		{
+			NewSizeX = LargerDim;
+			NewSizeY = SmallerDim;
+		}
+		else
+		{
+			NewSizeX = SmallerDim;
+			NewSizeY = LargerDim;
+		}
+
+		UE_LOG(LogInterchangeImport, Display, TEXT("Image %s larger than GetAutoLimitPixelCount.  Resizing %d x %d -> %d x %d (%d slices)"),
+			*Texture->GetName(),
+			Image.SizeX,Image.SizeY,
+			NewSizeX,NewSizeY,
+			NumSlices);
+
+		check( Image.NumMips == 1 );
+		
+		int64 SrcBytesPerSlice = (int64) Image.SizeX * Image.SizeY * BytesPerPixel;
+		int64 DstBytesPerSlice = (int64) NewSizeX * NewSizeY * BytesPerPixel;
+
+		FUniqueBuffer DstBuffer = FUniqueBuffer::Alloc(DstBytesPerSlice * NumSlices);
+		uint8 * DstBufferData = (uint8 *) DstBuffer.GetData();
+
+		FSharedBuffer SrcBuffer = BufferAndId.GetPayload();
+		
+		const uint8 * SrcBufferData = (const uint8 *) SrcBuffer.GetData();
+		int64 SrcBufferSize = SrcBuffer.GetSize();
+		
+		FImageView Src,Dst;
+
+		Src.SizeX = Image.SizeX;
+		Src.SizeY = Image.SizeY;
+		Src.NumSlices = 1;
+		Src.Format = FImageCoreUtils::ConvertToRawImageFormat(Image.Format);
+		Src.GammaSpace = Image.bSRGB ? EGammaSpace::sRGB : EGammaSpace::Linear;
+
+		Dst = Src;
+		Dst.SizeX = NewSizeX;
+		Dst.SizeY = NewSizeY;
+		
+		check( Src.GetImageSizeBytes() == SrcBytesPerSlice );
+		check( Dst.GetImageSizeBytes() == DstBytesPerSlice );
+
+		if ( Image.RawDataCompressionFormat != TSCF_None )
+		{
+			// only single 2d images can import compressed currently :
+			check( NumSlices == 1 );
+
+			// decompress image data
+			FImage SrcImage;
+			bool bOk = FImageUtils::DecompressImage(SrcBufferData,SrcBufferSize,SrcImage);
+			if ( ! bOk )
+			{
+				UE_LOG(LogInterchangeImport, Warning, TEXT("DecompressImage failed on [%s]."), *(Texture->GetName()) );
+				return;
+			}
+			
+			// read SrcImage , not "Src"
+
+			Dst.RawData = DstBufferData;
+
+			FImageCore::ResizeImage(SrcImage,Dst,FImageCore::EResizeImageFilter::Default);
+
+			// change Image now holds uncompressed data:
+			Image.RawDataCompressionFormat = TSCF_None;
+		}
+		else
+		{
+			// SrcBuffer contains raw images
+			check( SrcBufferSize == SrcBytesPerSlice * NumSlices );
+			if ( SrcBufferSize != SrcBytesPerSlice * NumSlices )
+			{
+				UE_LOG(LogInterchangeImport, Warning, TEXT("Invalid SrcBufferSize on [%s]."), *(Texture->GetName()) );
+				return;
+			}
+
+			// ResizeImage default edge mode is Clamp, so this is okay for cube map faces
+
+			for(int64 SliceIndex=0;SliceIndex<NumSlices;SliceIndex++)
+			{
+				Src.RawData = (void *)( SrcBufferData + SrcBytesPerSlice * SliceIndex );
+	
+				Dst.RawData = DstBufferData + DstBytesPerSlice * SliceIndex;
+
+				FImageCore::ResizeImage(Src,Dst,FImageCore::EResizeImageFilter::Default);
+			}
+		}
+
+		SrcBuffer.Reset(); // release the FSharedBuffer ref
+
+		// mutate Image:
+		Image.SizeX = NewSizeX;
+		Image.SizeY = NewSizeY;
+
+		// mutate BufferAndId
+		BufferAndId = UE::Serialization::FEditorBulkData::FSharedBufferWithID( DstBuffer.MoveToShared() );
+	}
+
+	void SetupTextureSourceDataFromBulkData_Editor(UTexture* Texture, FImportImage& Image, UE::Serialization::FEditorBulkData::FSharedBufferWithID&& BufferAndId, bool bIsReimport)
+	{
+		// Image has a .RawData but it is ignored here
+		//	use BufferAndId for the pixel data instead
+
+		ResizeImportImageIfNeeded(Texture,Image,1,BufferAndId);
+
 		Texture->Source.InitWithCompressedSourceData(
 			Image.SizeX,
 			Image.SizeY,
@@ -771,6 +1023,9 @@ namespace UE::Interchange::Private::InterchangeTextureFactory
 	{
 		if (BlockedImage.BlocksData.Num() > 1)
 		{
+			// note: ResizeImportImageIfNeeded not done
+			//	no size limit for UDIM
+
 			Texture2D->Source.InitBlocked(
 				&BlockedImage.Format,
 				BlockedImage.BlocksData.GetData(),
@@ -785,6 +1040,7 @@ namespace UE::Interchange::Private::InterchangeTextureFactory
 				
 				Texture2D->SRGB = UE::TextureUtilitiesCommon::GetDefaultSRGB(BlockedImage.CompressionSettings,BlockedImage.Format,BlockedImage.bSRGB);
 
+				// UDIM/Blocked must be VT :
 				Texture2D->VirtualTextureStreaming = true;
 
 				if (BlockedImage.MipGenSettings.IsSet())
@@ -814,13 +1070,27 @@ namespace UE::Interchange::Private::InterchangeTextureFactory
 		}
 	}
 
-	void SetupTextureSourceDataFromBulkData_Editor(UTexture* Texture, const FImportSlicedImage& SlicedImage, UE::Serialization::FEditorBulkData::FSharedBufferWithID&& BufferAndId,  bool bIsReimport)
+	void SetupTextureSourceDataFromBulkData_Editor(UTexture* Texture, FImportSlicedImage& SlicedImage, UE::Serialization::FEditorBulkData::FSharedBufferWithID&& BufferAndId,  bool bIsReimport)
 	{
+		bool bTextureIsVolume = Texture->GetTextureClass() == ETextureClass::Volume;
+		if ( SlicedImage.bIsVolume )
+		{
+			check( bTextureIsVolume );
+
+			// @@ no Volume resize currently
+		}
+		else
+		{
+			// NumSlice not changed by resize
+			// FImportSlicedImage passed as FImportImage
+			ResizeImportImageIfNeeded(Texture,SlicedImage,SlicedImage.NumSlice,BufferAndId);
+		}
+
 		Texture->Source.InitLayered(
 			SlicedImage.SizeX,
 			SlicedImage.SizeY,
 			SlicedImage.NumSlice,
-			1,
+			1, // NumLayers
 			SlicedImage.NumMips,
 			&SlicedImage.Format,
 			MoveTemp(BufferAndId)
@@ -1391,7 +1661,7 @@ namespace UE::Interchange::Private::InterchangeTextureFactory
 		return false;
 	}
 
-	void SetupTextureSourceDataFromBulkData(UTexture* Texture, const FImportImage& Image, UE::Serialization::FEditorBulkData::FSharedBufferWithID&& BufferAndId, bool bIsReimport)
+	void SetupTextureSourceDataFromBulkData(UTexture* Texture, FImportImage& Image, UE::Serialization::FEditorBulkData::FSharedBufferWithID&& BufferAndId, bool bIsReimport)
 	{
 #if WITH_EDITOR
 		SetupTextureSourceDataFromBulkData_Editor(Texture, Image, MoveTemp(BufferAndId), bIsReimport);
diff --git a/Engine/Source/Runtime/TextureUtilitiesCommon/Private/TextureImportSettings.cpp b/Engine/Source/Runtime/TextureUtilitiesCommon/Private/TextureImportSettings.cpp
--- a/Engine/Source/Runtime/TextureUtilitiesCommon/Private/TextureImportSettings.cpp
+++ b/Engine/Source/Runtime/TextureUtilitiesCommon/Private/TextureImportSettings.cpp
@@ -52,6 +52,53 @@ ETextureImportPNGInfill UTextureImportSettings::GetPNGInfillMapDefault() const
 	}
 }
 
+bool UTextureImportSettings::IsImportAutoVTEnabled() const
+{
+	if ( AutoVTSize <= 0 )
+	{
+		return false;
+	}
+
+	static const TConsoleVariableData<int32>* CVarVirtualTexturesEnabled = IConsoleManager::Get().FindTConsoleVariableDataInt(TEXT("r.VirtualTextures"));
+	check(CVarVirtualTexturesEnabled);
+
+	static const auto CVarVirtualTexturesAutoImportEnabled = IConsoleManager::Get().FindTConsoleVariableDataInt(TEXT("r.VT.EnableAutoImport"));
+	check(CVarVirtualTexturesAutoImportEnabled);
+
+	return ( CVarVirtualTexturesEnabled->GetValueOnAnyThread() && CVarVirtualTexturesAutoImportEnabled->GetValueOnAnyThread() );
+}
+
+int64 UTextureImportSettings::GetAutoLimitPixelCount() const
+{
+	if ( AutoLimitDimension <= 0 )
+	{
+		// no limit
+		return 0;
+	}
+	
+	int64 UseAutoLimitDimension = AutoLimitDimension;
+
+	UseAutoLimitDimension = FMath::Min<int64>(UseAutoLimitDimension, UTexture::GetMaximumDimensionOfNonVT());
+
+	if ( IsImportAutoVTEnabled() )
+	{
+		if ( AutoVTSize != AutoLimitDimension )
+		{
+			// AutoVTSize and AutoLimitDimension cannot both be enabled and be different
+
+			UE_CALL_ONCE( [&](){	
+				UE_LOG(LogCore, Warning, TEXT("VT is enabled with AutoVTSize (%d) not equal AutoLimitDimension (%d); they must be equal or zero, fix config!  Ignoring AutoLimitDimension and using AutoVTSize."), 
+					AutoVTSize,AutoLimitDimension
+					);	
+			} );
+
+			UseAutoLimitDimension = AutoVTSize;
+		}
+	}
+
+	return UseAutoLimitDimension*UseAutoLimitDimension;
+}
+
 #if WITH_EDITOR
 void UTextureImportSettings::PostEditChangeProperty(FPropertyChangedEvent& PropertyChangedEvent)
 {
@@ -151,38 +198,34 @@ namespace UE::TextureUtilitiesCommon
 	{
 		// If the texture is larger than a certain threshold make it VT.
 		// Note that previously for re-imports we still checked size and potentially changed the VT status.
-		// But that was unintuitive for many users so now for re-imports we will end up ignoring this and respecting the existing setting below.
+		// But that was unintuitive for many users so now for re-imports we will end up ignoring this and respecting the existing setting.
 		
-		static const TConsoleVariableData<int32>* CVarVirtualTexturesEnabled = IConsoleManager::Get().FindTConsoleVariableDataInt(TEXT("r.VirtualTextures"));
-		check(CVarVirtualTexturesEnabled);
+		if ( ! GetDefault<UTextureImportSettings>()->IsImportAutoVTEnabled() )
+		{
+			return false;
+		}
 
-		static const auto CVarVirtualTexturesAutoImportEnabled = IConsoleManager::Get().FindTConsoleVariableDataInt(TEXT("r.VT.EnableAutoImport"));
-		check(CVarVirtualTexturesAutoImportEnabled);
+		const int64 VirtualTextureAutoEnableThreshold = GetDefault<UTextureImportSettings>()->AutoVTSize;
 
-		if (CVarVirtualTexturesEnabled->GetValueOnAnyThread() && CVarVirtualTexturesAutoImportEnabled->GetValueOnAnyThread())
+		if ( VirtualTextureAutoEnableThreshold <= 0 )
 		{
-			const int64 VirtualTextureAutoEnableThreshold = GetDefault<UTextureImportSettings>()->AutoVTSize;
-
-			if ( VirtualTextureAutoEnableThreshold == 0 )
-			{
-				return false;
-			}
+			return false;
+		}
+			
+		const int64 VirtualTextureAutoEnableThresholdPixels = VirtualTextureAutoEnableThreshold * VirtualTextureAutoEnableThreshold;
 
-			const int64 VirtualTextureAutoEnableThresholdPixels = VirtualTextureAutoEnableThreshold * VirtualTextureAutoEnableThreshold;
-
-			// We do this in pixels so a 8192 x 128 texture won't get VT enabled 
-			// We use the Source size instead of simple Texture2D->GetSizeX() as this uses the size of the platform data
-			// however for a new texture platform data may not be generated yet, and for an reimport of a texture this is the size of the
-			// old texture. 
-			// Using source size gives one small caveat. It looks at the size before mipmap power of two padding adjustment.
-			// Textures with more than 1 block (UDIM textures) must be imported as VT
-			if (Texture->Source.GetNumBlocks() > 1 ||
-				( (int64) Texture->Source.GetSizeX() * Texture->Source.GetSizeY() ) >= VirtualTextureAutoEnableThresholdPixels ||
-				Texture->Source.GetSizeX() > UTexture::GetMaximumDimensionOfNonVT() ||
-				Texture->Source.GetSizeY() > UTexture::GetMaximumDimensionOfNonVT() )
-			{
-				return true;
-			}
+		// We do this in pixels so a 8192 x 128 texture won't get VT enabled 
+		// We use the Source size instead of simple Texture2D->GetSizeX() as this uses the size of the platform data
+		// however for a new texture platform data may not be generated yet, and for an reimport of a texture this is the size of the
+		// old texture. 
+		// Using source size gives one small caveat. It looks at the size before mipmap power of two padding adjustment.
+		// Textures with more than 1 block (UDIM textures) must be imported as VT
+		if (Texture->Source.GetNumBlocks() > 1 ||
+			( (int64) Texture->Source.GetSizeX() * Texture->Source.GetSizeY() ) >= VirtualTextureAutoEnableThresholdPixels ||
+			Texture->Source.GetSizeX() > UTexture::GetMaximumDimensionOfNonVT() ||
+			Texture->Source.GetSizeY() > UTexture::GetMaximumDimensionOfNonVT() )
+		{
+			return true;
 		}
 
 		return false;
diff --git a/Engine/Source/Runtime/TextureUtilitiesCommon/Public/TextureImportSettings.h b/Engine/Source/Runtime/TextureUtilitiesCommon/Public/TextureImportSettings.h
--- a/Engine/Source/Runtime/TextureUtilitiesCommon/Public/TextureImportSettings.h
+++ b/Engine/Source/Runtime/TextureUtilitiesCommon/Public/TextureImportSettings.h
@@ -56,8 +56,13 @@ public:
 
 	UPROPERTY(config, EditAnywhere, Category=VirtualTextures, meta = (
 		DisplayName = "Auto Virtual Texturing Size",
-		ToolTip = "Automatically enable the 'Virtual Texture Streaming' texture setting for textures larger than or equal to this size. This setting will not affect existing textures in the project."))
+		ToolTip = "Automatically enable the 'Virtual Texture Streaming' texture setting for textures larger than or equal to this size. Zero to disable.  This setting will not affect existing textures in the project."))
 	int32 AutoVTSize = 4096;
+
+	UPROPERTY(config, EditAnywhere, Category=VirtualTextures, meta = (
+		DisplayName = "Auto resize large textures",
+		ToolTip = "If not zero, textures over this size that are not VT will be resized to be equal to or less than this.  Treated as pixel count limit for cubes and volumes.  Must equal AutoVTSize if both are non-zero.  This setting will not affect existing textures in the project."))
+	int32 AutoLimitDimension = 0;
 	
 	UPROPERTY(config, EditAnywhere, Category=ImportSettings, meta = (
 		DisplayName = "Turn on NormalizeNormals for normal maps",
@@ -86,6 +91,13 @@ public:
 	// Get the PNGInfill setting, with Default mapped to a concrete choice
 	TEXTUREUTILITIESCOMMON_API ETextureImportPNGInfill GetPNGInfillMapDefault() const;
 
+	// get AutoLimitDimension as a pixel count (squared), and with sanity checks applied
+	//	zero means do not limit
+	TEXTUREUTILITIESCOMMON_API int64 GetAutoLimitPixelCount() const;
+	
+	// IsImportAutoVTEnabled checks not just this config but also the project VT settings
+	TEXTUREUTILITIESCOMMON_API bool IsImportAutoVTEnabled() const;
+
 #if WITH_EDITOR
 	TEXTUREUTILITIESCOMMON_API virtual void PostEditChangeProperty(FPropertyChangedEvent& PropertyChangedEvent) override;
 #endif
