Avoid crashing in the interchange texture pipeline if a texture is compiling in the post asset import function. Instead just wait for it to finish and ensure if that happened on a import as this should be avoidable.

#rnx

diff --git a/Engine/Plugins/Interchange/Runtime/Source/Pipelines/Private/InterchangeGenericTexturePipeline.cpp b/Engine/Plugins/Interchange/Runtime/Source/Pipelines/Private/InterchangeGenericTexturePipeline.cpp
--- a/Engine/Plugins/Interchange/Runtime/Source/Pipelines/Private/InterchangeGenericTexturePipeline.cpp
+++ b/Engine/Plugins/Interchange/Runtime/Source/Pipelines/Private/InterchangeGenericTexturePipeline.cpp
@@ -401,7 +401,12 @@ void UInterchangeGenericTexturePipeline::PostImportTextureAssetImport(UObject* C
 	// (Note - as part of the standard interchange import this is called during the object
 	// import iteration, _before_ the iteration to call to PostEditChange which is what starts the texture build via UpdateResource,
 	// so altering properties here should be safe!)
-	check(!FTextureCompilingManager::Get().IsCompilingTexture(Texture));
+
+	if (Texture->IsCompiling())
+	{
+		ensure(!bIsAReimport);
+		FTextureCompilingManager::Get().FinishCompilation(MakeArrayView(&Texture, 1));
+	}
 
 	if(bFlipNormalMapGreenChannel && Texture->IsNormalMap())
 	{
