[Interchange glTF] Iridescence config usage fix.

#rnx

#virtualized

diff --git a/Engine/Plugins/Interchange/Runtime/Source/Import/Private/Gltf/InterchangeGLTFMaterial.cpp b/Engine/Plugins/Interchange/Runtime/Source/Import/Private/Gltf/InterchangeGLTFMaterial.cpp
--- a/Engine/Plugins/Interchange/Runtime/Source/Import/Private/Gltf/InterchangeGLTFMaterial.cpp
+++ b/Engine/Plugins/Interchange/Runtime/Source/Import/Private/Gltf/InterchangeGLTFMaterial.cpp
@@ -607,7 +607,7 @@ namespace UE::Interchange::GLTFMaterials
 					//IridescenceTexture_Rotation
 					//IridescenceTexture_TexCoord
 					//IridescenceTexture_TilingMethod
-					bool bHasIridescenceTexture = SetMap(Inputs::IridescenceTexture, GltfMaterial.Iridescence.Texture);
+					SetMap(Inputs::IridescenceTexture, GltfMaterial.Iridescence.Texture);
 
 					//IridescenceFactor
 					SetScalar(Inputs::IridescenceFactor, GltfMaterial.Iridescence.Factor, 0.0f);
@@ -617,7 +617,7 @@ namespace UE::Interchange::GLTFMaterials
 					//IridescenceThicknessTexture_Rotation
 					//IridescenceThicknessTexture_TexCoord
 					//IridescenceThicknessTexture_TilingMethod
-					bHasIridescenceTexture = SetMap(Inputs::IridescenceThicknessTexture, GltfMaterial.Iridescence.Thickness.Texture) || bHasIridescenceTexture;
+					SetMap(Inputs::IridescenceThicknessTexture, GltfMaterial.Iridescence.Thickness.Texture);
 
 					//IridescenceThicknessMinimum
 					SetScalar(Inputs::IridescenceThicknessMinimum, GltfMaterial.Iridescence.Thickness.Minimum, 100.0f);
@@ -625,12 +625,12 @@ namespace UE::Interchange::GLTFMaterials
 					SetScalar(Inputs::IridescenceThicknessMaximum, GltfMaterial.Iridescence.Thickness.Maximum, 400.0f);
 
 					//Configuration
-					SetBool(Inputs::Configuration::bHasIridescenceTexture, bHasIridescenceTexture);
+					SetBool(Inputs::Configuration::bHasIridescence, true);
 				}
 				else
 				{
 					//Configuration
-					SetBool(Inputs::Configuration::bHasIridescenceTexture, false);
+					SetBool(Inputs::Configuration::bHasIridescence, false);
 				}
 
 				if (GltfMaterial.Anisotropy.bHasAnisotropy)
diff --git a/Engine/Plugins/Interchange/Runtime/Source/Import/Public/Gltf/InterchangeGLTFMaterial.h b/Engine/Plugins/Interchange/Runtime/Source/Import/Public/Gltf/InterchangeGLTFMaterial.h
--- a/Engine/Plugins/Interchange/Runtime/Source/Import/Public/Gltf/InterchangeGLTFMaterial.h
+++ b/Engine/Plugins/Interchange/Runtime/Source/Import/Public/Gltf/InterchangeGLTFMaterial.h
@@ -145,7 +145,7 @@ namespace UE::Interchange::GLTFMaterials
 			const FString bHasClearCoatTexture = TEXT("bHasClearCoatTexture");
 			const FString bHasSheenTexture = TEXT("bHasSheenTexture");
 			const FString bHasTransmissionTexture = TEXT("bHasTransmissionTexture");
-			const FString bHasIridescenceTexture = TEXT("bHasIridescenceTexture");
+			const FString bHasIridescence = TEXT("bHasIridescence");
 			const FString bHasAnisotropyTextureAndOrRotation = TEXT("bHasAnisotropyTextureAndOrRotation"); //AnistropyRotation can affect the Tangents so we check
 		}
 	}
