Interchange: fix ensure with d3d12 scissor, avoid registering the asset before the build is done. Also fix issue where two progress dialog where set on the game thread.

#rnx

diff --git a/Engine/Plugins/Interchange/Runtime/Source/Import/Private/Mesh/InterchangeSkeletalMeshFactory.cpp b/Engine/Plugins/Interchange/Runtime/Source/Import/Private/Mesh/InterchangeSkeletalMeshFactory.cpp
--- a/Engine/Plugins/Interchange/Runtime/Source/Import/Private/Mesh/InterchangeSkeletalMeshFactory.cpp
+++ b/Engine/Plugins/Interchange/Runtime/Source/Import/Private/Mesh/InterchangeSkeletalMeshFactory.cpp
@@ -2145,40 +2145,24 @@ void UInterchangeSkeletalMeshFactory::FinalizeObject_GameThread(const FSetupObje
 	//to do this we need to put the cloth binding data in the import data use by the build (i.e. Meshdescription)
 	if (ImportAssetObjectData.ExistingClothingBindings.Num() > 0)
 	{
-		float ProgressCount = 2.0f;
-
-		FScopedSlowTask Progress(ProgressCount, NSLOCTEXT("UInterchangeSkeletalMeshPostImportTask", "SkeletalMeshFinalizeImportGameThread", "Executing Skeletal Mesh Finalize Import Task..."));
-		Progress.MakeDialog();
-
 		//Make sure we rebuild the skeletal mesh after re-importing all skin weight
 		FScopedSkeletalMeshPostEditChange ScopePostEditChange(SkeletalMesh);
-
 		//Wait until the asset is finish building then lock the skeletal mesh properties to prevent the UI to update during the alternate skinning reimport
 		FEvent* LockEvent = SkeletalMesh->LockPropertiesUntil();
-
-		Progress.EnterProgressFrame(1.0f);
-
 		FSkinnedAssetAsyncBuildScope AsyncBuildScope(SkeletalMesh);
-
-		//Restore the clothing
-		if (ImportAssetObjectData.ExistingClothingBindings.Num() > 0)
+		for (ClothingAssetUtils::FClothingAssetMeshBinding& ExistingClothMeshBinding : ImportAssetObjectData.ExistingClothingBindings)
 		{
-			for (ClothingAssetUtils::FClothingAssetMeshBinding& ExistingClothMeshBinding : ImportAssetObjectData.ExistingClothingBindings)
+			if (UClothingAssetCommon* ClothAssetCommon = ExistingClothMeshBinding.Asset)
 			{
-				if (UClothingAssetCommon* ClothAssetCommon = ExistingClothMeshBinding.Asset)
-				{
-					ClothAssetCommon->RefreshBoneMapping(SkeletalMesh);
-				}
-			}
-			FSkeletalMeshModel* ImportedResource = SkeletalMesh->GetImportedModel();
-			for (int32 LodIndex = 0; LodIndex < ImportedResource->LODModels.Num(); ++LodIndex)
-			{
-				// Re-apply our clothing assets
-				FLODUtilities::RestoreClothingFromBackup(SkeletalMesh, ImportAssetObjectData.ExistingClothingBindings, LodIndex);
+				ClothAssetCommon->RefreshBoneMapping(SkeletalMesh);
 			}
-			Progress.EnterProgressFrame(1.0f);
 		}
-
+		FSkeletalMeshModel* ImportedResource = SkeletalMesh->GetImportedModel();
+		for (int32 LodIndex = 0; LodIndex < ImportedResource->LODModels.Num(); ++LodIndex)
+		{
+			// Re-apply our clothing assets
+			FLODUtilities::RestoreClothingFromBackup(SkeletalMesh, ImportAssetObjectData.ExistingClothingBindings, LodIndex);
+		}
 		//Release the skeletal mesh async properties
 		LockEvent->Trigger();
 	}
diff --git a/Engine/Source/Runtime/Interchange/Engine/Private/Tasks/InterchangeTaskCompletion.cpp b/Engine/Source/Runtime/Interchange/Engine/Private/Tasks/InterchangeTaskCompletion.cpp
--- a/Engine/Source/Runtime/Interchange/Engine/Private/Tasks/InterchangeTaskCompletion.cpp
+++ b/Engine/Source/Runtime/Interchange/Engine/Private/Tasks/InterchangeTaskCompletion.cpp
@@ -157,12 +157,6 @@ void UE::Interchange::FTaskPreCompletion::DoTask(ENamedThreads::Type CurrentThre
 			if (bIsAsset)
 			{
 				AsyncHelper->AssetImportResult->AddImportedObject(ImportedObject);
-
-				if (!AsyncHelper->TaskData.ReimportObject)
-				{
-					//Notify the asset registry, only when we have created the asset
-					FAssetRegistryModule::AssetCreated(ImportedObject);
-				}
 			}
 			else
 			{
@@ -206,14 +200,12 @@ void UE::Interchange::FTaskCompletion::DoTask(ENamedThreads::Type CurrentThread,
 				{
 					if (UObject* Asset = AssetInfo.ImportedObject)
 					{
-						//Call post edit change if it was not call previously
-#if WITH_EDITOR
-						if (!AssetInfo.bPostEditChangeCalled)
+						if (!AsyncHelper->TaskData.ReimportObject)
 						{
-							Asset->PostEditChange();
+							//Notify the asset registry, only when we have created the asset
+							FAssetRegistryModule::AssetCreated(Asset);
 						}
-#endif //WITH_EDITOR
-						if (AsyncHelper->TaskData.ReimportObject && AsyncHelper->TaskData.ReimportObject == Asset)
+						else if (AsyncHelper->TaskData.ReimportObject && AsyncHelper->TaskData.ReimportObject == Asset)
 						{
 							UInterchangeManager::GetInterchangeManager().OnAssetPostReimport.Broadcast(Asset);
 						}
diff --git a/Engine/Source/Runtime/Interchange/Engine/Private/Tasks/InterchangeTaskPipeline.cpp b/Engine/Source/Runtime/Interchange/Engine/Private/Tasks/InterchangeTaskPipeline.cpp
--- a/Engine/Source/Runtime/Interchange/Engine/Private/Tasks/InterchangeTaskPipeline.cpp
+++ b/Engine/Source/Runtime/Interchange/Engine/Private/Tasks/InterchangeTaskPipeline.cpp
@@ -182,6 +182,24 @@ void UE::Interchange::FTaskPostImport::DoTask(ENamedThreads::Type CurrentThread,
 		return;
 	}
 
+#if WITH_EDITOR
+	auto CallPostEditChangeForAsset = [](const TArray<UE::Interchange::FImportAsyncHelper::FImportedObjectInfo>& ImportedInfos)
+		{
+			for (const UE::Interchange::FImportAsyncHelper::FImportedObjectInfo& ObjectInfo : ImportedInfos)
+			{
+				UObject* ImportedObject = ObjectInfo.ImportedObject;
+				if (!ObjectInfo.bPostEditChangeCalled)
+				{
+					ImportedObject->PostEditChange();
+				}
+			}
+		};
+
+	AsyncHelper->IterateImportedAssets(SourceIndex, CallPostEditChangeForAsset);
+	AsyncHelper->IterateImportedSceneObjects(SourceIndex, CallPostEditChangeForAsset);
+
+#endif //WITH_EDITOR
+
 	TArray<FString> NodeUniqueIDs;
 	TArray<UObject*> ImportedObjects;
 	TArray<bool> IsAssetsReimported;
@@ -207,8 +225,7 @@ void UE::Interchange::FTaskPostImport::DoTask(ENamedThreads::Type CurrentThread,
 			{
 				UObject* ImportedObject = ObjectInfo.ImportedObject;
 
-				//In case Some factory code cannot run outside of the main thread we offer this callback to finish the work after calling post edit change (building the asset)
-				//Its possible the build of the asset to be asynchronous, the factory must handle is own asset correctly
+				//In case Some factory code cannot run outside of the main thread we offer this callback to finish the work after asset build is finish.
 				if (bCallPostImportGameThreadCallback && ObjectInfo.Factory)
 				{
 					Arguments.ImportedObject = ImportedObject;
