Unshelved from pending changelist '33509305':

Interchange: fix fbx shape name clash to make sure all blendshape are imported correctly. Add a feature to merge blendshape with the same name together at import.

#rnx

#virtualized

diff --git a/Engine/Plugins/Interchange/Runtime/Source/FactoryNodes/Private/InterchangeSkeletalMeshFactoryNode.cpp b/Engine/Plugins/Interchange/Runtime/Source/FactoryNodes/Private/InterchangeSkeletalMeshFactoryNode.cpp
--- a/Engine/Plugins/Interchange/Runtime/Source/FactoryNodes/Private/InterchangeSkeletalMeshFactoryNode.cpp
+++ b/Engine/Plugins/Interchange/Runtime/Source/FactoryNodes/Private/InterchangeSkeletalMeshFactoryNode.cpp
@@ -224,6 +224,15 @@ bool UInterchangeSkeletalMeshFactoryNode::FillCustomBoneInfluenceLimitFromAsset(
 	IMPLEMENT_SKELETALMESH_BUILD_ASSET_TO_VALUE(BoneInfluenceLimit, BoneInfluenceLimit);
 }
 
+bool UInterchangeSkeletalMeshFactoryNode::GetCustomMergeMorphTargetShapeWithSameName(bool& AttributeValue) const
+{
+	IMPLEMENT_NODE_ATTRIBUTE_GETTER(MergeMorphTargetShapeWithSameName, bool)
+}
+bool UInterchangeSkeletalMeshFactoryNode::SetCustomMergeMorphTargetShapeWithSameName(const bool& AttributeValue)
+{
+	IMPLEMENT_NODE_ATTRIBUTE_SETTER_NODELEGATE(MergeMorphTargetShapeWithSameName, bool)
+}
+
 bool UInterchangeSkeletalMeshFactoryNode::GetCustomImportContentType(EInterchangeSkeletalMeshContentType& AttributeValue) const
 {
 	FString OperationName = GetTypeName() + TEXT(".GetImportContentType");
diff --git a/Engine/Plugins/Interchange/Runtime/Source/FactoryNodes/Public/InterchangeSkeletalMeshFactoryNode.h b/Engine/Plugins/Interchange/Runtime/Source/FactoryNodes/Public/InterchangeSkeletalMeshFactoryNode.h
--- a/Engine/Plugins/Interchange/Runtime/Source/FactoryNodes/Public/InterchangeSkeletalMeshFactoryNode.h
+++ b/Engine/Plugins/Interchange/Runtime/Source/FactoryNodes/Public/InterchangeSkeletalMeshFactoryNode.h
@@ -152,6 +152,14 @@ public:
 	UFUNCTION(BlueprintCallable, Category = "Interchange | Node | SkeletalMesh")
 	bool SetCustomBoneInfluenceLimit(const int32& AttributeValue, bool bAddApplyDelegate = true);
 
+	/** Query whether the skeletal mesh factory should merge morph target shape with the same name under one morph target. Return false if the attribute was not set.*/
+	UFUNCTION(BlueprintCallable, Category = "Interchange | Node | SkeletalMesh")
+	bool GetCustomMergeMorphTargetShapeWithSameName(bool& AttributeValue) const;
+
+	/** Set whether the skeletal mesh factory should merge morph target shape with the same name under one morph target. Return false if the attribute could not be set.*/
+	UFUNCTION(BlueprintCallable, Category = "Interchange | Node | SkeletalMesh")
+	bool SetCustomMergeMorphTargetShapeWithSameName(const bool& AttributeValue);
+
 	/**
 	 * The skeletal mesh thumbnail can have an overlay if the last reimport was geometry only. This thumbnail overlay feature uses the metadata to find out if the last import was geometry only.
 	 */
@@ -178,6 +186,7 @@ private:
 	const UE::Interchange::FAttributeKey Macro_CustomThresholdUVKey = UE::Interchange::FAttributeKey(TEXT("ThresholdUV"));
 	const UE::Interchange::FAttributeKey Macro_CustomMorphThresholdPositionKey = UE::Interchange::FAttributeKey(TEXT("MorphThresholdPosition"));
 	const UE::Interchange::FAttributeKey Macro_CustomBoneInfluenceLimitKey = UE::Interchange::FAttributeKey(TEXT("BoneInfluenceLimit"));
+	const UE::Interchange::FAttributeKey Macro_CustomMergeMorphTargetShapeWithSameNameKey = UE::Interchange::FAttributeKey(TEXT("MergeMorphTargetShapeWithSameName"));
 
 	bool ApplyCustomUseHighPrecisionSkinWeightsToAsset(UObject* Asset) const;
 	bool FillCustomUseHighPrecisionSkinWeightsFromAsset(UObject* Asset);
diff --git a/Engine/Plugins/Interchange/Runtime/Source/Import/Private/Mesh/InterchangeSkeletalMeshFactory.cpp b/Engine/Plugins/Interchange/Runtime/Source/Import/Private/Mesh/InterchangeSkeletalMeshFactory.cpp
--- a/Engine/Plugins/Interchange/Runtime/Source/Import/Private/Mesh/InterchangeSkeletalMeshFactory.cpp
+++ b/Engine/Plugins/Interchange/Runtime/Source/Import/Private/Mesh/InterchangeSkeletalMeshFactory.cpp
@@ -153,19 +153,32 @@ namespace UE
 				}
 			}
 
-			void CopyMorphTargetsMeshDescriptionToSkeletalMeshImportData(const TMap<FString, TOptional<UE::Interchange::FMeshPayloadData>>& LodMorphTargetMeshDescriptions, FSkeletalMeshImportData& DestinationSkeletalMeshImportData)
+			void CopyMorphTargetsMeshDescriptionToSkeletalMeshImportData(const TMap<FString, TOptional<UE::Interchange::FMeshPayloadData>>& LodMorphTargetMeshDescriptions, FSkeletalMeshImportData& DestinationSkeletalMeshImportData, const bool bMergeMorphTargetWithSameName)
 			{
 				TRACE_CPUPROFILER_EVENT_SCOPE(CopyMorphTargetsMeshDescriptionToSkeletalMeshImportData)
 				const int32 OriginalMorphTargetCount = LodMorphTargetMeshDescriptions.Num();
-				TArray<FString> Keys;
-				int32 MorphTargetCount = 0;
+				TArray < TPair<FString, TArray<FString>>> KeysPerName;
+				auto FindOrAdd = [&KeysPerName](const FString& MorphTargetName)->TPair<FString, TArray<FString>>&
+					{
+						for (TPair<FString, TArray<FString>>& Pair : KeysPerName)
+						{
+							if (Pair.Key.Equals(MorphTargetName))
+							{
+								return Pair;
+							}
+						}
+
+						TPair<FString, TArray<FString>>& NewPair = KeysPerName.AddDefaulted_GetRef();
+						NewPair.Key = MorphTargetName;
+						return NewPair;
+					};
 				for (const TPair<FString, TOptional<UE::Interchange::FMeshPayloadData>>& Pair : LodMorphTargetMeshDescriptions)
 				{
-					const FString MorphTargetName(Pair.Key);
+					const FString MorphTargetUniqueId(Pair.Key);
 					const TOptional<UE::Interchange::FMeshPayloadData>& MorphTargetPayloadData = Pair.Value;
 					if (!MorphTargetPayloadData.IsSet())
 					{
-						UE_LOG(LogInterchangeImport, Error, TEXT("Empty morph target optional payload data [%s]."), *MorphTargetName);
+						UE_LOG(LogInterchangeImport, Error, TEXT("Empty morph target optional payload data [%s]."), *MorphTargetUniqueId);
 						continue;
 					}
 
@@ -175,12 +188,25 @@ namespace UE
 					const int32 DestinationVertexIndexMax = VertexOffset + SourceMeshVertexCount;
 					if (!DestinationSkeletalMeshImportData.Points.IsValidIndex(DestinationVertexIndexMax-1))
 					{
-						UE_LOG(LogInterchangeImport, Error, TEXT("Corrupted morph target optional payload data [%s]."), *MorphTargetName);
+						UE_LOG(LogInterchangeImport, Error, TEXT("Corrupted morph target optional payload data [%s]."), *MorphTargetUniqueId);
 						continue;
 					}
-					Keys.Add(Pair.Key);
-					MorphTargetCount++;
+
+					if (bMergeMorphTargetWithSameName)
+					{
+						TPair<FString, TArray<FString>>& PairMorphNameWithKeys = FindOrAdd(MorphTargetPayloadData->MorphTargetName);
+						PairMorphNameWithKeys.Value.Add(MorphTargetUniqueId);
+					}
+					else
+					{
+						TPair<FString, TArray<FString>>& NewPair = KeysPerName.AddDefaulted_GetRef();
+						NewPair.Key = MorphTargetPayloadData->MorphTargetName;
+						NewPair.Value.Add(MorphTargetUniqueId);
+					}
 				}
+				
+				//Adjust the count from the merge context
+				const int32 MorphTargetCount = KeysPerName.Num();
 
 				//No morph target to import
 				if (MorphTargetCount == 0)
@@ -188,7 +214,6 @@ namespace UE
 					return;
 				}
 
-				ensure(Keys.Num() == MorphTargetCount);
 				//Allocate the data
 				DestinationSkeletalMeshImportData.MorphTargetNames.AddDefaulted(MorphTargetCount);
 				DestinationSkeletalMeshImportData.MorphTargetModifiedPoints.AddDefaulted(MorphTargetCount);
@@ -203,70 +228,75 @@ namespace UE
 							MorphTargetCount,
 							NumMorphGroup,
 							&LodMorphTargetMeshDescriptions,
-							&Keys,
+							&KeysPerName,
+							bMergeMorphTargetWithSameName,
 							&DestinationSkeletalMeshImportData](const int32 MorphTargetGroupIndex)
 				{
 					const int32 MorphTargetIndexOffset = MorphTargetGroupIndex * MorphTargetGroupSize;
 					const int32 MorphTargetEndLoopCount = MorphTargetIndexOffset + MorphTargetGroupSize;
 					for (int32 MorphTargetIndex = MorphTargetIndexOffset; MorphTargetIndex < MorphTargetEndLoopCount; ++MorphTargetIndex)
 					{
-						if (!Keys.IsValidIndex(MorphTargetIndex))
+						if (!KeysPerName.IsValidIndex(MorphTargetIndex))
 						{
 							ensure(MorphTargetGroupIndex + 1 == NumMorphGroup);
 							//Executing the last morph target group, in case we do not have a full last group.
 							break;
 						}
-						const FString MorphTargetKey(Keys[MorphTargetIndex]);
-						const TOptional<UE::Interchange::FMeshPayloadData>& MorphTargetPayloadData = LodMorphTargetMeshDescriptions.FindChecked(MorphTargetKey);
-						if (!ensure(MorphTargetPayloadData.IsSet()))
-						{
-							//This error was suppose to be catch in the pre parallel for loop
-							break;
-						}
-
-						const FMeshDescription& SourceMeshDescription = MorphTargetPayloadData.GetValue().MeshDescription;
-						const FTransform GlobalTransform = MorphTargetPayloadData->GlobalTransform.IsSet() ? MorphTargetPayloadData->GlobalTransform.GetValue() : FTransform::Identity;
-						const int32 VertexOffset = MorphTargetPayloadData->VertexOffset;
-						const int32 SourceMeshVertexCount = SourceMeshDescription.Vertices().Num();
-						const int32 DestinationVertexIndexMax = VertexOffset + SourceMeshVertexCount;
-						if (!ensure(DestinationSkeletalMeshImportData.Points.IsValidIndex(DestinationVertexIndexMax-1)))
-						{
-							//This error was suppose to be catch in the pre parallel for loop
-							break;
-						}
-						TArray<FVector3f> CompressPoints;
-						CompressPoints.Reserve(SourceMeshVertexCount);
-						FStaticMeshConstAttributes Attributes(SourceMeshDescription);
-						TVertexAttributesConstRef<FVector3f> VertexPositions = Attributes.GetVertexPositions();
-
 						//Create the morph target source data
 						FString& MorphTargetName = DestinationSkeletalMeshImportData.MorphTargetNames[MorphTargetIndex];
-						MorphTargetName = MorphTargetPayloadData->MorphTargetName;
+						MorphTargetName = KeysPerName[MorphTargetIndex].Key;
 						TSet<uint32>& ModifiedPoints = DestinationSkeletalMeshImportData.MorphTargetModifiedPoints[MorphTargetIndex];
 						FSkeletalMeshImportData& MorphTargetData = DestinationSkeletalMeshImportData.MorphTargets[MorphTargetIndex];
 
-						//Reserve the point and influences
-						MorphTargetData.Points.AddZeroed(SourceMeshVertexCount);
-
-						for (FVertexID VertexID : SourceMeshDescription.Vertices().GetElementIDs())
+						TArray<FVector3f> CompressPoints;
+						for (const FString& MorphTargetKey : KeysPerName[MorphTargetIndex].Value)
 						{
-							//We can use GetValue because the Meshdescription was compacted before the copy
-							MorphTargetData.Points[VertexID.GetValue()] = (FVector3f)GlobalTransform.TransformPosition((FVector)VertexPositions[VertexID]);
-						}
+							const TOptional<UE::Interchange::FMeshPayloadData>& MorphTargetPayloadData = LodMorphTargetMeshDescriptions.FindChecked(MorphTargetKey);
+							if (!ensure(MorphTargetPayloadData.IsSet()))
+							{
+								//This error was suppose to be catch in the pre parallel for loop
+								break;
+							}
 
-						for (int32 PointIdx = VertexOffset; PointIdx < DestinationSkeletalMeshImportData.Points.Num(); ++PointIdx)
-						{
-							int32 OriginalPointIdx = DestinationSkeletalMeshImportData.PointToRawMap[PointIdx] - VertexOffset;
-							if (!MorphTargetData.Points.IsValidIndex(OriginalPointIdx))
+							const FMeshDescription& SourceMeshDescription = MorphTargetPayloadData.GetValue().MeshDescription;
+							const FTransform GlobalTransform = MorphTargetPayloadData->GlobalTransform.IsSet() ? MorphTargetPayloadData->GlobalTransform.GetValue() : FTransform::Identity;
+							const int32 VertexOffset = MorphTargetPayloadData->VertexOffset;
+							const int32 SourceMeshVertexCount = SourceMeshDescription.Vertices().Num();
+							const int32 DestinationVertexIndexMax = VertexOffset + SourceMeshVertexCount;
+							if (!ensure(DestinationSkeletalMeshImportData.Points.IsValidIndex(DestinationVertexIndexMax - 1)))
 							{
-								//We break if we get over a valid point it mean we are inspecting a different part of the mesh
+								//This error was suppose to be catch in the pre parallel for loop
 								break;
 							}
-							//Rebuild the data with only the modified point
-							if ((MorphTargetData.Points[OriginalPointIdx] - DestinationSkeletalMeshImportData.Points[PointIdx]).SizeSquared() > FMath::Square(THRESH_POINTS_ARE_SAME))
+							
+							CompressPoints.Reserve(CompressPoints.Num() + SourceMeshVertexCount);
+							FStaticMeshConstAttributes Attributes(SourceMeshDescription);
+							TVertexAttributesConstRef<FVector3f> VertexPositions = Attributes.GetVertexPositions();
+
+							//Reserve the point and influences
+							MorphTargetData.Points.Reset(SourceMeshVertexCount);
+							MorphTargetData.Points.AddZeroed(SourceMeshVertexCount);
+
+							for (FVertexID VertexID : SourceMeshDescription.Vertices().GetElementIDs())
+							{
+								//We can use GetValue because the Meshdescription was compacted before the copy
+								MorphTargetData.Points[VertexID.GetValue()] = (FVector3f)GlobalTransform.TransformPosition((FVector)VertexPositions[VertexID]);
+							}
+
+							for (int32 PointIdx = VertexOffset; PointIdx < DestinationSkeletalMeshImportData.Points.Num(); ++PointIdx)
 							{
-								ModifiedPoints.Add(PointIdx);
-								CompressPoints.Add(MorphTargetData.Points[OriginalPointIdx]);
+								int32 OriginalPointIdx = DestinationSkeletalMeshImportData.PointToRawMap[PointIdx] - VertexOffset;
+								if (!MorphTargetData.Points.IsValidIndex(OriginalPointIdx))
+								{
+									//We break if we get over a valid point it mean we are inspecting a different part of the mesh
+									break;
+								}
+								//Rebuild the data with only the modified point
+								if ((MorphTargetData.Points[OriginalPointIdx] - DestinationSkeletalMeshImportData.Points[PointIdx]).SizeSquared() > FMath::Square(THRESH_POINTS_ARE_SAME))
+								{
+									ModifiedPoints.Add(PointIdx);
+									CompressPoints.Add(MorphTargetData.Points[OriginalPointIdx]);
+								}
 							}
 						}
 						MorphTargetData.Points = CompressPoints;
@@ -626,8 +656,10 @@ namespace UE
 				DestinationImportData = FSkeletalMeshImportData::CreateFromMeshDescription(LodMeshDescription);
 				DestinationImportData.RefBonesBinary = RefBonesBinary;
 
+				bool bMergeMorphTargetWithSameName = false;
+				SkeletalMeshFactoryNode->GetCustomMergeMorphTargetShapeWithSameName(bMergeMorphTargetWithSameName);
 				//Copy all the lod morph targets data to the DestinationImportData.
-				CopyMorphTargetsMeshDescriptionToSkeletalMeshImportData(MorphTargetMeshDescriptionsPerMorphTargetName, DestinationImportData);
+				CopyMorphTargetsMeshDescriptionToSkeletalMeshImportData(MorphTargetMeshDescriptionsPerMorphTargetName, DestinationImportData, bMergeMorphTargetWithSameName);
 			}
 
 			void ProcessImportMeshInfluences(const int32 WedgeCount, TArray<SkeletalMeshImportData::FRawBoneInfluence>& Influences)
diff --git a/Engine/Plugins/Interchange/Runtime/Source/Parsers/Fbx/Private/FbxMesh.cpp b/Engine/Plugins/Interchange/Runtime/Source/Parsers/Fbx/Private/FbxMesh.cpp
--- a/Engine/Plugins/Interchange/Runtime/Source/Parsers/Fbx/Private/FbxMesh.cpp
+++ b/Engine/Plugins/Interchange/Runtime/Source/Parsers/Fbx/Private/FbxMesh.cpp
@@ -1886,6 +1886,17 @@ void FFbxMesh::AddAllMeshes(FbxScene* SDKScene, FbxGeometryConverter* SDKGeometr
 					FString MorphTargetAttributeName = Parser.GetFbxHelper()->GetMeshName(Shape);
 					FString MorphTargetUniqueID = Parser.GetFbxHelper()->GetMeshUniqueID(Shape);
 					const UInterchangeMeshNode* ExistingMorphTargetNode = Cast<const UInterchangeMeshNode>(NodeContainer.GetNode(MorphTargetUniqueID));
+					if (ExistingMorphTargetNode)
+					{
+						int32 UniqueId = 1;
+						FString NameClash = "_ncl_";
+						while (ExistingMorphTargetNode)
+						{
+							MorphTargetUniqueID = Parser.GetFbxHelper()->GetMeshUniqueID(Shape) + NameClash + FString::FromInt(UniqueId++);
+							ExistingMorphTargetNode = Cast<const UInterchangeMeshNode>(NodeContainer.GetNode(MorphTargetUniqueID));
+						}
+					}
+
 					if (!ExistingMorphTargetNode)
 					{
 						UInterchangeMeshNode* MorphTargetNode = CreateMeshNode(NodeContainer, MorphTargetAttributeName, MorphTargetUniqueID);
diff --git a/Engine/Plugins/Interchange/Runtime/Source/Pipelines/Private/InterchangeGenericSkeletalMeshPipeline.cpp b/Engine/Plugins/Interchange/Runtime/Source/Pipelines/Private/InterchangeGenericSkeletalMeshPipeline.cpp
--- a/Engine/Plugins/Interchange/Runtime/Source/Pipelines/Private/InterchangeGenericSkeletalMeshPipeline.cpp
+++ b/Engine/Plugins/Interchange/Runtime/Source/Pipelines/Private/InterchangeGenericSkeletalMeshPipeline.cpp
@@ -501,6 +501,7 @@ UInterchangeSkeletalMeshFactoryNode* UInterchangeGenericMeshPipeline::CreateSkel
 	SkeletalMeshFactoryNode->SetCustomThresholdUV(ThresholdUV);
 	SkeletalMeshFactoryNode->SetCustomMorphThresholdPosition(MorphThresholdPosition);
 	SkeletalMeshFactoryNode->SetCustomBoneInfluenceLimit(BoneInfluenceLimit);
+	SkeletalMeshFactoryNode->SetCustomMergeMorphTargetShapeWithSameName(bMergeMorphTargetsWithSameName);
 
 	return SkeletalMeshFactoryNode;
 }
diff --git a/Engine/Plugins/Interchange/Runtime/Source/Pipelines/Public/InterchangeGenericMeshPipeline.h b/Engine/Plugins/Interchange/Runtime/Source/Pipelines/Public/InterchangeGenericMeshPipeline.h
--- a/Engine/Plugins/Interchange/Runtime/Source/Pipelines/Public/InterchangeGenericMeshPipeline.h
+++ b/Engine/Plugins/Interchange/Runtime/Source/Pipelines/Public/InterchangeGenericMeshPipeline.h
@@ -177,6 +177,10 @@ public:
 	UPROPERTY(EditAnywhere, BlueprintReadWrite, Category = "Skeletal Meshes")
 	bool bImportMorphTargets = true;
 
+	/** If enabled, all morph target shapes with the same name will be merge together. Turn it to false if you want to control those morph with different values. */
+	UPROPERTY(EditAnywhere, BlueprintReadWrite, Category = "Skeletal Meshes", meta = (EditCondition = "bImportMorphTargets"))
+	bool bMergeMorphTargetsWithSameName = true;
+
 	/** If enabled, imports per-vertex attributes from the FBX file. */
 	UPROPERTY(EditAnywhere, BlueprintReadWrite, Category = "Skeletal Meshes", meta = (ToolTip = "If enabled, creates named vertex attributes for secondary vertex color data."))
 	bool bImportVertexAttributes = false;
