MaterialX:
- Refactoring shaders to support custom inputs
- Material Functions Standard Surface (Substrate/Legacy): adding subsurface inputs in translucency even if it's not used to avoid warning/error when import
- Adding specular_color to MX_StandardSurface
- Adding support of coat_anisotropy in legacy mode for StandardSurface and OpenPBR material function
- Improving ThinFilm iridescence in Legacy by changing the factor 0.039 by 0.1, giving closer results to MaterialX viewer

#virtualized

diff --git a/Engine/Plugins/Interchange/Runtime/Source/Import/Private/MaterialX/MaterialXUtils/MaterialXOpenPBRSurfaceShader.cpp b/Engine/Plugins/Interchange/Runtime/Source/Import/Private/MaterialX/MaterialXUtils/MaterialXOpenPBRSurfaceShader.cpp
--- a/Engine/Plugins/Interchange/Runtime/Source/Import/Private/MaterialX/MaterialXUtils/MaterialXOpenPBRSurfaceShader.cpp
+++ b/Engine/Plugins/Interchange/Runtime/Source/Import/Private/MaterialX/MaterialXUtils/MaterialXOpenPBRSurfaceShader.cpp
@@ -23,93 +23,26 @@ TSharedRef<FMaterialXBase> FMaterialXOpenPBRSurfaceShader::MakeInstance(UInterch
 void FMaterialXOpenPBRSurfaceShader::Translate(MaterialX::NodePtr OpenPBRSurfaceNode)
 {
 	using namespace UE::Interchange::Materials;
-	using namespace mx::OpenPBRSurface;
 
 	this->SurfaceShaderNode = OpenPBRSurfaceNode;
-	UInterchangeFunctionCallShaderNode* OpenPBRSurfaceShaderNode = CreateFunctionCallShaderNode(SurfaceShaderNode->getName().c_str(), UE::Interchange::MaterialX::IndexSurfaceShaders, uint8(EInterchangeMaterialXShaders::OpenPBRSurface));
 
-	const bool bTangentSpace = true;
+	UInterchangeShaderNode* OpenPBRSurfaceShaderNode = FMaterialXSurfaceShaderAbstract::Translate(EInterchangeMaterialXShaders::OpenPBRSurface);
 
-	// Inputs
-	//Base
-	ConnectNodeOutputToInput(Input::BaseWeight, OpenPBRSurfaceShaderNode, OpenPBRSurface::Parameters::BaseWeight.ToString(), DefaultValue::BaseWeight);
-	ConnectNodeOutputToInput(Input::BaseColor, OpenPBRSurfaceShaderNode, OpenPBRSurface::Parameters::BaseColor.ToString(), DefaultValue::BaseColor);
-	ConnectNodeOutputToInput(Input::BaseMetalness, OpenPBRSurfaceShaderNode, OpenPBRSurface::Parameters::BaseMetalness.ToString(), DefaultValue::BaseMetalness);
-	ConnectNodeOutputToInput(Input::BaseRoughness, OpenPBRSurfaceShaderNode, OpenPBRSurface::Parameters::BaseRoughness.ToString(), DefaultValue::BaseRoughness);
-
-	//Specular
-	ConnectNodeOutputToInput(Input::SpecularAnisotropy, OpenPBRSurfaceShaderNode, OpenPBRSurface::Parameters::SpecularAnisotropy.ToString(), DefaultValue::SpecularAnisotropy);
-	ConnectNodeOutputToInput(Input::SpecularColor, OpenPBRSurfaceShaderNode, OpenPBRSurface::Parameters::SpecularColor.ToString(), DefaultValue::SpecularColor);
-	ConnectNodeOutputToInput(Input::SpecularIOR, OpenPBRSurfaceShaderNode, OpenPBRSurface::Parameters::SpecularIOR.ToString(), DefaultValue::SpecularIOR);
-	ConnectNodeOutputToInput(Input::SpecularIORLevel, OpenPBRSurfaceShaderNode, OpenPBRSurface::Parameters::SpecularIORLevel.ToString(), DefaultValue::SpecularIORLevel);
-	ConnectNodeOutputToInput(Input::SpecularRotation, OpenPBRSurfaceShaderNode, OpenPBRSurface::Parameters::SpecularRotation.ToString(), DefaultValue::SpecularRotation);
-	ConnectNodeOutputToInput(Input::SpecularRoughness, OpenPBRSurfaceShaderNode, OpenPBRSurface::Parameters::SpecularRoughness.ToString(), DefaultValue::SpecularRoughness);
-	ConnectNodeOutputToInput(Input::SpecularWeight, OpenPBRSurfaceShaderNode, OpenPBRSurface::Parameters::SpecularWeight.ToString(), DefaultValue::SpecularWeight);
-
-	//Coat
-	ConnectNodeOutputToInput(Input::CoatAnisotropy, OpenPBRSurfaceShaderNode, OpenPBRSurface::Parameters::CoatAnisotropy.ToString(), DefaultValue::CoatAnisotropy);
-	ConnectNodeOutputToInput(Input::CoatColor, OpenPBRSurfaceShaderNode, OpenPBRSurface::Parameters::CoatColor.ToString(), DefaultValue::CoatColor);
-	ConnectNodeOutputToInput(Input::CoatIOR, OpenPBRSurfaceShaderNode, OpenPBRSurface::Parameters::CoatIOR.ToString(), DefaultValue::CoatIOR);
-	ConnectNodeOutputToInput(Input::CoatIORLevel, OpenPBRSurfaceShaderNode, OpenPBRSurface::Parameters::CoatIORLevel.ToString(), DefaultValue::CoatIORLevel);
-	ConnectNodeOutputToInput(Input::CoatRotation, OpenPBRSurfaceShaderNode, OpenPBRSurface::Parameters::CoatRotation.ToString(), DefaultValue::CoatRotation);
-	ConnectNodeOutputToInput(Input::CoatRoughness, OpenPBRSurfaceShaderNode, OpenPBRSurface::Parameters::CoatRoughness.ToString(), DefaultValue::CoatRoughness);
-	ConnectNodeOutputToInput(Input::CoatWeight, OpenPBRSurfaceShaderNode, OpenPBRSurface::Parameters::CoatWeight.ToString(), DefaultValue::CoatWeight);
-
-	//Emission
-	ConnectNodeOutputToInput(Input::EmissionColor, OpenPBRSurfaceShaderNode, OpenPBRSurface::Parameters::EmissionColor.ToString(), DefaultValue::EmissionColor);
-	ConnectNodeOutputToInput(Input::EmissionLuminance, OpenPBRSurfaceShaderNode, OpenPBRSurface::Parameters::EmissionLuminance.ToString(), DefaultValue::EmissionLuminance);
-
-	//ThinFilm
-	ConnectNodeOutputToInput(Input::ThinFilmIOR, OpenPBRSurfaceShaderNode, OpenPBRSurface::Parameters::ThinFilmIOR.ToString(), DefaultValue::ThinFilmIOR);
-	ConnectNodeOutputToInput(Input::ThinFilmThickness, OpenPBRSurfaceShaderNode, OpenPBRSurface::Parameters::ThinFilmThickness.ToString(), DefaultValue::ThinFilmThickness);
-
-	//Transmission
-	ConnectNodeOutputToInput(Input::TransmissionColor, OpenPBRSurfaceShaderNode, OpenPBRSurface::Parameters::TransmissionColor.ToString(), DefaultValue::TransmissionColor);
-	ConnectNodeOutputToInput(Input::TransmissionDepth, OpenPBRSurfaceShaderNode, OpenPBRSurface::Parameters::TransmissionDepth.ToString(), DefaultValue::TransmissionDepth);
-	ConnectNodeOutputToInput(Input::TransmissionDispersionScale, OpenPBRSurfaceShaderNode, OpenPBRSurface::Parameters::TransmissionDispersionScale.ToString(), DefaultValue::TransmissionDispersionScale);
-	ConnectNodeOutputToInput(Input::TransmissionDispersionAbbeNumber, OpenPBRSurfaceShaderNode, OpenPBRSurface::Parameters::TransmissionDispersionAbbeNumber.ToString(), DefaultValue::TransmissionDispersionAbbeNumber);
-	ConnectNodeOutputToInput(Input::TransmissionScatter, OpenPBRSurfaceShaderNode, OpenPBRSurface::Parameters::TransmissionScatter.ToString(), DefaultValue::TransmissionScatter);
-	ConnectNodeOutputToInput(Input::TransmissionScatterAnisotropy, OpenPBRSurfaceShaderNode, OpenPBRSurface::Parameters::TransmissionScatterAnisotropy.ToString(), DefaultValue::TransmissionScatterAnisotropy);
-	ConnectNodeOutputToInput(Input::TransmissionWeight, OpenPBRSurfaceShaderNode, OpenPBRSurface::Parameters::TransmissionWeight.ToString(), DefaultValue::TransmissionWeight);
-
-	//Geometry
-	ConnectNodeOutputToInput(Input::GeometryCoatNormal, OpenPBRSurfaceShaderNode, OpenPBRSurface::Parameters::GeometryCoatNormal.ToString(), DefaultValue::GeometryCoatNormal, bTangentSpace);
-	ConnectNodeOutputToInput(Input::GeometryNormal, OpenPBRSurfaceShaderNode, OpenPBRSurface::Parameters::GeometryNormal.ToString(), DefaultValue::GeometryNormal, bTangentSpace);
-	ConnectNodeOutputToInput(Input::GeometryOpacity, OpenPBRSurfaceShaderNode, OpenPBRSurface::Parameters::GeometryOpacity.ToString(), DefaultValue::GeometryOpacity);
-	ConnectNodeOutputToInput(Input::GeometryTangent, OpenPBRSurfaceShaderNode, OpenPBRSurface::Parameters::GeometryTangent.ToString(), DefaultValue::GeometryTangent, bTangentSpace);
-	ConnectNodeOutputToInput(Input::GeometryThinWalled, OpenPBRSurfaceShaderNode, OpenPBRSurface::Parameters::GeometryThinWalled.ToString(), DefaultValue::GeometryThinWalled);
+	//Two sided
+	if(MaterialX::InputPtr Input = GetInput(SurfaceShaderNode, mx::OpenPBRSurface::Input::GeometryThinWalled);
+	   Input->hasValue() && mx::fromValueString<bool>(Input->getValueString()) == true)
 	{
-		MaterialX::DocumentPtr Document = SurfaceShaderNode->getDocument();
-		MaterialX::InputPtr Input = GetInput(SurfaceShaderNode, mx::OpenPBRSurface::Input::GeometryThinWalled);
-		if(Input->hasValue() && mx::fromValueString<bool>(Input->getValueString()) == true)
-		{
-			// weird that we also have to enable that to have a two sided material (seems to only have meaning for Translucent material)
-			ShaderGraphNode->SetCustomTwoSidedTransmission(true);
-			ShaderGraphNode->SetCustomTwoSided(true);
-		}
+		// weird that we also have to enable that to have a two sided material (seems to only have meaning for Translucent material)
+		ShaderGraphNode->SetCustomTwoSidedTransmission(true);
+		ShaderGraphNode->SetCustomTwoSided(true);
 	}
 
-	//Fuzz
-	ConnectNodeOutputToInput(Input::FuzzColor, OpenPBRSurfaceShaderNode, OpenPBRSurface::Parameters::FuzzColor.ToString(), DefaultValue::FuzzColor);
-	ConnectNodeOutputToInput(Input::FuzzRoughness, OpenPBRSurfaceShaderNode, OpenPBRSurface::Parameters::FuzzRoughness.ToString(), DefaultValue::FuzzRoughness);
-	ConnectNodeOutputToInput(Input::FuzzWeight, OpenPBRSurfaceShaderNode, OpenPBRSurface::Parameters::FuzzWeight.ToString(), DefaultValue::FuzzWeight);
-
-	// We can't have Subsurface inputs if we have Transmission inputs
 	if(UInterchangeShaderPortsAPI::HasInput(OpenPBRSurfaceShaderNode, OpenPBRSurface::Parameters::TransmissionWeight))
 	{
 		OpenPBRSurfaceShaderNode->AddInt32Attribute(UE::Interchange::MaterialX::Attributes::EnumType, UE::Interchange::MaterialX::IndexSurfaceShaders);
 		OpenPBRSurfaceShaderNode->AddInt32Attribute(UE::Interchange::MaterialX::Attributes::EnumValue, int32(EInterchangeMaterialXShaders::OpenPBRSurfaceTransmission));
 		ShaderGraphNode->SetCustomBlendMode(EBlendMode::BLEND_TranslucentColoredTransmittance);
 	}
-	else
-	{
-		//Subsurface
-		ConnectNodeOutputToInput(Input::SubsurfaceAnisotropy, OpenPBRSurfaceShaderNode, OpenPBRSurface::Parameters::SubsurfaceAnisotropy.ToString(), DefaultValue::SubsurfaceAnisotropy);
-		ConnectNodeOutputToInput(Input::SubsurfaceColor, OpenPBRSurfaceShaderNode, OpenPBRSurface::Parameters::SubsurfaceColor.ToString(), DefaultValue::SubsurfaceColor);
-		ConnectNodeOutputToInput(Input::SubsurfaceRadius, OpenPBRSurfaceShaderNode, OpenPBRSurface::Parameters::SubsurfaceRadius.ToString(), DefaultValue::SubsurfaceRadius);
-		ConnectNodeOutputToInput(Input::SubsurfaceRadiusScale, OpenPBRSurfaceShaderNode, OpenPBRSurface::Parameters::SubsurfaceRadiusScale.ToString(), DefaultValue::SubsurfaceRadiusScale);
-		ConnectNodeOutputToInput(Input::SubsurfaceWeight, OpenPBRSurfaceShaderNode, OpenPBRSurface::Parameters::SubsurfaceWeight.ToString(), DefaultValue::SubsurfaceWeight);
-	}
 
 	if(bIsSubstrateEnabled)
 	{
@@ -156,8 +89,6 @@ void FMaterialXOpenPBRSurfaceShader::Translate(MaterialX::NodePtr OpenPBRSurface
 			UInterchangeShaderPortsAPI::ConnectOuputToInputByName(ShaderGraphNode, Subsurface::Parameters::SubsurfaceColor.ToString(), OpenPBRSurfaceShaderNode->GetUniqueID(), Subsurface::Parameters::SubsurfaceColor.ToString());
 		}
 	}
-
-	// Outputs
 }
 
 #endif //WITH_EDITOR
diff --git a/Engine/Plugins/Interchange/Runtime/Source/Import/Private/MaterialX/MaterialXUtils/MaterialXStandardSurfaceShader.cpp b/Engine/Plugins/Interchange/Runtime/Source/Import/Private/MaterialX/MaterialXUtils/MaterialXStandardSurfaceShader.cpp
--- a/Engine/Plugins/Interchange/Runtime/Source/Import/Private/MaterialX/MaterialXUtils/MaterialXStandardSurfaceShader.cpp
+++ b/Engine/Plugins/Interchange/Runtime/Source/Import/Private/MaterialX/MaterialXUtils/MaterialXStandardSurfaceShader.cpp
@@ -2,7 +2,6 @@
 
 #if WITH_EDITOR
 #include "MaterialXStandardSurfaceShader.h"
-
 #include "Engine/EngineTypes.h"
 
 namespace mx = MaterialX;
@@ -26,15 +25,15 @@ void FMaterialXStandardSurfaceShader::Translate(mx::NodePtr StandardSurfaceNode)
 	
 	using namespace UE::Interchange::Materials;
 
-	UInterchangeFunctionCallShaderNode* StandardSurfaceShaderNode = CreateFunctionCallShaderNode(SurfaceShaderNode->getName().c_str(), UE::Interchange::MaterialX::IndexSurfaceShaders, uint8(EInterchangeMaterialXShaders::StandardSurface));
+	UInterchangeShaderNode* StandardSurfaceShaderNode = FMaterialXSurfaceShaderAbstract::Translate(EInterchangeMaterialXShaders::StandardSurface);
 
-	if(bIsSubstrateEnabled)
-	{
-		ConnectToSubstrateStandardSurface(StandardSurfaceShaderNode);
-	}
-	else
+	//Two sided
+	if(MaterialX::InputPtr Input = GetInput(SurfaceShaderNode, mx::StandardSurface::Input::ThinWalled);
+	   Input->hasValue() && mx::fromValueString<bool>(Input->getValueString()) == true)
 	{
-		ConnectToStandardSurface(StandardSurfaceShaderNode);
+		// weird that we also have to enable that to have a two sided material (seems to only have meaning for Translucent material)
+		ShaderGraphNode->SetCustomTwoSidedTransmission(true);
+		ShaderGraphNode->SetCustomTwoSided(true);
 	}
 
 	if(UInterchangeShaderPortsAPI::HasInput(StandardSurfaceShaderNode, StandardSurface::Parameters::Transmission))
@@ -42,319 +41,72 @@ void FMaterialXStandardSurfaceShader::Translate(mx::NodePtr StandardSurfaceNode)
 		StandardSurfaceShaderNode->AddInt32Attribute(UE::Interchange::MaterialX::Attributes::EnumType, UE::Interchange::MaterialX::IndexSurfaceShaders);
 		StandardSurfaceShaderNode->AddInt32Attribute(UE::Interchange::MaterialX::Attributes::EnumValue, int32(EInterchangeMaterialXShaders::StandardSurfaceTransmission));
 	}
-}
-
-void FMaterialXStandardSurfaceShader::ConnectToStandardSurface(UInterchangeFunctionCallShaderNode* StandardSurfaceShaderNode)
-{
-	using namespace UE::Interchange::Materials;
-	using namespace mx::StandardSurface;
-	constexpr bool bInputInTangentSpace = true;
-
-	// Inputs
-	//Base
-	ConnectNodeOutputToInput(Input::Base, StandardSurfaceShaderNode, StandardSurface::Parameters::Base.ToString(), DefaultValue::Base);
-
-	//Base Color
-	ConnectNodeOutputToInput(Input::BaseColor, StandardSurfaceShaderNode, StandardSurface::Parameters::BaseColor.ToString(), DefaultValue::BaseColor);
-
-	//Diffuse Roughness
-	ConnectNodeOutputToInput(Input::DiffuseRoughness, StandardSurfaceShaderNode, StandardSurface::Parameters::DiffuseRoughness.ToString(), DefaultValue::DiffuseRoughness);
-
-	//Specular
-	ConnectNodeOutputToInput(Input::Specular, StandardSurfaceShaderNode, StandardSurface::Parameters::Specular.ToString(), DefaultValue::Specular);
-
-	//Specular Roughness
-	ConnectNodeOutputToInput(Input::SpecularRoughness, StandardSurfaceShaderNode, StandardSurface::Parameters::SpecularRoughness.ToString(), DefaultValue::SpecularRoughness);
-
-	//Specular IOR
-	ConnectNodeOutputToInput(Input::SpecularIOR, StandardSurfaceShaderNode, StandardSurface::Parameters::SpecularIOR.ToString(), DefaultValue::SpecularIOR);
-
-	//Specular Anisotropy
-	ConnectNodeOutputToInput(Input::SpecularAnisotropy, StandardSurfaceShaderNode, StandardSurface::Parameters::SpecularAnisotropy.ToString(), DefaultValue::SpecularAnisotropy);
-
-	//Specular Rotation
-	ConnectNodeOutputToInput(Input::SpecularRotation, StandardSurfaceShaderNode, StandardSurface::Parameters::SpecularRotation.ToString(), DefaultValue::SpecularRotation);
-
-	//Metallic
-	ConnectNodeOutputToInput(Input::Metalness, StandardSurfaceShaderNode, StandardSurface::Parameters::Metalness.ToString(), DefaultValue::Metalness);
-
-	//Subsurface
-	ConnectNodeOutputToInput(Input::Subsurface, StandardSurfaceShaderNode, StandardSurface::Parameters::Subsurface.ToString(), DefaultValue::Subsurface);
-
-	//Subsurface Color
-	ConnectNodeOutputToInput(Input::SubsurfaceColor, StandardSurfaceShaderNode, StandardSurface::Parameters::SubsurfaceColor.ToString(), DefaultValue::SubsurfaceColor);
-
-	//Subsurface Radius
-	ConnectNodeOutputToInput(Input::SubsurfaceRadius, StandardSurfaceShaderNode, StandardSurface::Parameters::SubsurfaceRadius.ToString(), DefaultValue::SubsurfaceRadius);
-
-	//Subsurface Scale
-	ConnectNodeOutputToInput(Input::SubsurfaceScale, StandardSurfaceShaderNode, StandardSurface::Parameters::SubsurfaceScale.ToString(), DefaultValue::SubsurfaceScale);
-
-	//Sheen
-	ConnectNodeOutputToInput(Input::Sheen, StandardSurfaceShaderNode, StandardSurface::Parameters::Sheen.ToString(), DefaultValue::Sheen);
-
-	//Sheen Color
-	ConnectNodeOutputToInput(Input::SheenColor, StandardSurfaceShaderNode, StandardSurface::Parameters::SheenColor.ToString(), DefaultValue::SheenColor);
-
-	//Sheen Roughness
-	ConnectNodeOutputToInput(Input::SheenRoughness, StandardSurfaceShaderNode, StandardSurface::Parameters::SheenRoughness.ToString(), DefaultValue::SheenRoughness);
-
-	//Coat
-	ConnectNodeOutputToInput(Input::Coat, StandardSurfaceShaderNode, StandardSurface::Parameters::Coat.ToString(), DefaultValue::Coat);
-
-	//Coat Color
-	ConnectNodeOutputToInput(Input::CoatColor, StandardSurfaceShaderNode, StandardSurface::Parameters::CoatColor.ToString(), DefaultValue::CoatColor);
-
-	//Coat Roughness
-	ConnectNodeOutputToInput(Input::CoatRoughness, StandardSurfaceShaderNode, StandardSurface::Parameters::CoatRoughness.ToString(), DefaultValue::CoatRoughness);
-
-	//Coat Normal: No need to take the default input if there is no CoatNormal input
-	ConnectNodeOutputToInput(Input::CoatNormal, StandardSurfaceShaderNode, StandardSurface::Parameters::CoatNormal.ToString(), nullptr, bInputInTangentSpace);
-
-	//Thin Film Thickness
-	ConnectNodeOutputToInput(Input::ThinFilmThickness, StandardSurfaceShaderNode, StandardSurface::Parameters::ThinFilmThickness.ToString(), DefaultValue::ThinFilmThickness);
-
-	//Thin Film IOR
-	ConnectNodeOutputToInput(Input::ThinFilmIOR, StandardSurfaceShaderNode, StandardSurface::Parameters::ThinFilmIOR.ToString(), DefaultValue::ThinFilmIOR);
-
-	//Emission
-	ConnectNodeOutputToInput(Input::Emission, StandardSurfaceShaderNode, StandardSurface::Parameters::Emission.ToString(), DefaultValue::Emission);
-
-	//Emission Color
-	ConnectNodeOutputToInput(Input::EmissionColor, StandardSurfaceShaderNode, StandardSurface::Parameters::EmissionColor.ToString(), DefaultValue::EmissionColor);
-
-	//Normal: No need to take the default input if there is no Normal input
-	ConnectNodeOutputToInput(Input::Normal, StandardSurfaceShaderNode, StandardSurface::Parameters::Normal.ToString(), nullptr, bInputInTangentSpace);
-
-	//Tangent: No need to take the default input if there is no Tangent input
-	ConnectNodeOutputToInput(Input::Tangent, StandardSurfaceShaderNode, StandardSurface::Parameters::Tangent.ToString(), nullptr, bInputInTangentSpace);
-
-	//Transmission
-	ConnectNodeOutputToInput(Input::Transmission, StandardSurfaceShaderNode, StandardSurface::Parameters::Transmission.ToString(), DefaultValue::Transmission);
-
-	//Transmission Color
-	ConnectNodeOutputToInput(Input::TransmissionColor, StandardSurfaceShaderNode, StandardSurface::Parameters::TransmissionColor.ToString(), DefaultValue::TransmissionColor);
-
-	//Transmission Depth
-	ConnectNodeOutputToInput(Input::TransmissionDepth, StandardSurfaceShaderNode, StandardSurface::Parameters::TransmissionDepth.ToString(), DefaultValue::TransmissionDepth);
-
-	//Transmission Scatter
-	ConnectNodeOutputToInput(Input::TransmissionScatter, StandardSurfaceShaderNode, StandardSurface::Parameters::TransmissionScatter.ToString(), DefaultValue::TransmissionScatter);
-
-	//Transmission Scatter Anisotropy
-	ConnectNodeOutputToInput(Input::TransmissionScatterAnisotropy, StandardSurfaceShaderNode, StandardSurface::Parameters::TransmissionScatterAnisotropy.ToString(), DefaultValue::TransmissionScatterAnisotropy);
-
-	//Transmission Dispersion
-	ConnectNodeOutputToInput(Input::TransmissionDispersion, StandardSurfaceShaderNode, StandardSurface::Parameters::TransmissionDispersion.ToString(), DefaultValue::TransmissionDispersion);
-
-	//Transmission Extra Roughness
-	ConnectNodeOutputToInput(Input::TransmissionExtraRoughness, StandardSurfaceShaderNode, StandardSurface::Parameters::TransmissionExtraRoughness.ToString(), DefaultValue::TransmissionExtraRoughness);
-
-	//Opacity
-	ConnectNodeOutputToInput(Input::Opacity, StandardSurfaceShaderNode, StandardSurface::Parameters::Opacity.ToString(), DefaultValue::Opacity);
-
-	//Thin Walled: we create a shader graph for that input even if it has no real-meaning, we only check the value to enable the Two-Sided option
-	ConnectNodeOutputToInput(Input::ThinWalled, StandardSurfaceShaderNode, StandardSurface::Parameters::ThinWalled.ToString(), DefaultValue::ThinWalled);
-	{
-		MaterialX::DocumentPtr Document = SurfaceShaderNode->getDocument();
-		MaterialX::InputPtr Input = GetInput(SurfaceShaderNode, mx::StandardSurface::Input::ThinWalled);
-		if(Input->hasValue() && mx::fromValueString<bool>(Input->getValueString()) == true)
-		{
-			ShaderGraphNode->SetCustomTwoSided(true);
-		}
-	}
 
 	// Outputs
-	UInterchangeShaderPortsAPI::ConnectOuputToInputByName(ShaderGraphNode, PBRMR::Parameters::BaseColor.ToString(), StandardSurfaceShaderNode->GetUniqueID(), TEXT("Base Color"));
-	UInterchangeShaderPortsAPI::ConnectOuputToInputByName(ShaderGraphNode, PBRMR::Parameters::Metallic.ToString(), StandardSurfaceShaderNode->GetUniqueID(), PBRMR::Parameters::Metallic.ToString());
-	UInterchangeShaderPortsAPI::ConnectOuputToInputByName(ShaderGraphNode, PBRMR::Parameters::Specular.ToString(), StandardSurfaceShaderNode->GetUniqueID(), PBRMR::Parameters::Specular.ToString());
-	UInterchangeShaderPortsAPI::ConnectOuputToInputByName(ShaderGraphNode, PBRMR::Parameters::Roughness.ToString(), StandardSurfaceShaderNode->GetUniqueID(), PBRMR::Parameters::Roughness.ToString());
-	UInterchangeShaderPortsAPI::ConnectOuputToInputByName(ShaderGraphNode, PBRMR::Parameters::EmissiveColor.ToString(), StandardSurfaceShaderNode->GetUniqueID(), PBRMR::Parameters::EmissiveColor.ToString());
-	UInterchangeShaderPortsAPI::ConnectOuputToInputByName(ShaderGraphNode, PBRMR::Parameters::Anisotropy.ToString(), StandardSurfaceShaderNode->GetUniqueID(), PBRMR::Parameters::Anisotropy.ToString());
-	UInterchangeShaderPortsAPI::ConnectOuputToInputByName(ShaderGraphNode, PBRMR::Parameters::Normal.ToString(), StandardSurfaceShaderNode->GetUniqueID(), PBRMR::Parameters::Normal.ToString());
-	UInterchangeShaderPortsAPI::ConnectOuputToInputByName(ShaderGraphNode, PBRMR::Parameters::Tangent.ToString(), StandardSurfaceShaderNode->GetUniqueID(), PBRMR::Parameters::Tangent.ToString());
-
-	// We can't have all shading models at once, so we have to make a choice here
-	if(UInterchangeShaderPortsAPI::HasInput(StandardSurfaceShaderNode, StandardSurface::Parameters::Transmission))
+	if(bIsSubstrateEnabled)
 	{
-		UInterchangeShaderPortsAPI::ConnectOuputToInputByName(ShaderGraphNode, PBRMR::Parameters::Opacity.ToString(), StandardSurfaceShaderNode->GetUniqueID(), PBRMR::Parameters::Opacity.ToString());
-		UInterchangeShaderPortsAPI::ConnectOuputToInputByName(ShaderGraphNode, ThinTranslucent::Parameters::TransmissionColor.ToString(), StandardSurfaceShaderNode->GetUniqueID(), ThinTranslucent::Parameters::TransmissionColor.ToString());
-		UInterchangeShaderPortsAPI::ConnectOuputToInputByName(ShaderGraphNode, Common::Parameters::Refraction.ToString(), StandardSurfaceShaderNode->GetUniqueID(), Common::Parameters::Refraction.ToString());
-
-		// If we have have Transmission and Opacity let's use the Surface Coverage instead
-		if(UInterchangeShaderPortsAPI::HasInput(StandardSurfaceShaderNode, StandardSurface::Parameters::Opacity))
+		if(UInterchangeShaderPortsAPI::HasInput(StandardSurfaceShaderNode, StandardSurface::Parameters::Transmission))
 		{
-			UInterchangeShaderPortsAPI::ConnectOuputToInputByName(ShaderGraphNode, ThinTranslucent::Parameters::SurfaceCoverage.ToString(), StandardSurfaceShaderNode->GetUniqueID(), ThinTranslucent::Parameters::SurfaceCoverage.ToString());
+			ShaderGraphNode->SetCustomBlendMode(EBlendMode::BLEND_TranslucentColoredTransmittance);
+			UInterchangeShaderPortsAPI::ConnectOuputToInputByName(ShaderGraphNode, PBRMR::Parameters::Opacity.ToString(), StandardSurfaceShaderNode->GetUniqueID(), PBRMR::Parameters::Opacity.ToString());
+			UInterchangeShaderPortsAPI::ConnectOuputToInputByName(ShaderGraphNode, SubstrateMaterial::Parameters::FrontMaterial.ToString(), StandardSurfaceShaderNode->GetUniqueID(), StandardSurface::SubstrateMaterial::Outputs::Translucent.ToString());
 		}
-	}
-	else if(UInterchangeShaderPortsAPI::HasInput(StandardSurfaceShaderNode, StandardSurface::Parameters::Sheen))
-	{
-		UInterchangeShaderPortsAPI::ConnectOuputToInputByName(ShaderGraphNode, Sheen::Parameters::SheenColor.ToString(), StandardSurfaceShaderNode->GetUniqueID(), Sheen::Parameters::SheenColor.ToString());
-		UInterchangeShaderPortsAPI::ConnectOuputToInputByName(ShaderGraphNode, Sheen::Parameters::SheenRoughness.ToString(), StandardSurfaceShaderNode->GetUniqueID(), Sheen::Parameters::SheenRoughness.ToString());
-	}
-	else if(UInterchangeShaderPortsAPI::HasInput(StandardSurfaceShaderNode, StandardSurface::Parameters::Coat))
-	{
-		UInterchangeShaderPortsAPI::ConnectOuputToInputByName(ShaderGraphNode, ClearCoat::Parameters::ClearCoat.ToString(), StandardSurfaceShaderNode->GetUniqueID(), ClearCoat::Parameters::ClearCoat.ToString());
-		UInterchangeShaderPortsAPI::ConnectOuputToInputByName(ShaderGraphNode, ClearCoat::Parameters::ClearCoatRoughness.ToString(), StandardSurfaceShaderNode->GetUniqueID(), ClearCoat::Parameters::ClearCoatRoughness.ToString());
-		UInterchangeShaderPortsAPI::ConnectOuputToInputByName(ShaderGraphNode, ClearCoat::Parameters::ClearCoatNormal.ToString(), StandardSurfaceShaderNode->GetUniqueID(), ClearCoat::Parameters::ClearCoatNormal.ToString());
-	}
-	else if(UInterchangeShaderPortsAPI::HasInput(StandardSurfaceShaderNode, StandardSurface::Parameters::Subsurface))
-	{
-		UInterchangeShaderPortsAPI::ConnectOuputToInputByName(ShaderGraphNode, PBRMR::Parameters::Opacity.ToString(), StandardSurfaceShaderNode->GetUniqueID(), PBRMR::Parameters::Opacity.ToString());
-		UInterchangeShaderPortsAPI::ConnectOuputToInputByName(ShaderGraphNode, Subsurface::Parameters::SubsurfaceColor.ToString(), StandardSurfaceShaderNode->GetUniqueID(), Subsurface::Parameters::SubsurfaceColor.ToString());
-	}
-	else if(UInterchangeShaderPortsAPI::HasInput(StandardSurfaceShaderNode, StandardSurface::Parameters::Opacity))
-	{
-		UInterchangeShaderPortsAPI::ConnectOuputToInputByName(ShaderGraphNode, ThinTranslucent::Parameters::SurfaceCoverage.ToString(), StandardSurfaceShaderNode->GetUniqueID(), ThinTranslucent::Parameters::SurfaceCoverage.ToString());
-	}
-}
-
-void FMaterialXStandardSurfaceShader::ConnectToSubstrateStandardSurface(UInterchangeFunctionCallShaderNode* StandardSurfaceShaderNode)
-{
-	using namespace UE::Interchange::Materials;
-	using namespace mx::StandardSurface;
-	constexpr bool bInputInTangentSpace = true;
-
-	// Inputs
-	//Base
-	ConnectNodeOutputToInput(Input::Base, StandardSurfaceShaderNode, StandardSurface::Parameters::Base.ToString(), DefaultValue::Base);
-
-	//Base Color
-	ConnectNodeOutputToInput(Input::BaseColor, StandardSurfaceShaderNode, StandardSurface::Parameters::BaseColor.ToString(), DefaultValue::BaseColor);
-
-	//Diffuse Roughness
-	ConnectNodeOutputToInput(Input::DiffuseRoughness, StandardSurfaceShaderNode, StandardSurface::Parameters::DiffuseRoughness.ToString(), DefaultValue::DiffuseRoughness);
-
-	//Metallic
-	ConnectNodeOutputToInput(Input::Metalness, StandardSurfaceShaderNode, StandardSurface::Parameters::Metalness.ToString(), DefaultValue::Metalness);
-
-	//Specular
-	ConnectNodeOutputToInput(Input::Specular, StandardSurfaceShaderNode, StandardSurface::Parameters::Specular.ToString(), DefaultValue::Specular);
-
-	//Specular Color
-	ConnectNodeOutputToInput(Input::SpecularColor, StandardSurfaceShaderNode, StandardSurface::Parameters::SpecularColor.ToString(), DefaultValue::SpecularColor);
-
-	//Specular Roughness
-	ConnectNodeOutputToInput(Input::SpecularRoughness, StandardSurfaceShaderNode, StandardSurface::Parameters::SpecularRoughness.ToString(), DefaultValue::SpecularRoughness);
-
-	//Specular IOR
-	ConnectNodeOutputToInput(Input::SpecularIOR, StandardSurfaceShaderNode, StandardSurface::Parameters::SpecularIOR.ToString(), DefaultValue::SpecularIOR);
-
-	//Specular Anisotropy
-	ConnectNodeOutputToInput(Input::SpecularAnisotropy, StandardSurfaceShaderNode, StandardSurface::Parameters::SpecularAnisotropy.ToString(), DefaultValue::SpecularAnisotropy);
-
-	//Specular Rotation
-	ConnectNodeOutputToInput(Input::SpecularRotation, StandardSurfaceShaderNode, StandardSurface::Parameters::SpecularRotation.ToString(), DefaultValue::SpecularRotation);
-
-	//Sheen
-	ConnectNodeOutputToInput(Input::Sheen, StandardSurfaceShaderNode, StandardSurface::Parameters::Sheen.ToString(), DefaultValue::Sheen);
-
-	//Sheen Color
-	ConnectNodeOutputToInput(Input::SheenColor, StandardSurfaceShaderNode, StandardSurface::Parameters::SheenColor.ToString(), DefaultValue::SheenColor);
-
-	//Sheen Roughness
-	ConnectNodeOutputToInput(Input::SheenRoughness, StandardSurfaceShaderNode, StandardSurface::Parameters::SheenRoughness.ToString(), DefaultValue::SheenRoughness);
-
-	//Coat
-	ConnectNodeOutputToInput(Input::Coat, StandardSurfaceShaderNode, StandardSurface::Parameters::Coat.ToString(), DefaultValue::Coat);
-
-	//Coat Color
-	ConnectNodeOutputToInput(Input::CoatColor, StandardSurfaceShaderNode, StandardSurface::Parameters::CoatColor.ToString(), DefaultValue::CoatColor);
-
-	//Coat Roughness
-	ConnectNodeOutputToInput(Input::CoatRoughness, StandardSurfaceShaderNode, StandardSurface::Parameters::CoatRoughness.ToString(), DefaultValue::CoatRoughness);
-
-	//Coat IOR
-	ConnectNodeOutputToInput(Input::CoatIOR, StandardSurfaceShaderNode, StandardSurface::Parameters::CoatIOR.ToString(), DefaultValue::CoatIOR);
-
-	//Coat Anisotropy
-	ConnectNodeOutputToInput(Input::CoatAnisotropy, StandardSurfaceShaderNode, StandardSurface::Parameters::CoatAnisotropy.ToString(), DefaultValue::CoatAnisotropy);
-
-	//Coat Rotation
-	ConnectNodeOutputToInput(Input::CoatRotation, StandardSurfaceShaderNode, StandardSurface::Parameters::CoatRotation.ToString(), DefaultValue::CoatAnisotropy);
-
-	//Coat Normal: No need to take the default input if there is no CoatNormal input
-	ConnectNodeOutputToInput(Input::CoatNormal, StandardSurfaceShaderNode, StandardSurface::Parameters::CoatNormal.ToString(), nullptr, bInputInTangentSpace);
-
-	//Thin Film Thickness
-	ConnectNodeOutputToInput(Input::ThinFilmThickness, StandardSurfaceShaderNode, StandardSurface::Parameters::ThinFilmThickness.ToString(), DefaultValue::ThinFilmThickness);
-
-	//Thin Film IOR
-	ConnectNodeOutputToInput(Input::ThinFilmIOR, StandardSurfaceShaderNode, StandardSurface::Parameters::ThinFilmIOR.ToString(), DefaultValue::ThinFilmIOR);
-
-	//Emission
-	ConnectNodeOutputToInput(Input::Emission, StandardSurfaceShaderNode, StandardSurface::Parameters::Emission.ToString(), DefaultValue::Emission);
-
-	//Emission Color
-	ConnectNodeOutputToInput(Input::EmissionColor, StandardSurfaceShaderNode, StandardSurface::Parameters::EmissionColor.ToString(), DefaultValue::EmissionColor);
-
-	//Normal: No need to take the default input if there is no Normal input
-	ConnectNodeOutputToInput(Input::Normal, StandardSurfaceShaderNode, StandardSurface::Parameters::Normal.ToString(), nullptr, bInputInTangentSpace);
-
-	//Tangent: No need to take the default input if there is no Tangent input
-	ConnectNodeOutputToInput(Input::Tangent, StandardSurfaceShaderNode, StandardSurface::Parameters::Tangent.ToString(), nullptr, bInputInTangentSpace);
-
-	//Opacity
-	ConnectNodeOutputToInput(Input::Opacity, StandardSurfaceShaderNode, StandardSurface::Parameters::Opacity.ToString(), DefaultValue::Opacity);
-
-	//Transmission
-	ConnectNodeOutputToInput(Input::Transmission, StandardSurfaceShaderNode, StandardSurface::Parameters::Transmission.ToString(), DefaultValue::Transmission);
-
-	//Transmission Color
-	ConnectNodeOutputToInput(Input::TransmissionColor, StandardSurfaceShaderNode, StandardSurface::Parameters::TransmissionColor.ToString(), DefaultValue::TransmissionColor);
-
-	//Transmission Depth
-	ConnectNodeOutputToInput(Input::TransmissionDepth, StandardSurfaceShaderNode, StandardSurface::Parameters::TransmissionDepth.ToString(), DefaultValue::TransmissionDepth);
-
-	//Transmission Scatter
-	ConnectNodeOutputToInput(Input::TransmissionScatter, StandardSurfaceShaderNode, StandardSurface::Parameters::TransmissionScatter.ToString(), DefaultValue::TransmissionScatter);
-
-	//Transmission Extra Roughness
-	ConnectNodeOutputToInput(Input::TransmissionExtraRoughness, StandardSurfaceShaderNode, StandardSurface::Parameters::TransmissionExtraRoughness.ToString(), DefaultValue::TransmissionExtraRoughness);
-
-	//Thin Walled: we create a shader graph for that input even if it has no real-meaning, we only check the value to enable the Two-Sided option
-	ConnectNodeOutputToInput(Input::ThinWalled, StandardSurfaceShaderNode, StandardSurface::Parameters::ThinWalled.ToString(), DefaultValue::ThinWalled);
-	{
-		MaterialX::DocumentPtr Document = SurfaceShaderNode->getDocument();
-		MaterialX::InputPtr Input = GetInput(SurfaceShaderNode, mx::StandardSurface::Input::ThinWalled);
-		if(Input->hasValue() && mx::fromValueString<bool>(Input->getValueString()) == true)
+		else
 		{
-			// weird that we also have to enable that to have a two sided material (seems to only have meaning for Translucent material)
-			ShaderGraphNode->SetCustomTwoSidedTransmission(true);
-			ShaderGraphNode->SetCustomTwoSided(true);
-		}
-	}
 
-	// Outputs
-	if(UInterchangeShaderPortsAPI::HasInput(StandardSurfaceShaderNode, StandardSurface::Parameters::Transmission))
-	{
-		ShaderGraphNode->SetCustomBlendMode(EBlendMode::BLEND_TranslucentColoredTransmittance);
-		UInterchangeShaderPortsAPI::ConnectOuputToInputByName(ShaderGraphNode, PBRMR::Parameters::Opacity.ToString(), StandardSurfaceShaderNode->GetUniqueID(), PBRMR::Parameters::Opacity.ToString());
-		UInterchangeShaderPortsAPI::ConnectOuputToInputByName(ShaderGraphNode, SubstrateMaterial::Parameters::FrontMaterial.ToString(), StandardSurfaceShaderNode->GetUniqueID(), StandardSurface::SubstrateMaterial::Outputs::Translucent.ToString());
+			UInterchangeShaderPortsAPI::ConnectOuputToInputByName(ShaderGraphNode, SubstrateMaterial::Parameters::FrontMaterial.ToString(), StandardSurfaceShaderNode->GetUniqueID(), StandardSurface::SubstrateMaterial::Outputs::Opaque.ToString());
+			if(UInterchangeShaderPortsAPI::HasInput(StandardSurfaceShaderNode, StandardSurface::Parameters::Opacity))
+			{
+				UInterchangeShaderPortsAPI::ConnectOuputToInputByName(ShaderGraphNode, SubstrateMaterial::Parameters::OpacityMask.ToString(), StandardSurfaceShaderNode->GetUniqueID(), StandardSurface::SubstrateMaterial::Outputs::Opacity.ToString());
+				ShaderGraphNode->SetCustomBlendMode(EBlendMode::BLEND_Masked);
+			}
+		}
 	}
 	else
 	{
-		//Subsurface
-		ConnectNodeOutputToInput(Input::Subsurface, StandardSurfaceShaderNode, StandardSurface::Parameters::Subsurface.ToString(), DefaultValue::Subsurface);
-
-		//Subsurface Color
-		ConnectNodeOutputToInput(Input::SubsurfaceColor, StandardSurfaceShaderNode, StandardSurface::Parameters::SubsurfaceColor.ToString(), DefaultValue::SubsurfaceColor);
-
-		//Subsurface Radius
-		ConnectNodeOutputToInput(Input::SubsurfaceRadius, StandardSurfaceShaderNode, StandardSurface::Parameters::SubsurfaceRadius.ToString(), DefaultValue::SubsurfaceRadius);
-
-		//Subsurface Scale
-		ConnectNodeOutputToInput(Input::SubsurfaceScale, StandardSurfaceShaderNode, StandardSurface::Parameters::SubsurfaceScale.ToString(), DefaultValue::SubsurfaceScale);
-
-		//Subsurface Anisotropy
-		ConnectNodeOutputToInput(Input::SubsurfaceAnisotropy, StandardSurfaceShaderNode, StandardSurface::Parameters::SubsurfaceAnisotropy.ToString(), DefaultValue::SubsurfaceAnisotropy);
-
-		UInterchangeShaderPortsAPI::ConnectOuputToInputByName(ShaderGraphNode, SubstrateMaterial::Parameters::FrontMaterial.ToString(), StandardSurfaceShaderNode->GetUniqueID(), StandardSurface::SubstrateMaterial::Outputs::Opaque.ToString());
-		if(UInterchangeShaderPortsAPI::HasInput(StandardSurfaceShaderNode, StandardSurface::Parameters::Opacity))
+		UInterchangeShaderPortsAPI::ConnectOuputToInputByName(ShaderGraphNode, PBRMR::Parameters::BaseColor.ToString(), StandardSurfaceShaderNode->GetUniqueID(), TEXT("Base Color"));
+		UInterchangeShaderPortsAPI::ConnectOuputToInputByName(ShaderGraphNode, PBRMR::Parameters::Metallic.ToString(), StandardSurfaceShaderNode->GetUniqueID(), PBRMR::Parameters::Metallic.ToString());
+		UInterchangeShaderPortsAPI::ConnectOuputToInputByName(ShaderGraphNode, PBRMR::Parameters::Specular.ToString(), StandardSurfaceShaderNode->GetUniqueID(), PBRMR::Parameters::Specular.ToString());
+		UInterchangeShaderPortsAPI::ConnectOuputToInputByName(ShaderGraphNode, PBRMR::Parameters::Roughness.ToString(), StandardSurfaceShaderNode->GetUniqueID(), PBRMR::Parameters::Roughness.ToString());
+		UInterchangeShaderPortsAPI::ConnectOuputToInputByName(ShaderGraphNode, PBRMR::Parameters::EmissiveColor.ToString(), StandardSurfaceShaderNode->GetUniqueID(), PBRMR::Parameters::EmissiveColor.ToString());
+		UInterchangeShaderPortsAPI::ConnectOuputToInputByName(ShaderGraphNode, PBRMR::Parameters::Anisotropy.ToString(), StandardSurfaceShaderNode->GetUniqueID(), PBRMR::Parameters::Anisotropy.ToString());
+		UInterchangeShaderPortsAPI::ConnectOuputToInputByName(ShaderGraphNode, PBRMR::Parameters::Normal.ToString(), StandardSurfaceShaderNode->GetUniqueID(), PBRMR::Parameters::Normal.ToString());
+		UInterchangeShaderPortsAPI::ConnectOuputToInputByName(ShaderGraphNode, PBRMR::Parameters::Tangent.ToString(), StandardSurfaceShaderNode->GetUniqueID(), PBRMR::Parameters::Tangent.ToString());
+
+		// We can't have all shading models at once, so we have to make a choice here
+		if(UInterchangeShaderPortsAPI::HasInput(StandardSurfaceShaderNode, StandardSurface::Parameters::Transmission))
+		{
+			UInterchangeShaderPortsAPI::ConnectOuputToInputByName(ShaderGraphNode, PBRMR::Parameters::Opacity.ToString(), StandardSurfaceShaderNode->GetUniqueID(), PBRMR::Parameters::Opacity.ToString());
+			UInterchangeShaderPortsAPI::ConnectOuputToInputByName(ShaderGraphNode, ThinTranslucent::Parameters::TransmissionColor.ToString(), StandardSurfaceShaderNode->GetUniqueID(), ThinTranslucent::Parameters::TransmissionColor.ToString());
+			UInterchangeShaderPortsAPI::ConnectOuputToInputByName(ShaderGraphNode, Common::Parameters::Refraction.ToString(), StandardSurfaceShaderNode->GetUniqueID(), Common::Parameters::Refraction.ToString());
+
+			// If we have have Transmission and Opacity let's use the Surface Coverage instead
+			if(UInterchangeShaderPortsAPI::HasInput(StandardSurfaceShaderNode, StandardSurface::Parameters::Opacity))
+			{
+				UInterchangeShaderPortsAPI::ConnectOuputToInputByName(ShaderGraphNode, ThinTranslucent::Parameters::SurfaceCoverage.ToString(), StandardSurfaceShaderNode->GetUniqueID(), ThinTranslucent::Parameters::SurfaceCoverage.ToString());
+			}
+		}
+		else if(UInterchangeShaderPortsAPI::HasInput(StandardSurfaceShaderNode, StandardSurface::Parameters::Sheen))
 		{
-			UInterchangeShaderPortsAPI::ConnectOuputToInputByName(ShaderGraphNode, SubstrateMaterial::Parameters::OpacityMask.ToString(), StandardSurfaceShaderNode->GetUniqueID(), StandardSurface::SubstrateMaterial::Outputs::Opacity.ToString());
-			ShaderGraphNode->SetCustomBlendMode(EBlendMode::BLEND_Masked);
+			UInterchangeShaderPortsAPI::ConnectOuputToInputByName(ShaderGraphNode, Sheen::Parameters::SheenColor.ToString(), StandardSurfaceShaderNode->GetUniqueID(), Sheen::Parameters::SheenColor.ToString());
+			UInterchangeShaderPortsAPI::ConnectOuputToInputByName(ShaderGraphNode, Sheen::Parameters::SheenRoughness.ToString(), StandardSurfaceShaderNode->GetUniqueID(), Sheen::Parameters::SheenRoughness.ToString());
+		}
+		else if(UInterchangeShaderPortsAPI::HasInput(StandardSurfaceShaderNode, StandardSurface::Parameters::Coat))
+		{
+			UInterchangeShaderPortsAPI::ConnectOuputToInputByName(ShaderGraphNode, ClearCoat::Parameters::ClearCoat.ToString(), StandardSurfaceShaderNode->GetUniqueID(), ClearCoat::Parameters::ClearCoat.ToString());
+			UInterchangeShaderPortsAPI::ConnectOuputToInputByName(ShaderGraphNode, ClearCoat::Parameters::ClearCoatRoughness.ToString(), StandardSurfaceShaderNode->GetUniqueID(), ClearCoat::Parameters::ClearCoatRoughness.ToString());
+			UInterchangeShaderPortsAPI::ConnectOuputToInputByName(ShaderGraphNode, ClearCoat::Parameters::ClearCoatNormal.ToString(), StandardSurfaceShaderNode->GetUniqueID(), ClearCoat::Parameters::ClearCoatNormal.ToString());
+		}
+		else if(UInterchangeShaderPortsAPI::HasInput(StandardSurfaceShaderNode, StandardSurface::Parameters::Subsurface))
+		{
+			UInterchangeShaderPortsAPI::ConnectOuputToInputByName(ShaderGraphNode, PBRMR::Parameters::Opacity.ToString(), StandardSurfaceShaderNode->GetUniqueID(), PBRMR::Parameters::Opacity.ToString());
+			UInterchangeShaderPortsAPI::ConnectOuputToInputByName(ShaderGraphNode, Subsurface::Parameters::SubsurfaceColor.ToString(), StandardSurfaceShaderNode->GetUniqueID(), Subsurface::Parameters::SubsurfaceColor.ToString());
+		}
+		else if(UInterchangeShaderPortsAPI::HasInput(StandardSurfaceShaderNode, StandardSurface::Parameters::Opacity))
+		{
+			UInterchangeShaderPortsAPI::ConnectOuputToInputByName(ShaderGraphNode, ThinTranslucent::Parameters::SurfaceCoverage.ToString(), StandardSurfaceShaderNode->GetUniqueID(), ThinTranslucent::Parameters::SurfaceCoverage.ToString());
 		}
 	}
 }
+
 #endif
diff --git a/Engine/Plugins/Interchange/Runtime/Source/Import/Private/MaterialX/MaterialXUtils/MaterialXStandardSurfaceShader.h b/Engine/Plugins/Interchange/Runtime/Source/Import/Private/MaterialX/MaterialXUtils/MaterialXStandardSurfaceShader.h
--- a/Engine/Plugins/Interchange/Runtime/Source/Import/Private/MaterialX/MaterialXUtils/MaterialXStandardSurfaceShader.h
+++ b/Engine/Plugins/Interchange/Runtime/Source/Import/Private/MaterialX/MaterialXUtils/MaterialXStandardSurfaceShader.h
@@ -19,11 +19,5 @@ public:
 	static TSharedRef<FMaterialXBase> MakeInstance(UInterchangeBaseNodeContainer& BaseNodeContainer);
 
 	virtual void Translate(MaterialX::NodePtr StandardSurfaceNode) override;
-
-private:
-
-	void ConnectToStandardSurface(UInterchangeFunctionCallShaderNode* StandardSurfaceShaderNode);
-
-	void ConnectToSubstrateStandardSurface(UInterchangeFunctionCallShaderNode* StandardSurfaceShaderNode);
 };
 #endif
\ No newline at end of file
diff --git a/Engine/Plugins/Interchange/Runtime/Source/Import/Private/MaterialX/MaterialXUtils/MaterialXSurfaceShader.cpp b/Engine/Plugins/Interchange/Runtime/Source/Import/Private/MaterialX/MaterialXUtils/MaterialXSurfaceShader.cpp
--- a/Engine/Plugins/Interchange/Runtime/Source/Import/Private/MaterialX/MaterialXUtils/MaterialXSurfaceShader.cpp
+++ b/Engine/Plugins/Interchange/Runtime/Source/Import/Private/MaterialX/MaterialXUtils/MaterialXSurfaceShader.cpp
@@ -24,19 +24,10 @@ void FMaterialXSurfaceShader::Translate(MaterialX::NodePtr SurfaceNode)
 {
 	this->SurfaceShaderNode = SurfaceNode;
 
-	UInterchangeFunctionCallShaderNode* FunctionSurfaceShaderNode = CreateFunctionCallShaderNode(SurfaceNode->getName().c_str(), UE::Interchange::MaterialX::IndexSurfaceShaders, uint8(EInterchangeMaterialXShaders::Surface));
+	UInterchangeShaderNode* FunctionSurfaceShaderNode = FMaterialXSurfaceShaderAbstract::Translate(EInterchangeMaterialXShaders::Surface);
 
 	using namespace UE::Interchange::Materials;
 
-	// BSDF
-	ConnectNodeOutputToInput(mx::Surface::Input::Bsdf, FunctionSurfaceShaderNode, Surface::Parameters::BSDF.ToString(), nullptr);
-
-	// EDF
-	ConnectNodeOutputToInput(mx::Surface::Input::Edf, FunctionSurfaceShaderNode, Surface::Parameters::EDF.ToString(), nullptr);
-
-	// Opacity
-	ConnectNodeOutputToInput(mx::Surface::Input::Opacity, FunctionSurfaceShaderNode, Surface::Parameters::Opacity.ToString(), mx::SurfaceUnlit::DefaultValue::Float::Opacity);
-
 	// Outputs
 	if(!bIsSubstrateEnabled)
 	{
diff --git a/Engine/Plugins/Interchange/Runtime/Source/Import/Private/MaterialX/MaterialXUtils/MaterialXSurfaceShaderAbstract.cpp b/Engine/Plugins/Interchange/Runtime/Source/Import/Private/MaterialX/MaterialXUtils/MaterialXSurfaceShaderAbstract.cpp
--- a/Engine/Plugins/Interchange/Runtime/Source/Import/Private/MaterialX/MaterialXUtils/MaterialXSurfaceShaderAbstract.cpp
+++ b/Engine/Plugins/Interchange/Runtime/Source/Import/Private/MaterialX/MaterialXUtils/MaterialXSurfaceShaderAbstract.cpp
@@ -4,10 +4,8 @@
 #include "MaterialX/MaterialXUtils/MaterialXSurfaceShaderAbstract.h"
 
 #include "InterchangeImportLog.h"
-#include "InterchangeManager.h"
 #include "InterchangeTexture2DNode.h"
 #include "InterchangeTextureBlurNode.h"
-#include "InterchangeTranslatorBase.h"
 #include "Materials/MaterialExpressionNoise.h"
 #include "Materials/MaterialExpressionTransform.h"
 #include "Materials/MaterialExpressionTransformPosition.h"
@@ -1491,5 +1489,120 @@ void FMaterialXSurfaceShaderAbstract::SetAttributeNewName(MaterialX::InputPtr In
 	Input->setAttribute(mx::Attributes::NewName, NewName);
 }
 
+namespace
+{
+	template<typename T>
+	constexpr T DefaultMaxValue = std::numeric_limits<T>::max();
+
+	template<typename T>
+	std::string ValueToString = TCHAR_TO_ANSI(*FString::SanitizeFloat(DefaultMaxValue<T>));
+
+	template<typename T>
+	T GetValue(const auto& Value)
+	{
+		using Type = std::remove_cvref_t<decltype(Value)>;
+
+		if constexpr(std::is_same_v<mx::Color3, Type>)
+		{
+			return FLinearColor{ Value[0], Value[1], Value[2] };
+		}
+		else if constexpr(std::is_same_v<mx::Color4, Type>)
+		{
+			return FLinearColor{ Value[0], Value[3], Value[2], Value[3] };
+		}
+		else if constexpr(std::is_same_v<mx::Vector2, Type>)
+		{
+			return FVector2f{ Value[0], Value[1] };
+		}
+		else if constexpr(std::is_same_v<mx::Vector3, Type>)
+		{
+			return FVector3f{ Value[0], Value[1], Value[2] };
+		}
+		else if constexpr(std::is_same_v<mx::Vector4, Type>)
+		{
+			return FVector4f{ Value[0], Value[1], Value[2], Value[3] };
+		}
+	}
+}
+UInterchangeShaderNode* FMaterialXSurfaceShaderAbstract::Translate(EInterchangeMaterialXShaders ShaderType)
+{
+	constexpr bool bInputInTangentSpace = true;
+
+	UInterchangeFunctionCallShaderNode* FunctionCallShaderNode = CreateFunctionCallShaderNode(SurfaceShaderNode->getName().c_str(), UE::Interchange::MaterialX::IndexSurfaceShaders, uint8(ShaderType));
+
+	for(mx::InputPtr Input : SurfaceShaderNode->getInputs())
+	{
+		mx::ValuePtr DefaultValue = Input->getDefaultValue();
+
+		//Let's create default values in case there is none
+		if(!DefaultValue)
+		{
+			if(Input->getType() == mx::Type::Float)
+			{
+				DefaultValue = mx::Value::createValueFromStrings(ValueToString<float>, mx::Type::Float);
+			}
+			else if(Input->getType() == mx::Type::Color3)
+			{
+				DefaultValue = mx::Value::createValueFromStrings(ValueToString<float> + "," + ValueToString<float> + "," + ValueToString<float>, mx::Type::Color3);
+			}
+			else if(Input->getType() == mx::Type::Color4)
+			{
+				DefaultValue = mx::Value::createValueFromStrings(ValueToString<float> + "," + ValueToString<float> + "," + ValueToString<float> + "," + ValueToString<float>, mx::Type::Color4);
+			}
+			else if(Input->getType() == mx::Type::Boolean)
+			{
+				DefaultValue = mx::Value::createValueFromStrings("false", mx::Type::Boolean);
+			}
+			else if(Input->getType() == mx::Type::Integer)
+			{
+				DefaultValue = mx::Value::createValueFromStrings(ValueToString<int32>, mx::Type::Integer);
+			}
+			else if(Input->getType() == mx::Type::Vector2)
+			{
+				DefaultValue = mx::Value::createValueFromStrings(ValueToString<float> + "," + ValueToString<float>, mx::Type::Vector2);
+			}
+			else if(Input->getType() == mx::Type::Vector3)
+			{
+				DefaultValue = mx::Value::createValueFromStrings(ValueToString<float> + "," + ValueToString<float> + "," + ValueToString<float>, mx::Type::Vector3);
+			}
+			else if(Input->getType() == mx::Type::Vector4)
+			{
+				DefaultValue = mx::Value::createValueFromStrings(ValueToString<float> + "," + ValueToString<float> + "," + ValueToString<float> + "," + ValueToString<float>, mx::Type::Vector4);
+			}
+		}
+
+		if(DefaultValue->getTypeString() == mx::Type::Float)
+		{
+			ConnectNodeOutputToInput(Input->getName().c_str(), FunctionCallShaderNode, Input->getName().c_str(), mx::fromValueString<float>(DefaultValue->getValueString()));
+		}
+		else if(DefaultValue->getTypeString() == mx::Type::Color3)
+		{
+			ConnectNodeOutputToInput(Input->getName().c_str(), FunctionCallShaderNode, Input->getName().c_str(), GetValue<FLinearColor>(mx::fromValueString<mx::Color3>(DefaultValue->getValueString())));
+		}
+		else if(DefaultValue->getTypeString() == mx::Type::Color4)
+		{
+			ConnectNodeOutputToInput(Input->getName().c_str(), FunctionCallShaderNode, Input->getName().c_str(), GetValue<FLinearColor>(mx::fromValueString<mx::Color4>(DefaultValue->getValueString())));
+		}
+		else if(DefaultValue->getTypeString() == mx::Type::Boolean)
+		{
+			ConnectNodeOutputToInput(Input->getName().c_str(), FunctionCallShaderNode, Input->getName().c_str(), mx::fromValueString<bool>(DefaultValue->getValueString()));
+		}
+		else if(DefaultValue->getTypeString() == mx::Type::Vector2)
+		{
+			ConnectNodeOutputToInput(Input->getName().c_str(), FunctionCallShaderNode, Input->getName().c_str(), FVector4f{ GetValue<FVector2f>(mx::fromValueString<mx::Vector2>(DefaultValue->getValueString())), FVector2f{0,0} }, bInputInTangentSpace);
+		}
+		else if(DefaultValue->getTypeString() == mx::Type::Vector3)
+		{
+			ConnectNodeOutputToInput(Input->getName().c_str(), FunctionCallShaderNode, Input->getName().c_str(), GetValue<FVector3f>(mx::fromValueString<mx::Vector3>(DefaultValue->getValueString())), bInputInTangentSpace);
+		}
+		else if(DefaultValue->getTypeString() == mx::Type::Vector4)
+		{
+			ConnectNodeOutputToInput(Input->getName().c_str(), FunctionCallShaderNode, Input->getName().c_str(), GetValue<FVector4f>(mx::fromValueString<mx::Vector4>(DefaultValue->getValueString())), bInputInTangentSpace);
+		}
+	}
+
+	return FunctionCallShaderNode;
+}
+
 #undef LOCTEXT_NAMESPACE
 #endif
\ No newline at end of file
diff --git a/Engine/Plugins/Interchange/Runtime/Source/Import/Private/MaterialX/MaterialXUtils/MaterialXSurfaceUnlitShader.cpp b/Engine/Plugins/Interchange/Runtime/Source/Import/Private/MaterialX/MaterialXUtils/MaterialXSurfaceUnlitShader.cpp
--- a/Engine/Plugins/Interchange/Runtime/Source/Import/Private/MaterialX/MaterialXUtils/MaterialXSurfaceUnlitShader.cpp
+++ b/Engine/Plugins/Interchange/Runtime/Source/Import/Private/MaterialX/MaterialXUtils/MaterialXSurfaceUnlitShader.cpp
@@ -22,26 +22,10 @@ TSharedRef<FMaterialXBase> FMaterialXSurfaceUnlitShader::MakeInstance(UInterchan
 
 void FMaterialXSurfaceUnlitShader::Translate(mx::NodePtr SurfaceUnlitNode)
 {
-	this->SurfaceShaderNode = SurfaceUnlitNode;
-
-	UInterchangeFunctionCallShaderNode* SurfaceUnlitShaderNode = CreateFunctionCallShaderNode(SurfaceUnlitNode->getName().c_str(), UE::Interchange::MaterialX::IndexSurfaceShaders, uint8(EInterchangeMaterialXShaders::SurfaceUnlit));
-
 	using namespace UE::Interchange::Materials;
+	this->SurfaceShaderNode = SurfaceUnlitNode;
 
-	//Emission
-	ConnectNodeOutputToInput(mx::SurfaceUnlit::Input::Emission, SurfaceUnlitShaderNode, SurfaceUnlit::Parameters::Emission.ToString(), mx::SurfaceUnlit::DefaultValue::Float::Emission);
-
-	//Emission Color
-	ConnectNodeOutputToInput(mx::SurfaceUnlit::Input::EmissionColor, SurfaceUnlitShaderNode, SurfaceUnlit::Parameters::EmissionColor.ToString(), mx::SurfaceUnlit::DefaultValue::Color3::EmissionColor);
-
-	//Opacity
-	ConnectNodeOutputToInput(mx::SurfaceUnlit::Input::Opacity, SurfaceUnlitShaderNode, SurfaceUnlit::Parameters::Opacity.ToString(), mx::SurfaceUnlit::DefaultValue::Float::Opacity);
-
-	//Transmission
-	ConnectNodeOutputToInput(mx::SurfaceUnlit::Input::Transmission, SurfaceUnlitShaderNode, SurfaceUnlit::Parameters::Transmission.ToString(), mx::SurfaceUnlit::DefaultValue::Float::Transmission);
-
-	//Transmission Color
-	ConnectNodeOutputToInput(mx::SurfaceUnlit::Input::TransmissionColor, SurfaceUnlitShaderNode, SurfaceUnlit::Parameters::TransmissionColor.ToString(), mx::SurfaceUnlit::DefaultValue::Color3::TransmissionColor);
+	UInterchangeShaderNode* SurfaceUnlitShaderNode = FMaterialXSurfaceShaderAbstract::Translate(EInterchangeMaterialXShaders::SurfaceUnlit);
 
 	if(!bIsSubstrateEnabled)
 	{
diff --git a/Engine/Plugins/Interchange/Runtime/Source/Import/Private/MaterialX/MaterialXUtils/MaterialXUsdPreviewSurfaceShader.cpp b/Engine/Plugins/Interchange/Runtime/Source/Import/Private/MaterialX/MaterialXUtils/MaterialXUsdPreviewSurfaceShader.cpp
--- a/Engine/Plugins/Interchange/Runtime/Source/Import/Private/MaterialX/MaterialXUtils/MaterialXUsdPreviewSurfaceShader.cpp
+++ b/Engine/Plugins/Interchange/Runtime/Source/Import/Private/MaterialX/MaterialXUtils/MaterialXUsdPreviewSurfaceShader.cpp
@@ -25,50 +25,10 @@ void FMaterialXUsdPreviewSurfaceShader::Translate(MaterialX::NodePtr UsdPreviewS
 
 	using namespace UE::Interchange::Materials;
 
-	UInterchangeFunctionCallShaderNode* UsdPreviewSurfaceShaderNode = CreateFunctionCallShaderNode(UsdPreviewSurfaceNode->getName().c_str(), UE::Interchange::MaterialX::IndexSurfaceShaders, uint8(EInterchangeMaterialXShaders::UsdPreviewSurface));
+	UInterchangeShaderNode* UsdPreviewSurfaceShaderNode = FMaterialXSurfaceShaderAbstract::Translate(EInterchangeMaterialXShaders::UsdPreviewSurface);
 
-	// Inputs
-	//Diffuse Color
-	ConnectNodeOutputToInput(mx::UsdPreviewSurface::Input::DiffuseColor, UsdPreviewSurfaceShaderNode, UsdPreviewSurface::Parameters::DiffuseColor.ToString(), mx::UsdPreviewSurface::DefaultValue::Color3::DiffuseColor);
-
-	//Emissive Color
-	ConnectNodeOutputToInput(mx::UsdPreviewSurface::Input::EmissiveColor, UsdPreviewSurfaceShaderNode, UsdPreviewSurface::Parameters::EmissiveColor.ToString(), mx::UsdPreviewSurface::DefaultValue::Color3::EmissiveColor);
-
-	//Specular Color
-	ConnectNodeOutputToInput(mx::UsdPreviewSurface::Input::SpecularColor, UsdPreviewSurfaceShaderNode, UsdPreviewSurface::Parameters::SpecularColor.ToString(), mx::UsdPreviewSurface::DefaultValue::Color3::SpecularColor);
-
-	//Metallic
-	ConnectNodeOutputToInput(mx::UsdPreviewSurface::Input::Metallic, UsdPreviewSurfaceShaderNode, UsdPreviewSurface::Parameters::Metallic.ToString(), mx::UsdPreviewSurface::DefaultValue::Float::Metallic);
-
-	//Roughness
-	ConnectNodeOutputToInput(mx::UsdPreviewSurface::Input::Roughness, UsdPreviewSurfaceShaderNode, UsdPreviewSurface::Parameters::Roughness.ToString(), mx::UsdPreviewSurface::DefaultValue::Float::Roughness);
-
-	//Clearcoat
-	ConnectNodeOutputToInput(mx::UsdPreviewSurface::Input::Clearcoat, UsdPreviewSurfaceShaderNode, UsdPreviewSurface::Parameters::Clearcoat.ToString(), mx::UsdPreviewSurface::DefaultValue::Float::Clearcoat);
-
-	//Clearcoat Roughness
-	ConnectNodeOutputToInput(mx::UsdPreviewSurface::Input::ClearcoatRoughness, UsdPreviewSurfaceShaderNode, UsdPreviewSurface::Parameters::ClearcoatRoughness.ToString(), mx::UsdPreviewSurface::DefaultValue::Float::ClearcoatRoughness);
-
-	//Opacity
-	ConnectNodeOutputToInput(mx::UsdPreviewSurface::Input::Opacity, UsdPreviewSurfaceShaderNode, UsdPreviewSurface::Parameters::Opacity.ToString(), mx::UsdPreviewSurface::DefaultValue::Float::Opacity);
-
-	//Opacity Threshold
-	ConnectNodeOutputToInput(mx::UsdPreviewSurface::Input::OpacityThreshold, UsdPreviewSurfaceShaderNode, UsdPreviewSurface::Parameters::OpacityThreshold.ToString(), mx::UsdPreviewSurface::DefaultValue::Float::OpacityThreshold);
-
-	//IOR
-	ConnectNodeOutputToInput(mx::UsdPreviewSurface::Input::IOR, UsdPreviewSurfaceShaderNode, UsdPreviewSurface::Parameters::IOR.ToString(), mx::UsdPreviewSurface::DefaultValue::Float::IOR);
-
-	//Normal
-	ConnectNodeOutputToInput(mx::UsdPreviewSurface::Input::Normal, UsdPreviewSurfaceShaderNode, UsdPreviewSurface::Parameters::Normal.ToString(), mx::UsdPreviewSurface::DefaultValue::Vector3::Normal);
-
-	//Displacement
-	ConnectNodeOutputToInput(mx::UsdPreviewSurface::Input::Displacement, ShaderGraphNode, UsdPreviewSurface::Parameters::Displacement.ToString(), mx::UsdPreviewSurface::DefaultValue::Float::Displacement);
-
-	//Occlusion
-	ConnectNodeOutputToInput(mx::UsdPreviewSurface::Input::Occlusion, UsdPreviewSurfaceShaderNode, UsdPreviewSurface::Parameters::Occlusion.ToString(), mx::UsdPreviewSurface::DefaultValue::Float::Occlusion);
-
-	if(!bIsSubstrateEnabled)
 	// Outputs
+	if(!bIsSubstrateEnabled)
 	{
 		UInterchangeShaderPortsAPI::ConnectOuputToInputByName(ShaderGraphNode, PBRMR::Parameters::BaseColor.ToString(), UsdPreviewSurfaceShaderNode->GetUniqueID(), PBRMR::Parameters::BaseColor.ToString());
 		UInterchangeShaderPortsAPI::ConnectOuputToInputByName(ShaderGraphNode, PBRMR::Parameters::Metallic.ToString(), UsdPreviewSurfaceShaderNode->GetUniqueID(), PBRMR::Parameters::Metallic.ToString());
diff --git a/Engine/Plugins/Interchange/Runtime/Source/Import/Public/MaterialX/MaterialXUtils/MaterialXSurfaceShaderAbstract.h b/Engine/Plugins/Interchange/Runtime/Source/Import/Public/MaterialX/MaterialXUtils/MaterialXSurfaceShaderAbstract.h
--- a/Engine/Plugins/Interchange/Runtime/Source/Import/Public/MaterialX/MaterialXUtils/MaterialXSurfaceShaderAbstract.h
+++ b/Engine/Plugins/Interchange/Runtime/Source/Import/Public/MaterialX/MaterialXUtils/MaterialXSurfaceShaderAbstract.h
@@ -118,11 +118,11 @@ protected:
 	 * @param InputShaderName - The name of the input of the shader node to connect to
 	 * @param DefaultValue - The default value of the MaterialX input
 	 * @param bIsTangentSpaceInput - Set the tangent space along the path of an input
+	 * @param bUseDefaultValue - Usually we want to avoid creating default values on direct inputs of shaders (StandardSurface, OpenPBR), but some shaders (UsdPreviewSurface) we need to decide that at run-time
 	 */
 	template<typename T>
 	bool ConnectNodeOutputToInput(const char* InputName, UInterchangeShaderNode* ShaderNode, const FString& InputShaderName, T DefaultValue, bool bIsTangentSpaceInput = false)
 	{
-		MaterialX::DocumentPtr Document = SurfaceShaderNode->getDocument();
 		MaterialX::InputPtr Input = GetInput(SurfaceShaderNode, InputName);
 
 		TGuardValue<bool>InputTypeBeingProcessedGuard(bTangentSpaceInput, bIsTangentSpaceInput);
@@ -144,7 +144,7 @@ protected:
 					bIsConnected = AddLinearColorAttribute(Input, InputShaderName, ShaderNode, DefaultValue);
 				}
 				else if constexpr(std::is_same_v<decltype(DefaultValue), FVector4f> || std::is_same_v<decltype(DefaultValue), FVector3f> || std::is_same_v<decltype(DefaultValue), FVector2f>)
-				{					
+				{
 					bIsConnected = AddVectorAttribute(Input, InputShaderName, ShaderNode, DefaultValue);
 				}
 				if constexpr(std::is_same_v<decltype(DefaultValue), bool>)
@@ -466,6 +466,16 @@ protected:
 	 */
 	void SetAttributeNewName(MaterialX::InputPtr Input, const char* NewName) const;
 
+	/**
+	 * This function should be called first by the Translate method of derived class, SurfaceShaderNode should initialized first by the derived class
+	 * 
+	 * @param ShaderType - the type of ShaderGraphNode to create
+	 * @return The shader node created, usually a function call shader node
+	 */
+	UInterchangeShaderNode* Translate(EInterchangeMaterialXShaders ShaderType);
+
+	virtual void Translate(MaterialX::NodePtr ShaderNode) override = 0;
+
 private:
 	
 	UInterchangeShaderNode* ConnectGeometryInputToOutput(const FConnectNode& Connect, const FString& ShaderType, const FString& TransformShaderType, const FString& TransformInput, const FString& TransformSourceType, int32 TransformSource, const FString& TransformType, int32 TransformSDestination, bool bIsVector = true);
