[Interchange glTF] Updated expected color space for SpecularFactor.

#rnx

#virtualized

diff --git a/Engine/Plugins/Interchange/Runtime/Source/Import/Private/Gltf/InterchangeGltfTranslator.cpp b/Engine/Plugins/Interchange/Runtime/Source/Import/Private/Gltf/InterchangeGltfTranslator.cpp
--- a/Engine/Plugins/Interchange/Runtime/Source/Import/Private/Gltf/InterchangeGltfTranslator.cpp
+++ b/Engine/Plugins/Interchange/Runtime/Source/Import/Private/Gltf/InterchangeGltfTranslator.cpp
@@ -600,9 +600,12 @@ bool UInterchangeGLTFTranslator::Translate( UInterchangeBaseNodeContainer& NodeC
 		int32 MaterialIndex = 0;
 		for ( const GLTF::FMaterial& GltfMaterial : GltfAsset.Materials )
 		{
-			//Based on the gltf specification the basecolor and emissive textures have SRGB colors:
+			//Based on the gltf specification the basecolor and emissive and specular textures have SRGB colors:
 			SetTextureSRGB(NodeContainer, GltfMaterial.BaseColor, true);
 			SetTextureSRGB(NodeContainer, GltfMaterial.Emissive, true);
+			SetTextureSRGB(NodeContainer, GltfMaterial.Specular.SpecularColorMap, true);
+			SetTextureSRGB(NodeContainer, GltfMaterial.Specular.SpecularMap, true);		//Technically SpecularTexture is only using the Alpha channel, but could be packed onto other textures
+																						//Parent material's expectation is that the SpecularMap is SRGB (as the expectation is that it is packed onto the SpecularColorTexture.)
 			//Textures that are expected to use Scalar outputs we want to set them as SRGB false explicitly, based on UInterchangeGenericMaterialPipeline::HandleTextureNode
 			SetTextureSRGB(NodeContainer, GltfMaterial.MetallicRoughness.Map, false);
 			SetTextureSRGB(NodeContainer, GltfMaterial.Occlusion, false);
@@ -974,6 +977,14 @@ void UInterchangeGLTFTranslator::SetTextureSRGB(UInterchangeBaseNodeContainer& N
 		const FString TextureUid = UInterchangeTextureNode::MakeNodeUid(GltfAsset.Textures[TextureMap.TextureIndex].UniqueId);
 		if (UInterchangeTextureNode* TextureNode = const_cast<UInterchangeTextureNode*>(Cast<UInterchangeTextureNode>(NodeContainer.GetNode(TextureUid))))
 		{
+			bool bExistingSRGBSetting;
+			if (TextureNode->GetCustomSRGB(bExistingSRGBSetting))
+			{
+				if (bExistingSRGBSetting != bSRGB)
+				{
+					UE_LOG(LogInterchangeImport, Warning, TEXT("UInterchangeGLTFPipeline: Unexpected SRGB/Linear setting on Texture: %s."), *TextureNode->GetDisplayLabel());
+				}
+			}
 			TextureNode->SetCustomSRGB(bSRGB);
 		}
 	}
