Initial USD Interchange translator (with scene nodes only)

diff --git a/Engine/Plugins/Interchange/Runtime/Interchange.uplugin b/Engine/Plugins/Interchange/Runtime/Interchange.uplugin
--- a/Engine/Plugins/Interchange/Runtime/Interchange.uplugin
+++ b/Engine/Plugins/Interchange/Runtime/Interchange.uplugin
@@ -87,6 +87,10 @@
 		}
 	],
 	"Plugins": [
+		{
+			"Name": "USDCore",
+			"Enabled": true
+		},
 		{
 			"Name": "VariantManagerContent",
 			"Enabled": true
diff --git a/Engine/Plugins/Interchange/Runtime/Source/Import/InterchangeImport.Build.cs b/Engine/Plugins/Interchange/Runtime/Source/Import/InterchangeImport.Build.cs
--- a/Engine/Plugins/Interchange/Runtime/Source/Import/InterchangeImport.Build.cs
+++ b/Engine/Plugins/Interchange/Runtime/Source/Import/InterchangeImport.Build.cs
@@ -43,6 +43,8 @@ namespace UnrealBuildTool.Rules
 					"RenderCore",
 					"RHI",
 					"TextureUtilitiesCommon",
+					"UnrealUSDWrapper",
+					"USDUtilities",
 					"VariantManagerContent",
 				}
 			);
diff --git a/Engine/Plugins/Interchange/Runtime/Source/Import/Private/InterchangeImportModule.cpp b/Engine/Plugins/Interchange/Runtime/Source/Import/Private/InterchangeImportModule.cpp
--- a/Engine/Plugins/Interchange/Runtime/Source/Import/Private/InterchangeImportModule.cpp
+++ b/Engine/Plugins/Interchange/Runtime/Source/Import/Private/InterchangeImportModule.cpp
@@ -35,6 +35,7 @@
 #include "Texture/InterchangeUEJPEGTranslator.h"
 #include "Texture/InterchangePSDTranslator.h"
 #include "Texture/InterchangeTextureFactory.h"
+#include "Usd/InterchangeUsdTranslator.h"
 
 DEFINE_LOG_CATEGORY(LogInterchangeImport);
 
@@ -80,6 +81,7 @@ void FInterchangeImportModule::StartupModule()
 		UInterchangeFbxTranslator::CleanUpTemporaryFolder();
 		InterchangeManager.RegisterTranslator(UInterchangeFbxTranslator::StaticClass());
 #endif
+		InterchangeManager.RegisterTranslator(UInterchangeUSDTranslator::StaticClass());
 		InterchangeManager.RegisterTranslator(UInterchangeGLTFTranslator::StaticClass());
 		InterchangeManager.RegisterTranslator(UInterchangeOBJTranslator::StaticClass());
 
diff --git a/Engine/Plugins/Interchange/Runtime/Source/Import/Private/Usd/InterchangeUsdTranslator.cpp b/Engine/Plugins/Interchange/Runtime/Source/Import/Private/Usd/InterchangeUsdTranslator.cpp
new file mode 100644
--- /dev/null
+++ b/Engine/Plugins/Interchange/Runtime/Source/Import/Private/Usd/InterchangeUsdTranslator.cpp
@@ -0,0 +1,235 @@
+// Copyright Epic Games, Inc. All Rights Reserved. 
+
+#include "Usd/InterchangeUsdTranslator.h"
+
+#include "UnrealUSDWrapper.h"
+#include "UsdWrappers/SdfPath.h"
+#include "UsdWrappers/UsdGeomXformable.h"
+#include "UsdWrappers/UsdAttribute.h"
+#include "UsdWrappers/UsdPrim.h"
+#include "UsdWrappers/UsdStage.h"
+#include "UsdWrappers/UsdTyped.h"
+#include "UsdWrappers/VtValue.h"
+
+#include "USDConversionUtils.h"
+#include "USDPrimConversion.h"
+#include "USDValueConversion.h"
+
+#include "HAL/IConsoleManager.h"
+#include "InterchangeCameraNode.h"
+#include "InterchangeImportLog.h"
+#include "InterchangeLightNode.h"
+#include "InterchangeMeshNode.h"
+#include "InterchangeSceneNode.h"
+
+#include UE_INLINE_GENERATED_CPP_BY_NAME(InterchangeUsdTranslator)
+
+#define LOCTEXT_NAMESPACE "InterchangeUSDTranslator"
+
+static bool GInterchangeEnableUSDImport = false;
+static FAutoConsoleVariableRef CVarInterchangeEnableUSDImport(
+	TEXT("Interchange.FeatureFlags.Import.USD"),
+	GInterchangeEnableUSDImport,
+	TEXT("Whether USD support is enabled.")
+);
+
+UInterchangeUSDTranslator::UInterchangeUSDTranslator()
+{
+}
+
+EInterchangeTranslatorType UInterchangeUSDTranslator::GetTranslatorType() const
+{
+	return EInterchangeTranslatorType::Scenes;
+}
+
+EInterchangeTranslatorAssetType UInterchangeUSDTranslator::GetSupportedAssetTypes() const
+{
+	return EInterchangeTranslatorAssetType::Materials | EInterchangeTranslatorAssetType::Meshes | EInterchangeTranslatorAssetType::Animations;
+}
+
+TArray<FString> UInterchangeUSDTranslator::GetSupportedFormats() const
+{
+	TArray<FString> Extensions;
+	if (GInterchangeEnableUSDImport)
+	{
+		UnrealUSDWrapper::AddUsdImportFileFormatDescriptions(Extensions);
+	}
+	return Extensions;
+}
+
+namespace UE::InterchangeUsdTranslator::Private
+{
+	void AddLightNode(const UE::FUsdPrim& Prim, UInterchangeBaseNodeContainer& NodeContainer)
+	{
+		FString NodeUid = TEXT("Light:") + Prim.GetPrimPath().GetString();
+		FString NodeName(Prim.GetName().ToString());
+
+		auto Attributes = Prim.GetAttributes();
+		for (auto& Attribute : Attributes)
+		{
+			UE_LOG(LogInterchangeImport, Log, TEXT("%s %s %s"), *Attribute.GetName().ToString(), *Attribute.GetBaseName().ToString(), *Attribute.GetTypeName().ToString());
+		}
+
+		// #ueent_todo: Retrieve light attributes and set them below
+
+		if (Prim.IsA(TEXT("DistantLight")))
+		{
+			UInterchangeDirectionalLightNode* LightNode = NewObject<UInterchangeDirectionalLightNode>(&NodeContainer);
+			LightNode->InitializeNode(NodeUid, NodeName, EInterchangeNodeContainerType::TranslatedAsset);
+
+			//LightNode->SetCustomLightColor(Color);
+			//LightNode->SetCustomIntensity(Intensity);
+
+			NodeContainer.AddNode(LightNode);
+		}
+		else if (Prim.IsA(TEXT("SphereLight")))
+		{
+			if (Prim.HasAPI(TEXT("ShapingAPI")))
+			{
+				UInterchangeSpotLightNode* LightNode = NewObject<UInterchangeSpotLightNode>(&NodeContainer);
+				LightNode->InitializeNode(NodeUid, NodeName, EInterchangeNodeContainerType::TranslatedAsset);
+
+				LightNode->SetCustomIntensityUnits(EInterchangeLightUnits::Lumens);
+				//LightNode->SetCustomLightColor(Color));
+				//LightNode->SetCustomIntensity(Intensity);
+
+				//LightNode->SetCustomInnerConeAngle(FMath::RadiansToDegrees(InnerConeAngle));
+				//LightNode->SetCustomOuterConeAngle(FMath::RadiansToDegrees(OuterConeAngle));
+
+				NodeContainer.AddNode(LightNode);
+			}
+			else
+			{
+				UInterchangePointLightNode* LightNode = NewObject<UInterchangePointLightNode>(&NodeContainer);
+				LightNode->InitializeNode(NodeUid, NodeName, EInterchangeNodeContainerType::TranslatedAsset);
+
+				LightNode->SetCustomIntensityUnits(EInterchangeLightUnits::Lumens);
+				//LightNode->SetCustomLightColor(FLinearColor(Color));
+				//LightNode->SetCustomIntensity(Intensity);
+
+				//LightNode->SetCustomAttenuationRadius(AttenuationRadius);
+
+				NodeContainer.AddNode(LightNode);
+			}
+		}
+		else if (Prim.IsA(TEXT("RectLight")) || Prim.IsA(TEXT("DiskLight")))
+		{
+			UInterchangeRectLightNode* LightNode = NewObject<UInterchangeRectLightNode>(&NodeContainer);
+			LightNode->InitializeNode(NodeUid, NodeName, EInterchangeNodeContainerType::TranslatedAsset);
+			NodeContainer.AddNode(LightNode);
+		}
+		// #ueent_todo:
+		// DomeLight -> SkyLight
+	}
+
+	void AddCameraNode(const UE::FUsdPrim& Prim, UInterchangeBaseNodeContainer& NodeContainer)
+	{
+		//FUsdCamera Camera;
+		//if (!UsdToUnreal::ConvertGeomCamera(Prim, Camera, UsdUtils::GetEarliestTimeCode()))
+		//{
+		//	return;
+		//}
+
+		FString NodeUid = TEXT("Camera:") + Prim.GetPrimPath().GetString();
+		FString NodeName(Prim.GetName().ToString());
+
+		UInterchangePhysicalCameraNode* CameraNode = NewObject<UInterchangePhysicalCameraNode>(&NodeContainer);
+		CameraNode->InitializeNode(NodeUid, NodeName, EInterchangeNodeContainerType::TranslatedAsset);
+		NodeContainer.AddNode(CameraNode);
+
+		//CameraNode->SetCustomFocalLength(Camera.FocalLength);
+		//CameraNode->SetCustomSensorHeight(Camera.SensorHeight);
+		//CameraNode->SetCustomSensorWidth(Camera.SensorWidth);
+	}
+
+	void AddSceneNode(const UE::FUsdStage& UsdStage, const UE::FUsdPrim& Prim, UInterchangeBaseNodeContainer& NodeContainer, const UInterchangeBaseNode* ParentNode)
+	{
+#if USE_USD_SDK
+		FTransform Transform = FTransform::Identity;
+		if (UsdToUnreal::ConvertXformable(UsdStage, UE::FUsdTyped(Prim), Transform, UsdUtils::GetEarliestTimeCode(), nullptr))
+		{
+			FString NodeUid = Prim.GetPrimPath().GetString();
+			FString DisplayLabel(Prim.GetName().ToString());
+
+			UInterchangeSceneNode* SceneNode = NewObject<UInterchangeSceneNode>(&NodeContainer);
+			SceneNode->InitializeNode(NodeUid, DisplayLabel, EInterchangeNodeContainerType::TranslatedScene);
+			SceneNode->SetAssetName(DisplayLabel);
+			SceneNode->SetCustomLocalTransform(&NodeContainer, Transform);
+			NodeContainer.AddNode(SceneNode);
+
+			FString Prefix;
+			if (Prim.IsA(TEXT("Camera")))
+			{
+				Prefix = TEXT("Camera:");
+				AddCameraNode(Prim, NodeContainer);
+			}
+			else if (Prim.HasAPI(TEXT("LightAPI")))
+			{
+				Prefix = TEXT("Light:");
+				AddLightNode(Prim, NodeContainer);
+			}
+
+			if (!Prefix.IsEmpty())
+			{
+				const FString AssetNodeUid = Prefix + NodeUid;
+				SceneNode->SetCustomAssetInstanceUid(AssetNodeUid);
+			}
+
+			if (ParentNode)
+			{
+				NodeContainer.SetNodeParentUid(SceneNode->GetUniqueID(), ParentNode->GetUniqueID());
+			}
+
+			ParentNode = SceneNode;
+		}
+
+		for (const FUsdPrim& ChildPrim : Prim.GetChildren())
+		{
+			AddSceneNode(UsdStage, ChildPrim, NodeContainer, ParentNode);
+		}
+#endif // USE_USD_SDK
+	}
+}
+
+bool UInterchangeUSDTranslator::Translate(UInterchangeBaseNodeContainer& NodeContainer) const
+{
+	using namespace UE;
+
+	FString FilePath = GetSourceData()->GetFilename();
+	if (!FPaths::FileExists(FilePath))
+	{
+		return false;
+	}
+
+	FUsdStage UsdStage = UnrealUSDWrapper::OpenStage(*FilePath, EUsdInitialLoadSet::LoadAll);
+	if (!UsdStage)
+	{
+		return false;
+	}
+
+	// Scene hierarchy
+	FUsdPrim RootPrim = UsdStage.GetPseudoRoot();
+	for (const FUsdPrim& Prim : RootPrim.GetChildren())
+	{
+		InterchangeUsdTranslator::Private::AddSceneNode(UsdStage, Prim, NodeContainer, nullptr);
+	}
+
+	return true;
+}
+
+TFuture<TOptional<UE::Interchange::FMeshPayloadData>> UInterchangeUSDTranslator::GetMeshPayloadData(const FInterchangeMeshPayLoadKey& PayLoadKey, const FTransform& MeshGlobalTransform) const
+{
+	return {};
+}
+
+TOptional<UE::Interchange::FImportImage> UInterchangeUSDTranslator::GetTexturePayloadData(const FString& PayloadKey, TOptional<FString>& AlternateTexturePath) const
+{
+	return {};
+}
+
+TArray<UE::Interchange::FAnimationPayloadData> UInterchangeUSDTranslator::GetAnimationPayloadData(const TArray<UE::Interchange::FAnimationPayloadQuery>& PayloadQuery) const
+{
+	return {};
+}
+
+#undef LOCTEXT_NAMESPACE
diff --git a/Engine/Plugins/Interchange/Runtime/Source/Import/Public/Usd/InterchangeUsdTranslator.h b/Engine/Plugins/Interchange/Runtime/Source/Import/Public/Usd/InterchangeUsdTranslator.h
new file mode 100644
--- /dev/null
+++ b/Engine/Plugins/Interchange/Runtime/Source/Import/Public/Usd/InterchangeUsdTranslator.h
@@ -0,0 +1,36 @@
+// Copyright Epic Games, Inc. All Rights Reserved.
+
+#pragma once
+
+#include "InterchangeTranslatorBase.h"
+#include "Animation/InterchangeAnimationPayloadInterface.h"
+#include "Mesh/InterchangeMeshPayloadInterface.h"
+#include "Texture/InterchangeTexturePayloadInterface.h"
+#include "UObject/Object.h"
+#include "UObject/ObjectMacros.h"
+
+#include "InterchangeUsdTranslator.generated.h"
+
+/* For now, USD Interchange (FBX parity) translator supports textures/materials, static/skeletal meshes, and animation*/
+UCLASS(BlueprintType)
+class UInterchangeUSDTranslator : public UInterchangeTranslatorBase
+	, public IInterchangeMeshPayloadInterface
+	, public IInterchangeTexturePayloadInterface
+	, public IInterchangeAnimationPayloadInterface
+{
+	GENERATED_BODY()
+
+public:
+	UInterchangeUSDTranslator();
+
+	/** Begin UInterchangeTranslatorBase API*/
+	virtual EInterchangeTranslatorType GetTranslatorType() const override;
+	virtual EInterchangeTranslatorAssetType GetSupportedAssetTypes() const override;
+	virtual TArray<FString> GetSupportedFormats() const override;
+	virtual bool Translate( UInterchangeBaseNodeContainer& BaseNodeContainer ) const override;
+	/** End UInterchangeTranslatorBase API*/
+
+	virtual TFuture<TOptional<UE::Interchange::FMeshPayloadData>> GetMeshPayloadData(const FInterchangeMeshPayLoadKey& PayLoadKey, const FTransform& MeshGlobalTransform) const override;
+	virtual TOptional<UE::Interchange::FImportImage> GetTexturePayloadData(const FString& PayloadKey, TOptional<FString>& AlternateTexturePath) const override;
+	virtual TArray<UE::Interchange::FAnimationPayloadData> GetAnimationPayloadData(const TArray<UE::Interchange::FAnimationPayloadQuery>& PayloadQuery) const override;
+};
