[Interchange] Only Animation automatic detection fix for MorphTargetMeshes(with conversion turned on) Nodes.

#rnx

diff --git a/Engine/Plugins/Interchange/Runtime/Source/Pipelines/Private/InterchangeGenericAnimationPipeline.cpp b/Engine/Plugins/Interchange/Runtime/Source/Pipelines/Private/InterchangeGenericAnimationPipeline.cpp
--- a/Engine/Plugins/Interchange/Runtime/Source/Pipelines/Private/InterchangeGenericAnimationPipeline.cpp
+++ b/Engine/Plugins/Interchange/Runtime/Source/Pipelines/Private/InterchangeGenericAnimationPipeline.cpp
@@ -26,7 +26,7 @@
 
 namespace UE::Interchange::Private
 {
-	bool IsTranslatedDataContainOnlyJointAnimation(const UInterchangeBaseNodeContainer* InBaseNodeContainer)
+	bool IsTranslatedDataContainOnlyJointAnimation(const UInterchangeBaseNodeContainer* InBaseNodeContainer, bool bConvertStaticsWithMorphTargetsToSkeletals)
 	{
 		//Its valid to call GetMeshesInformationFromTranslatedData with a null container
 		if (!InBaseNodeContainer)
@@ -44,7 +44,7 @@ namespace UE::Interchange::Private
 		{
 			//if we have bone animation and no skinned mesh, we want to import animation only.
 			bool bContainSkinnedMeshNode = false;
-			InBaseNodeContainer->BreakableIterateNodesOfType<UInterchangeMeshNode>([&bContainSkinnedMeshNode](const FString& NodeUid, UInterchangeMeshNode* MeshNode)
+			InBaseNodeContainer->BreakableIterateNodesOfType<UInterchangeMeshNode>([&bContainSkinnedMeshNode, &bConvertStaticsWithMorphTargetsToSkeletals](const FString& NodeUid, UInterchangeMeshNode* MeshNode)
 				{
 					if (!MeshNode->IsMorphTarget())
 					{
@@ -53,6 +53,13 @@ namespace UE::Interchange::Private
 							bContainSkinnedMeshNode = true;
 						}
 					}
+					else
+					{
+						if (bConvertStaticsWithMorphTargetsToSkeletals)
+						{
+							bContainSkinnedMeshNode = true;
+						}
+					}
 					return bContainSkinnedMeshNode;
 				});
 			bContainOnlyJointAnimation = !bContainSkinnedMeshNode;
@@ -110,7 +117,7 @@ void UInterchangeGenericAnimationPipeline::AdjustSettingsForContext(EInterchange
 	TArray<FString> HideCategories;
 	if(ImportType == EInterchangePipelineContext::AssetImport)
 	{
-		if(UE::Interchange::Private::IsTranslatedDataContainOnlyJointAnimation(InBaseNodeContainer))
+		if(UE::Interchange::Private::IsTranslatedDataContainOnlyJointAnimation(InBaseNodeContainer, CommonSkeletalMeshesAndAnimationsProperties->bConvertStaticsWithMorphTargetsToSkeletals))
 		{
 			bImportAnimations = true;
 			CommonSkeletalMeshesAndAnimationsProperties->bImportOnlyAnimations = true;
