[Interchange - glTF] Allowing no BufferViews for Accessors with Sparse as per specification.

#rnx

diff --git a/Engine/Plugins/Interchange/Runtime/Source/Parsers/GLTFCore/Private/GLTF/GLTFAccessor.cpp b/Engine/Plugins/Interchange/Runtime/Source/Parsers/GLTFCore/Private/GLTF/GLTFAccessor.cpp
--- a/Engine/Plugins/Interchange/Runtime/Source/Parsers/GLTFCore/Private/GLTF/GLTFAccessor.cpp
+++ b/Engine/Plugins/Interchange/Runtime/Source/Parsers/GLTFCore/Private/GLTF/GLTFAccessor.cpp
@@ -290,6 +290,12 @@ namespace GLTF
 		template<typename ItemType, uint32 ItemElementCount>
 		void CopyWithoutConversion(const uint32 ByteStride, const FAccessor::EComponentType ComponentType, bool bNormalized, const uint32 Count, const FBufferView& BufferView, const uint64 ByteOffset, ItemType* Buffer)
 		{
+			if (!BufferView.IsValid())
+			{
+				//Copy can be called for when BufferView is not set, but Sparse is used. (which makes the Accessor still valid.)
+				return;
+			}
+
 			// Stride equals item size => use simpler copy
 			if ((ByteStride == 0) || ByteStride == sizeof(ItemType))
 			{
@@ -324,6 +330,11 @@ namespace GLTF
 		template<typename ItemType, uint32 ItemElementCount>
 		void CopyWithoutConversion(const FAccessor& Accessor, ItemType* Buffer)
 		{
+			if (!Accessor.BufferView.IsValid())
+			{
+				//Copy can be called for when BufferView is not set, but Sparse is used. (which makes the Accessor still valid.)
+				return;
+			}
 			CopyWithoutConversion<ItemType, ItemElementCount>(Accessor.ByteStride, Accessor.ComponentType, Accessor.bNormalized, Accessor.Count, Accessor.BufferView, Accessor.ByteOffset, Buffer);
 		}
 
@@ -956,7 +967,7 @@ namespace GLTF
 
 	bool FAccessor::IsValid() const
 	{
-		return BufferView.IsValid();
+		return BufferView.IsValid() || Sparse.bHasSparse;
 	}
 
 	FMD5Hash FAccessor::GetHash() const
@@ -1233,56 +1244,56 @@ namespace GLTF
 	void FAccessor::GetUnsignedIntArray(TArray<uint32>& Buffer) const
 	{
 		if (IsValid())
-			Buffer.SetNumUninitialized(Count, EAllowShrinking::No);
+			Buffer.SetNumZeroed(Count, EAllowShrinking::No);
 		GetUnsignedIntArray(Buffer.GetData());
 	}
 
 	void FAccessor::GetFloatArray(TArray<float>& Buffer) const
 	{
 		if (IsValid())
-			Buffer.SetNumUninitialized(Count, EAllowShrinking::No);
+			Buffer.SetNumZeroed(Count, EAllowShrinking::No);
 		GetFloatArray(Buffer.GetData());
 	}
 
 	void FAccessor::GetVec2Array(TArray<FVector2f>& Buffer) const
 	{
 		if (IsValid())
-			Buffer.SetNumUninitialized(Count, EAllowShrinking::No);
+			Buffer.SetNumZeroed(Count, EAllowShrinking::No);
 		GetVec2Array(Buffer.GetData());
 	}
 
 	void FAccessor::GetVec3Array(TArray<FVector3f>& Buffer) const
 	{
 		if (IsValid())
-			Buffer.SetNumUninitialized(Count, EAllowShrinking::No);
+			Buffer.SetNumZeroed(Count, EAllowShrinking::No);
 		GetVec3Array(Buffer.GetData());
 	}
 
 	void FAccessor::GetCoordArray(TArray<FVector3f>& Buffer) const
 	{
 		if (IsValid())
-			Buffer.SetNumUninitialized(Count, EAllowShrinking::No);
+			Buffer.SetNumZeroed(Count, EAllowShrinking::No);
 		GetCoordArray(Buffer.GetData());
 	}
 
 	void FAccessor::GetVec4Array(TArray<FVector4f>& Buffer) const
 	{
 		if (IsValid())
-			Buffer.SetNumUninitialized(Count, EAllowShrinking::No);
+			Buffer.SetNumZeroed(Count, EAllowShrinking::No);
 		GetVec4Array(Buffer.GetData());
 	}
 
 	void FAccessor::GetQuatArray(TArray<FVector4f>& Buffer) const
 	{
 		if (IsValid())
-			Buffer.SetNumUninitialized(Count, EAllowShrinking::No);
+			Buffer.SetNumZeroed(Count, EAllowShrinking::No);
 		GetQuatArray(Buffer.GetData());
 	}
 
 	void FAccessor::GetMat4Array(TArray<FMatrix44f>& Buffer) const
 	{
 		if (IsValid())
-			Buffer.SetNumUninitialized(Count, EAllowShrinking::No);
+			Buffer.SetNumZeroed(Count, EAllowShrinking::No);
 		GetMat4Array(Buffer.GetData());
 	}
 }  // namespace GLTF
diff --git a/Engine/Plugins/Interchange/Runtime/Source/Parsers/GLTFCore/Private/GLTF/GLTFReader.cpp b/Engine/Plugins/Interchange/Runtime/Source/Parsers/GLTFCore/Private/GLTF/GLTFReader.cpp
--- a/Engine/Plugins/Interchange/Runtime/Source/Parsers/GLTFCore/Private/GLTF/GLTFReader.cpp
+++ b/Engine/Plugins/Interchange/Runtime/Source/Parsers/GLTFCore/Private/GLTF/GLTFReader.cpp
@@ -301,6 +301,8 @@ namespace GLTF
 			if (!Object.HasTypedField<EJson::Number>(TEXT("bufferView")))
 			{
 				//if bufferView does not exist in the Object, then the presumption is that it is a (Draco) CompressedAccessor:
+				//Accessors with Sparse present can have no bufferView, in which case:
+				//"When accessor.bufferView is undefined, the sparse accessor is initialized as an array of zeros of size (size of the accessor element) * (accessor.count) bytes"
 				const uint32                    Count = GetUnsignedInt(Object, TEXT("count"), 0);
 				const FAccessor::EType          Type = AccessorTypeFromString(Object.GetStringField(TEXT("type")));
 				const FAccessor::EComponentType CompType = ComponentTypeFromNumber(GetUnsignedInt(Object, TEXT("componentType"), 0));
