Fixed a crash when importing a texture at runtime via interchange.

The crash would happen when the texture format of the rendering data was different from the source data.For example by importing a HDR file.

diff --git a/Engine/Plugins/Interchange/Runtime/Source/Import/Private/Texture/InterchangeTextureFactory.cpp b/Engine/Plugins/Interchange/Runtime/Source/Import/Private/Texture/InterchangeTextureFactory.cpp
--- a/Engine/Plugins/Interchange/Runtime/Source/Import/Private/Texture/InterchangeTextureFactory.cpp
+++ b/Engine/Plugins/Interchange/Runtime/Source/Import/Private/Texture/InterchangeTextureFactory.cpp
@@ -1481,8 +1481,10 @@ namespace UE::Interchange::Private::InterchangeTextureFactory
 	{
 		Texture2D->SRGB = UE::TextureUtilitiesCommon::GetDefaultSRGB(Texture2D->CompressionSettings,ImportImage.Format,ImportImage.bSRGB);
 		
+		ERawImageFormat::Type SourcePixelRawFormat = FImageCoreUtils::ConvertToRawImageFormat(ImportImage.Format);
+
 		ERawImageFormat::Type PixelFormatRawFormat;
-		const EPixelFormat PixelFormat = FImageCoreUtils::GetPixelFormatForRawImageFormat(FImageCoreUtils::ConvertToRawImageFormat(ImportImage.Format), &PixelFormatRawFormat);
+		const EPixelFormat PixelFormat = FImageCoreUtils::GetPixelFormatForRawImageFormat(SourcePixelRawFormat, &PixelFormatRawFormat);
 
 		UE::Serialization::FEditorBulkData BulkData;
 		BulkData.UpdatePayload(MoveTemp(BufferAndId));
@@ -1490,7 +1492,7 @@ namespace UE::Interchange::Private::InterchangeTextureFactory
 
 		const EGammaSpace SourceGammaSpace = ImportImage.bSRGB ? EGammaSpace::sRGB : EGammaSpace::Linear;
 		constexpr int32 NumSlices = 1;
-		FImageView SourceImageView(const_cast<void*>(Payload.Get().GetData()), ImportImage.SizeX, ImportImage.SizeY, NumSlices, PixelFormatRawFormat, SourceGammaSpace);
+		FImageView SourceImageView(const_cast<void*>(Payload.Get().GetData()), ImportImage.SizeX, ImportImage.SizeY, NumSlices, SourcePixelRawFormat, SourceGammaSpace);
 
 		FImage DecompressedSourceImage;
 		if (ImportImage.RawDataCompressionFormat != TSCF_None)
