Interchange: fix race condition when getting multiple animation payloads.

#rnx

diff --git a/Engine/Plugins/Interchange/Runtime/Source/Parsers/Fbx/Private/FbxAnimation.cpp b/Engine/Plugins/Interchange/Runtime/Source/Parsers/Fbx/Private/FbxAnimation.cpp
--- a/Engine/Plugins/Interchange/Runtime/Source/Parsers/Fbx/Private/FbxAnimation.cpp
+++ b/Engine/Plugins/Interchange/Runtime/Source/Parsers/Fbx/Private/FbxAnimation.cpp
@@ -1007,10 +1007,13 @@ namespace UE::Interchange::Private
 
 				FString ResultPayloadUniqueId = PayloadQuery->GetHashString();
 
-				//If we already have extract this mesh, no need to extract again
-				if (ResultPayloads.Contains(ResultPayloadUniqueId))
 				{
-					continue;
+					FScopeLock Lock(ResultPayloadsCriticalSection);
+					//If we already have extract this mesh, no need to extract again
+					if (ResultPayloads.Contains(ResultPayloadUniqueId))
+					{
+						continue;
+					}
 				}
 
 				TArray<FPayloadQueryHelper>& PayloadHelpersPerAnimStack = PayloadQueryHelpers.FindOrAdd(AnimStack);
@@ -1078,16 +1081,15 @@ namespace UE::Interchange::Private
 					FPayloadDataHelper& PayloadDataHelper = PayloadDataHelpersForAnimStack[PayloadDataHelperIndex];
 					FString QueryHashString = PayloadDataHelper.QueryHashString; //PayloadQuery's GetHashString()
 
-					//If we already have extract this mesh, no need to extract again
-					if (ResultPayloads.Contains(QueryHashString))
-					{
-						return;
-					}
-
 					FString PayloadFilepathCopy;
-
 					{
 						FScopeLock Lock(ResultPayloadsCriticalSection);
+						//If we already have extract this mesh, no need to extract again
+						if (ResultPayloads.Contains(QueryHashString))
+						{
+							return;
+						}
+
 						FString& PayloadFilepath = ResultPayloads.FindOrAdd(QueryHashString);
 						//To avoid file path with too many character, we hash the payloadKey so we have a deterministic length for the file path.
 						PayloadFilepath = ResultFolder + TEXT("/") + QueryHashString + FString::FromInt(UniqueIdCounter.IncrementExchange()) + TEXT(".payload");
