Interchange: Small heuristic to handle translators that implement both the regular and the blocked texture payload interface.

diff --git a/Engine/Plugins/Interchange/Runtime/Source/Import/Private/Texture/InterchangeTextureFactory.cpp b/Engine/Plugins/Interchange/Runtime/Source/Import/Private/Texture/InterchangeTextureFactory.cpp
--- a/Engine/Plugins/Interchange/Runtime/Source/Import/Private/Texture/InterchangeTextureFactory.cpp
+++ b/Engine/Plugins/Interchange/Runtime/Source/Import/Private/Texture/InterchangeTextureFactory.cpp
@@ -680,12 +680,18 @@ namespace UE::Interchange::Private::InterchangeTextureFactory
 				TextureFactoryNode->GetCustomPreferCompressedSourceData(bShoudImportCompressedImage);
 			}
 
-			// Is there a case were a translator can be both interface and how should the factory chose which to invoke?
-			if (const IInterchangeTexturePayloadInterface* TextureTranslator = Cast<IInterchangeTexturePayloadInterface>(Translator))
+			const IInterchangeTexturePayloadInterface* TextureTranslator = Cast<IInterchangeTexturePayloadInterface>(Translator);
+			const IInterchangeBlockedTexturePayloadInterface* BlockedTextureTranslator = Cast<IInterchangeBlockedTexturePayloadInterface>(Translator);
+
+			// If the translator implements both interfaces we need to decide which one to invoke
+			const bool bInvokeBlockedInterface = (BlockedTextureTranslator && !TextureTranslator)
+										   || (BlockedTextureTranslator && TextureTranslator && !BlockAndSourceDataFiles.IsEmpty());
+
+			if (TextureTranslator && !bInvokeBlockedInterface)
 			{
 				if (BlockAndSourceDataFiles.IsEmpty())
 				{
-					if (Translator->GetClass() == UInterchangeJPGTranslator::StaticClass() 
+					if (Translator->GetClass() == UInterchangeJPGTranslator::StaticClass()
 						|| Translator->GetClass() == UInterchangeUEJPEGTranslator::StaticClass())
 					{
 						// Honor setting from TextureImporter.RetainJpegFormat in Editor.ini if it exists (ideally we should deprecate this as it is confusing and probably not thread safe)
@@ -710,7 +716,7 @@ namespace UE::Interchange::Private::InterchangeTextureFactory
 					return FTexturePayloadVariant(TInPlaceType<TOptional<FImportBlockedImage>>(), GetBlockedTexturePayloadDataFromSourceFiles(SourceData, BlockAndSourceDataFiles, Translator));
 				}
 			}
-			else if (const IInterchangeBlockedTexturePayloadInterface* BlockedTextureTranslator = Cast<IInterchangeBlockedTexturePayloadInterface>(Translator))
+			else if (bInvokeBlockedInterface)
 			{
 				return FTexturePayloadVariant(TInPlaceType<TOptional<FImportBlockedImage>>(), BlockedTextureTranslator->GetBlockedTexturePayloadData(PayloadKey, AlternateTexturePath));
 			}
