[Interchange Import] Fix an issue with MaterialX linkage that was causing the module to link both statically and dynamically

diff --git a/Engine/Plugins/Interchange/Runtime/Source/Import/InterchangeImport.Build.cs b/Engine/Plugins/Interchange/Runtime/Source/Import/InterchangeImport.Build.cs
--- a/Engine/Plugins/Interchange/Runtime/Source/Import/InterchangeImport.Build.cs
+++ b/Engine/Plugins/Interchange/Runtime/Source/Import/InterchangeImport.Build.cs
@@ -63,7 +63,15 @@ namespace UnrealBuildTool.Rules
 					}
 				);
 
-				AddEngineThirdPartyPrivateStaticDependencies(Target, "MaterialX");
+				if (Target.Platform == UnrealTargetPlatform.Win64)
+				{
+					PrivateDependencyModuleNames.AddRange(
+						new string[]
+						{
+							"MaterialX"
+						}
+					);
+				}
 			}
 		}
 	}
