[Interchange glTF] Fixing Errors reading Objects/Arrays as strings (for extras). Using FMatrix instead of FTransform for BindPose acquisitions.

#rnx

diff --git a/Engine/Plugins/Interchange/Runtime/Source/Parsers/GLTFCore/Private/GLTF/GLTFReader.cpp b/Engine/Plugins/Interchange/Runtime/Source/Parsers/GLTFCore/Private/GLTF/GLTFReader.cpp
--- a/Engine/Plugins/Interchange/Runtime/Source/Parsers/GLTFCore/Private/GLTF/GLTFReader.cpp
+++ b/Engine/Plugins/Interchange/Runtime/Source/Parsers/GLTFCore/Private/GLTF/GLTFReader.cpp
@@ -143,6 +143,41 @@ namespace GLTF
 			Transform.SetRotation(OutRotation);
 			Transform.SetTranslation(OutTranslation);
 		}
+
+		void ProcessExtras(const FJsonObject& Object, TMap<FString, FString>& StorageForExtras, TSet<FString> KeyExceptions = {})
+		{
+			if (Object.HasField(TEXT("extras")))
+			{
+				const TSharedPtr<FJsonObject>& Extras = Object.GetObjectField(TEXT("extras"));
+				if (KeyExceptions.Num() > 0)
+				{
+					for (TPair<FString, TSharedPtr<FJsonValue>>& Pair : Extras->Values)
+					{
+						if (KeyExceptions.Contains(Pair.Key))
+						{
+							continue;
+						}
+
+						FString ExtraString;
+						if (Pair.Value->TryGetString(ExtraString))
+						{
+							StorageForExtras.Add(Pair.Key, ExtraString);
+						}
+					}
+				}
+				else
+				{
+					for (TPair<FString, TSharedPtr<FJsonValue>>& Pair : Extras->Values)
+					{
+						FString ExtraString;
+						if (Pair.Value->TryGetString(ExtraString))
+						{
+							StorageForExtras.Add(Pair.Key, ExtraString);
+						}
+					}
+				}
+			}
+		}
 	}
 
 	FFileReader::FFileReader()
@@ -413,14 +448,7 @@ namespace GLTF
 			}
 		}
 
-		if (Object.HasField(TEXT("extras")))
-		{
-			const TSharedPtr<FJsonObject>& Extras = Object.GetObjectField(TEXT("extras"));
-			for (TPair<FString, TSharedPtr<FJsonValue>>& Pair : Extras->Values)
-			{
-				Mesh.Primitives.Last().Extras.Add(Pair.Key, Pair.Value->AsString());
-			}
-		}
+		ProcessExtras(Object, Mesh.Primitives.Last().Extras);
 
 		ExtensionsHandler->SetupPrimitiveExtensions(Object, Mesh.Primitives.Last(), Mesh.Primitives.Num()-1, Mesh.UniqueId);
 	}
@@ -482,13 +510,7 @@ namespace GLTF
 				}
 			}
 
-			for (TPair<FString, TSharedPtr<FJsonValue>>& Pair : Extras->Values)
-			{
-				if (Pair.Key != TEXT("targetNames"))
-				{
-					Mesh.Extras.Add(Pair.Key, Pair.Value->AsString());
-				}
-			}
+			ProcessExtras(Object, Mesh.Extras, { TEXT("targetNames") });
 		}
 
 		Mesh.GenerateIsValidCache();
@@ -514,14 +536,7 @@ namespace GLTF
 			}
 		}
 
-		if (Object.HasField(TEXT("extras")))
-		{
-			const TSharedPtr<FJsonObject>& Extras = Object.GetObjectField(TEXT("extras"));
-			for (TPair<FString, TSharedPtr<FJsonValue>>& Pair : Extras->Values)
-			{
-				Scene.Extras.Add(Pair.Key, Pair.Value->AsString());
-			}
-		}
+		ProcessExtras(Object, Scene.Extras);
 
 		ExtensionsHandler->SetupSceneExtensions(Object, Scene);
 	}
@@ -576,14 +591,7 @@ namespace GLTF
 			}
 		}
 
-		if (Object.HasField(TEXT("extras")))
-		{
-			const TSharedPtr<FJsonObject>& Extras = Object.GetObjectField(TEXT("extras"));
-			for (TPair<FString, TSharedPtr<FJsonValue>>& Pair : Extras->Values)
-			{
-				Node.Extras.Add(Pair.Key, Pair.Value->AsString());
-			}
-		}
+		ProcessExtras(Object, Node.Extras);
 
 		ExtensionsHandler->SetupNodeExtensions(Object, Node);
 	}
@@ -628,15 +636,8 @@ namespace GLTF
 			Messages.Emplace(EMessageSeverity::Error, TEXT("Invalid camera type: ") + Type);
 		}
 
-		if (Object.HasField(TEXT("extras")))
-		{
-			const TSharedPtr<FJsonObject>& Extras = Object.GetObjectField(TEXT("extras"));
-			for (TPair<FString, TSharedPtr<FJsonValue>>& Pair : Extras->Values)
-			{
-				Camera.Extras.Add(Pair.Key, Pair.Value->AsString());
-			}
-		}
-
+		ProcessExtras(Object, Camera.Extras);
+		
 		ExtensionsHandler->SetupCameraExtensions(Object, Camera);
 	}
 
@@ -703,15 +704,8 @@ namespace GLTF
 			}
 		}
 
-		if (Object.HasField(TEXT("extras")))
-		{
-			const TSharedPtr<FJsonObject>& Extras = Object.GetObjectField(TEXT("extras"));
-			for (TPair<FString, TSharedPtr<FJsonValue>>& Pair : Extras->Values)
-			{
-				Animation.Extras.Add(Pair.Key, Pair.Value->AsString());
-			}
-		}
-
+		ProcessExtras(Object, Animation.Extras);
+		
 		ExtensionsHandler->SetupAnimationExtensions(Object, Animation);
 	}
 
@@ -732,15 +726,8 @@ namespace GLTF
 
 		Skin.Skeleton = GetIndex(Object, TEXT("skeleton"));
 
-		if (Object.HasField(TEXT("extras")))
-		{
-			const TSharedPtr<FJsonObject>& Extras = Object.GetObjectField(TEXT("extras"));
-			for (TPair<FString, TSharedPtr<FJsonValue>>& Pair : Extras->Values)
-			{
-				Skin.Extras.Add(Pair.Key, Pair.Value->AsString());
-			}
-		}
-
+		ProcessExtras(Object, Skin.Extras);
+		
 		ExtensionsHandler->SetupSkinExtensions(Object, Skin);
 	}
 
@@ -819,15 +806,8 @@ namespace GLTF
 			CurrentBufferOffset += Image.DataByteLength;
 		}
 
-		if (Object.HasField(TEXT("extras")))
-		{
-			const TSharedPtr<FJsonObject>& Extras = Object.GetObjectField(TEXT("extras"));
-			for (TPair<FString, TSharedPtr<FJsonValue>>& Pair : Extras->Values)
-			{
-				Image.Extras.Add(Pair.Key, Pair.Value->AsString());
-			}
-		}
-
+		ProcessExtras(Object, Image.Extras);
+		
 		ExtensionsHandler->SetupImageExtensions(Object, Image);
 	}
 
@@ -864,15 +844,8 @@ namespace GLTF
 
 			Asset->Textures.Emplace(TexName, Source, Sampler);
 
-			if (Object.HasField(TEXT("extras")))
-			{
-				const TSharedPtr<FJsonObject>& Extras = Object.GetObjectField(TEXT("extras"));
-				for (TPair<FString, TSharedPtr<FJsonValue>>& Pair : Extras->Values)
-				{
-					Asset->Textures.Last().Extras.Add(Pair.Key, Pair.Value->AsString());
-				}
-			}
-
+			ProcessExtras(Object, Asset->Textures.Last().Extras);
+			
 			ExtensionsHandler->SetupTextureExtensions(Object, Asset->Textures.Last());
 		}
 		else
@@ -915,15 +888,8 @@ namespace GLTF
 
 		Material.bIsDoubleSided = GetBool(Object, TEXT("doubleSided"));
 
-		if (Object.HasField(TEXT("extras")))
-		{
-			const TSharedPtr<FJsonObject>& Extras = Object.GetObjectField(TEXT("extras"));
-			for (TPair<FString, TSharedPtr<FJsonValue>>& Pair : Extras->Values)
-			{
-				Material.Extras.Add(Pair.Key, Pair.Value->AsString());
-			}
-		}
-
+		ProcessExtras(Object, Material.Extras);
+		
 		ExtensionsHandler->SetupMaterialExtensions(Object, Material);
 	}
 
@@ -1365,14 +1331,8 @@ namespace GLTF
 				{
 					FMatrix InverseBindMatrix = Skin.InverseBindMatrices.GetMat4(JointCounter);
 					InverseBindMatrix = GLTF::ConvertMat(InverseBindMatrix);
-
-					FTransform InverseBindMatrixTransform;
-					InverseBindMatrixTransform.SetFromMatrix(InverseBindMatrix);
-					InverseBindMatrixTransform.SetRotation(GLTF::ConvertQuat(InverseBindMatrixTransform.GetRotation()));
-					InverseBindMatrixTransform.SetTranslation(GLTF::ConvertVec3(InverseBindMatrixTransform.GetTranslation()));
-					InverseBindMatrixTransform.SetScale3D(GLTF::ConvertVec3(InverseBindMatrixTransform.GetScale3D()));
-
-					Asset->Nodes[Skin.Joints[JointCounter]].SkinIndexToGlobalInverseBindTransform.Add(SkinIndex, InverseBindMatrixTransform);
+					
+					Asset->Nodes[Skin.Joints[JointCounter]].SkinIndexToGlobalInverseBindMatrix.Add(SkinIndex, InverseBindMatrix);
 				}
 			}
 		}
@@ -1408,32 +1368,32 @@ namespace GLTF
 					if (CurrentNode.ParentIndex != INDEX_NONE &&
 						Asset->Nodes.IsValidIndex(CurrentNode.ParentIndex) &&
 						Asset->Nodes[CurrentNode.ParentIndex].Type == FNode::EType::Joint &&
-						(Asset->Nodes[CurrentNode.ParentIndex].SkinIndexToGlobalInverseBindTransform.Contains(SkinIndex) || Asset->Nodes[CurrentNode.ParentIndex].SkinIndexToGlobalInverseBindTransform.Num() > 0)
+						(Asset->Nodes[CurrentNode.ParentIndex].SkinIndexToGlobalInverseBindMatrix.Contains(SkinIndex) || Asset->Nodes[CurrentNode.ParentIndex].SkinIndexToGlobalInverseBindMatrix.Num() > 0)
 						)
 					{
 						FNode& ParentNode = Asset->Nodes[CurrentNode.ParentIndex];
 
 						//LocalBindPose; //bind pose would be CurrentNode.GlobalInverseBindTransform.Inverse() * ParentNode.GlobalInverseBindTransform
-						FTransform ParentGlobalInverseBindTransform;
-						if (ParentNode.SkinIndexToGlobalInverseBindTransform.Contains(SkinIndex))
+						FMatrix ParentGlobalInverseBindMatrix;
+						if (ParentNode.SkinIndexToGlobalInverseBindMatrix.Contains(SkinIndex))
 						{
-							ParentGlobalInverseBindTransform = ParentNode.SkinIndexToGlobalInverseBindTransform[SkinIndex];
+							ParentGlobalInverseBindMatrix = ParentNode.SkinIndexToGlobalInverseBindMatrix[SkinIndex];
 						}
 						else
 						{
 							//Scenario is that the a Skin is instantiated at the end of another skin
 							//(Prime example is the RecursiveSkeleton gltf sample file.)
-							ParentGlobalInverseBindTransform = ParentNode.SkinIndexToGlobalInverseBindTransform.begin().Value();
+							ParentGlobalInverseBindMatrix = ParentNode.SkinIndexToGlobalInverseBindMatrix.begin().Value();
 						}
 
-						FTransform LocalBindPose = CurrentNode.SkinIndexToGlobalInverseBindTransform[SkinIndex].Inverse() * ParentGlobalInverseBindTransform;
+						FMatrix LocalBindMatrix = CurrentNode.SkinIndexToGlobalInverseBindMatrix[SkinIndex].Inverse() * ParentGlobalInverseBindMatrix;
 
-						CurrentNode.SkinIndexToLocalBindPose.Add(SkinIndex, LocalBindPose);
+						CurrentNode.SkinIndexToLocalBindMatrix.Add(SkinIndex, LocalBindMatrix);
 					}
 					else
 					{
-						FTransform LocalBindPose = CurrentNode.SkinIndexToGlobalInverseBindTransform[SkinIndex].Inverse();
-						CurrentNode.SkinIndexToLocalBindPose.Add(SkinIndex, LocalBindPose);
+							FMatrix LocalBindPose = CurrentNode.SkinIndexToGlobalInverseBindMatrix[SkinIndex].Inverse();
+							CurrentNode.SkinIndexToLocalBindMatrix.Add(SkinIndex, LocalBindPose);
 					}
 				}
 			}
@@ -1447,8 +1407,8 @@ namespace GLTF
 
 		for (const FNode& CurrentNode : Asset->Nodes)
 		{
-			TMap<int, FTransform>::TRangedForConstIterator Iter(CurrentNode.SkinIndexToLocalBindPose.begin());
-			FTransform ToCompareAgainst = Iter ? Iter.Value() : FTransform();
+			TMap<int, FMatrix>::TRangedForConstIterator Iter(CurrentNode.SkinIndexToLocalBindMatrix.begin());
+			FMatrix ToCompareAgainst = Iter ? Iter.Value() : FMatrix();
 			++Iter;
 			for (; Iter; ++Iter)
 			{
@@ -1490,10 +1450,17 @@ namespace GLTF
 					FNode& CurrentNode = Asset->Nodes[Skin.Joints[JointCounter]];
 
 					if (!CurrentNode.bHasLocalBindPose
-						&& CurrentNode.SkinIndexToLocalBindPose.Contains(SkinIndex))
+						&& CurrentNode.SkinIndexToLocalBindMatrix.Contains(SkinIndex))
 					{
 						CurrentNode.bHasLocalBindPose = true;
-						CurrentNode.LocalBindPose = CurrentNode.SkinIndexToLocalBindPose[SkinIndex];
+						
+						FTransform LocalBindPoseTransform;
+						LocalBindPoseTransform.SetFromMatrix(CurrentNode.SkinIndexToLocalBindMatrix[SkinIndex]);
+						LocalBindPoseTransform.SetRotation(GLTF::ConvertQuat(LocalBindPoseTransform.GetRotation()));
+						LocalBindPoseTransform.SetTranslation(GLTF::ConvertVec3(LocalBindPoseTransform.GetTranslation()));
+						LocalBindPoseTransform.SetScale3D(GLTF::ConvertVec3(LocalBindPoseTransform.GetScale3D()));
+
+						CurrentNode.LocalBindPose = LocalBindPoseTransform;
 					}
 				}
 			}
diff --git a/Engine/Plugins/Interchange/Runtime/Source/Parsers/GLTFCore/Public/GLTFNode.h b/Engine/Plugins/Interchange/Runtime/Source/Parsers/GLTFCore/Public/GLTFNode.h
--- a/Engine/Plugins/Interchange/Runtime/Source/Parsers/GLTFCore/Public/GLTFNode.h
+++ b/Engine/Plugins/Interchange/Runtime/Source/Parsers/GLTFCore/Public/GLTFNode.h
@@ -43,8 +43,9 @@ namespace GLTF
 
 		TArray<float> MorphTargetWeights; //for the instantiated mesh with morph targets.
 
-		TMap<int, FTransform> SkinIndexToGlobalInverseBindTransform;
-		TMap<int, FTransform> SkinIndexToLocalBindPose; //bind pose would be CurrentNode.GlobalInverseBindTransform.Inverse() * ParentNode.GlobalInverseBindTransform
+		TMap<int, FMatrix> SkinIndexToGlobalInverseBindMatrix; //Matrix in glTF space
+		TMap<int, FMatrix> SkinIndexToLocalBindMatrix;	//Matrix in glTF space
+														//bind pose would be CurrentNode.GlobalInverseBindTransform.Inverse() * ParentNode.GlobalInverseBindTransform
 		bool bHasLocalBindPose;
 		FTransform LocalBindPose;	// First Skin that's using the joint will fill the LocalBindPose.
 									//	Edge case Scenario which is currently not supported:
