[Backout] - CL34161723
Original CL Desc
-----------------------------------------------------------------
[Interchange Import] Fix an issue with MaterialX linkage that was causing the module to link both statically and dynamically

diff --git a/Engine/Plugins/Interchange/Runtime/Source/Import/InterchangeImport.Build.cs b/Engine/Plugins/Interchange/Runtime/Source/Import/InterchangeImport.Build.cs
--- a/Engine/Plugins/Interchange/Runtime/Source/Import/InterchangeImport.Build.cs
+++ b/Engine/Plugins/Interchange/Runtime/Source/Import/InterchangeImport.Build.cs
@@ -63,15 +63,7 @@ namespace UnrealBuildTool.Rules
 					}
 				);
 
-				if (Target.Platform == UnrealTargetPlatform.Win64)
-				{
-					PrivateDependencyModuleNames.AddRange(
-						new string[]
-						{
-							"MaterialX"
-						}
-					);
-				}
+				AddEngineThirdPartyPrivateStaticDependencies(Target, "MaterialX");
 			}
 		}
 	}
