[Interchange Import] Fix an issue with MaterialX linkage that was causing the module to link both statically and dynamically

diff --git a/Engine/Plugins/Interchange/Runtime/Source/Import/InterchangeImport.Build.cs b/Engine/Plugins/Interchange/Runtime/Source/Import/InterchangeImport.Build.cs
--- a/Engine/Plugins/Interchange/Runtime/Source/Import/InterchangeImport.Build.cs
+++ b/Engine/Plugins/Interchange/Runtime/Source/Import/InterchangeImport.Build.cs
@@ -60,10 +60,9 @@ namespace UnrealBuildTool.Rules
 						"SkeletalMeshUtilitiesCommon",
 						"UnrealEd",
 						"VariantManager",
+						"MaterialX",
 					}
 				);
-
-				AddEngineThirdPartyPrivateStaticDependencies(Target, "MaterialX");
 			}
 		}
 	}
