CIS fix for monolithic targets using interchange import library

#rnx

diff --git a/Engine/Plugins/Interchange/Runtime/Source/Import/InterchangeImport.Build.cs b/Engine/Plugins/Interchange/Runtime/Source/Import/InterchangeImport.Build.cs
--- a/Engine/Plugins/Interchange/Runtime/Source/Import/InterchangeImport.Build.cs
+++ b/Engine/Plugins/Interchange/Runtime/Source/Import/InterchangeImport.Build.cs
@@ -63,7 +63,7 @@ namespace UnrealBuildTool.Rules
 					}
 				);
 
-				if (Target.Platform == UnrealTargetPlatform.Win64)
+				if (Target.Platform == UnrealTargetPlatform.Win64 || Target.LinkType == TargetLinkType.Monolithic)
 				{
 					PrivateDependencyModuleNames.AddRange(
 						new string[]
