Interchange: Add the class of the object to give more context when adjusting pipeline settings.

#rnx

diff --git a/Engine/Plugins/Experimental/Enterprise/DatasmithInterchange/Source/DatasmithInterchange/Private/InterchangeDatasmithPipeline.cpp b/Engine/Plugins/Experimental/Enterprise/DatasmithInterchange/Source/DatasmithInterchange/Private/InterchangeDatasmithPipeline.cpp
--- a/Engine/Plugins/Experimental/Enterprise/DatasmithInterchange/Source/DatasmithInterchange/Private/InterchangeDatasmithPipeline.cpp
+++ b/Engine/Plugins/Experimental/Enterprise/DatasmithInterchange/Source/DatasmithInterchange/Private/InterchangeDatasmithPipeline.cpp
@@ -120,37 +120,37 @@ UInterchangeDatasmithPipeline::UInterchangeDatasmithPipeline()
 	AnimationPipeline->CommonSkeletalMeshesAndAnimationsProperties = CommonSkeletalMeshesAndAnimationsProperties;
 }
 
-void UInterchangeDatasmithPipeline::AdjustSettingsForContext(EInterchangePipelineContext ImportType, TObjectPtr<UObject> ReimportAsset, const UInterchangeBaseNodeContainer* InBaseNodeContainer)
+void UInterchangeDatasmithPipeline::AdjustSettingsForContext(const FInterchangePipelineContextParams& ContextParams)
 {
-	Super::AdjustSettingsForContext(ImportType, ReimportAsset, InBaseNodeContainer);
+	Super::AdjustSettingsForContext(ContextParams);
 
 	if (MaterialPipeline)
 	{
-		MaterialPipeline->AdjustSettingsForContext(ImportType, ReimportAsset, InBaseNodeContainer);
+		MaterialPipeline->AdjustSettingsForContext(ContextParams);
 	}
 
 	if (MeshPipeline)
 	{
-		MeshPipeline->AdjustSettingsForContext(ImportType, ReimportAsset, InBaseNodeContainer);
+		MeshPipeline->AdjustSettingsForContext(ContextParams);
 	}
 
 	if (LevelPipeline)
 	{
-		LevelPipeline->AdjustSettingsForContext(ImportType, ReimportAsset, InBaseNodeContainer);
+		LevelPipeline->AdjustSettingsForContext(ContextParams);
 	}
 
 	if (AnimationPipeline)
 	{
-		AnimationPipeline->AdjustSettingsForContext(ImportType, ReimportAsset, InBaseNodeContainer);
+		AnimationPipeline->AdjustSettingsForContext(ContextParams);
 	}
 }
 
 void UInterchangeDatasmithPipeline::PostDuplicate(bool bDuplicateForPIE)
 {
 	// Only adjust settings if there is anything cached.
-	if (CachePipelineContext != EInterchangePipelineContext::None)
+	if (CacheContextParam.ContextType != EInterchangePipelineContext::None)
 	{
-		AdjustSettingsForContext(CachePipelineContext, CacheReimportObject.Get(), nullptr);
+		AdjustSettingsForContext(CacheContextParam);
 	}
 }
 
@@ -180,7 +180,7 @@ void UInterchangeDatasmithPipeline::ExecutePipeline(UInterchangeBaseNodeContaine
 
 	ExecutePreImportPipelineFunc(MeshPipeline);
 
-	if (CachePipelineContext == EInterchangePipelineContext::SceneImport || CachePipelineContext == EInterchangePipelineContext::SceneReimport)
+	if (CacheContextParam.ContextType == EInterchangePipelineContext::SceneImport || CacheContextParam.ContextType == EInterchangePipelineContext::SceneReimport)
 	{
 		ExecutePreImportPipelineFunc(LevelPipeline);
 		ExecutePreImportPipelineFunc(AnimationPipeline);
@@ -260,7 +260,7 @@ void UInterchangeDatasmithPipeline::ExecutePostImportPipeline(const UInterchange
 		MeshPipeline->ScriptedExecutePostImportPipeline(InBaseNodeContainer, NodeKey, CreatedAsset, bIsAReimport);
 	}
 
-	const bool bSceneImport = CachePipelineContext == EInterchangePipelineContext::SceneImport || CachePipelineContext == EInterchangePipelineContext::SceneReimport;
+	const bool bSceneImport = CacheContextParam.ContextType == EInterchangePipelineContext::SceneImport || CacheContextParam.ContextType == EInterchangePipelineContext::SceneReimport;
 	if (bSceneImport)
 	{
 		if (LevelPipeline)
diff --git a/Engine/Plugins/Experimental/Enterprise/DatasmithInterchange/Source/DatasmithInterchange/Public/InterchangeDatasmithPipeline.h b/Engine/Plugins/Experimental/Enterprise/DatasmithInterchange/Source/DatasmithInterchange/Public/InterchangeDatasmithPipeline.h
--- a/Engine/Plugins/Experimental/Enterprise/DatasmithInterchange/Source/DatasmithInterchange/Public/InterchangeDatasmithPipeline.h
+++ b/Engine/Plugins/Experimental/Enterprise/DatasmithInterchange/Source/DatasmithInterchange/Public/InterchangeDatasmithPipeline.h
@@ -64,7 +64,7 @@ public:
 protected:
 	virtual void ExecutePipeline(UInterchangeBaseNodeContainer* InBaseNodeContainer, const TArray<UInterchangeSourceData*>& SourceDatas, const FString& ContentBasePath) override;
 	virtual void ExecutePostImportPipeline(const UInterchangeBaseNodeContainer* InBaseNodeContainer, const FString& NodeKey, UObject* CreatedAsset, bool bIsAReimport) override;
-	virtual void AdjustSettingsForContext(EInterchangePipelineContext ImportType, TObjectPtr<UObject> ReimportAsset, const UInterchangeBaseNodeContainer* BaseNodeContainer) override;
+	virtual void AdjustSettingsForContext(const FInterchangePipelineContextParams& ContextParams) override;
 
 	virtual bool CanExecuteOnAnyThread(EInterchangePipelineTask PipelineTask) override
 	{
diff --git a/Engine/Plugins/Interchange/Runtime/Source/Pipelines/Private/InterchangeGenericAnimationPipeline.cpp b/Engine/Plugins/Interchange/Runtime/Source/Pipelines/Private/InterchangeGenericAnimationPipeline.cpp
--- a/Engine/Plugins/Interchange/Runtime/Source/Pipelines/Private/InterchangeGenericAnimationPipeline.cpp
+++ b/Engine/Plugins/Interchange/Runtime/Source/Pipelines/Private/InterchangeGenericAnimationPipeline.cpp
@@ -89,20 +89,22 @@ bool UInterchangeGenericAnimationPipeline::CanEditChange(const FProperty* InProp
 }
 #endif
 
-void UInterchangeGenericAnimationPipeline::AdjustSettingsForContext(EInterchangePipelineContext ImportType, TObjectPtr<UObject> ReimportAsset, const UInterchangeBaseNodeContainer* InBaseNodeContainer)
+void UInterchangeGenericAnimationPipeline::AdjustSettingsForContext(const FInterchangePipelineContextParams& ContextParams)
 {
-	Super::AdjustSettingsForContext(ImportType, ReimportAsset, InBaseNodeContainer);
+	Super::AdjustSettingsForContext(ContextParams);
 
 #if WITH_EDITOR
 	check(CommonSkeletalMeshesAndAnimationsProperties.IsValid());
 	
-	bSceneImport = ImportType == EInterchangePipelineContext::SceneImport
-				|| ImportType == EInterchangePipelineContext::SceneReimport;
+	bSceneImport = ContextParams.ContextType == EInterchangePipelineContext::SceneImport
+				|| ContextParams.ContextType == EInterchangePipelineContext::SceneReimport;
 
-	if (ImportType == EInterchangePipelineContext::AssetCustomLODImport
-		|| ImportType == EInterchangePipelineContext::AssetCustomLODReimport
-		|| ImportType == EInterchangePipelineContext::AssetAlternateSkinningImport
-		|| ImportType == EInterchangePipelineContext::AssetAlternateSkinningReimport)
+	if (ContextParams.ContextType == EInterchangePipelineContext::AssetCustomLODImport
+		|| ContextParams.ContextType == EInterchangePipelineContext::AssetCustomLODReimport
+		|| ContextParams.ContextType == EInterchangePipelineContext::AssetAlternateSkinningImport
+		|| ContextParams.ContextType == EInterchangePipelineContext::AssetAlternateSkinningReimport
+		|| ContextParams.ContextType == EInterchangePipelineContext::AssetCustomMorphTargetImport
+		|| ContextParams.ContextType == EInterchangePipelineContext::AssetCustomMorphTargetReImport)
 	{
 		bImportAnimations = false;
 		CommonSkeletalMeshesAndAnimationsProperties->Skeleton = nullptr;
@@ -115,9 +117,9 @@ void UInterchangeGenericAnimationPipeline::AdjustSettingsForContext(EInterchange
 	const FString AnimationCategory = UInterchangeGenericAnimationPipeline::GetPipelineCategory(nullptr);
 
 	TArray<FString> HideCategories;
-	if(ImportType == EInterchangePipelineContext::AssetImport)
+	if(ContextParams.ContextType == EInterchangePipelineContext::AssetImport)
 	{
-		if(UE::Interchange::Private::IsTranslatedDataContainOnlyJointAnimation(InBaseNodeContainer, CommonSkeletalMeshesAndAnimationsProperties->bConvertStaticsWithMorphTargetsToSkeletals))
+		if(UE::Interchange::Private::IsTranslatedDataContainOnlyJointAnimation(ContextParams.BaseNodeContainer, CommonSkeletalMeshesAndAnimationsProperties->bConvertStaticsWithMorphTargetsToSkeletals))
 		{
 			bImportAnimations = true;
 			CommonSkeletalMeshesAndAnimationsProperties->bImportOnlyAnimations = true;
@@ -129,9 +131,9 @@ void UInterchangeGenericAnimationPipeline::AdjustSettingsForContext(EInterchange
 	}
 
 	
-	if (ImportType == EInterchangePipelineContext::AssetReimport)
+	if (ContextParams.ContextType == EInterchangePipelineContext::AssetReimport)
 	{
-		if (UAnimSequence* AnimSequence = Cast<UAnimSequence>(ReimportAsset))
+		if (UAnimSequence* AnimSequence = Cast<UAnimSequence>(ContextParams.ReimportAsset))
 		{
 			//Set the skeleton to the current asset skeleton and re-import only the animation
 			CommonSkeletalMeshesAndAnimationsProperties->Skeleton = AnimSequence->GetSkeleton();
diff --git a/Engine/Plugins/Interchange/Runtime/Source/Pipelines/Private/InterchangeGenericAssetsPipeline.cpp b/Engine/Plugins/Interchange/Runtime/Source/Pipelines/Private/InterchangeGenericAssetsPipeline.cpp
--- a/Engine/Plugins/Interchange/Runtime/Source/Pipelines/Private/InterchangeGenericAssetsPipeline.cpp
+++ b/Engine/Plugins/Interchange/Runtime/Source/Pipelines/Private/InterchangeGenericAssetsPipeline.cpp
@@ -139,11 +139,11 @@ bool UInterchangeGenericAssetsPipeline::IsSettingsAreValid(TOptional<FText>& Out
 	return Super::IsSettingsAreValid(OutInvalidReason);
 }
 
-void UInterchangeGenericAssetsPipeline::AdjustSettingsForContext(EInterchangePipelineContext ImportType, TObjectPtr<UObject> ReimportAsset, const UInterchangeBaseNodeContainer* InBaseNodeContainer)
+void UInterchangeGenericAssetsPipeline::AdjustSettingsForContext(const FInterchangePipelineContextParams& ContextParams)
 {
-	Super::AdjustSettingsForContext(ImportType, ReimportAsset, InBaseNodeContainer);
+	Super::AdjustSettingsForContext(ContextParams);
 
-	if (ImportType == EInterchangePipelineContext::AssetImport)
+	if (ContextParams.ContextType == EInterchangePipelineContext::AssetImport)
 	{
 		if (!CommonSkeletalMeshesAndAnimationsProperties->Skeleton.IsValid())
 		{
@@ -162,25 +162,25 @@ void UInterchangeGenericAssetsPipeline::AdjustSettingsForContext(EInterchangePip
 			}
 		}
 	}
-	else if (ImportType == EInterchangePipelineContext::AssetCustomMorphTargetImport
-		|| ImportType == EInterchangePipelineContext::AssetCustomMorphTargetReImport)
+	else if (ContextParams.ContextType == EInterchangePipelineContext::AssetCustomMorphTargetImport
+		|| ContextParams.ContextType == EInterchangePipelineContext::AssetCustomMorphTargetReImport)
 	{
 		bUseSourceNameForAsset = false;
 	}
 
 	if (MaterialPipeline)
 	{
-		MaterialPipeline->AdjustSettingsForContext(ImportType, ReimportAsset, InBaseNodeContainer);
+		MaterialPipeline->AdjustSettingsForContext(ContextParams);
 	}
 
 	if (MeshPipeline)
 	{
-		MeshPipeline->AdjustSettingsForContext(ImportType, ReimportAsset, InBaseNodeContainer);
+		MeshPipeline->AdjustSettingsForContext(ContextParams);
 	}
 
 	if (AnimationPipeline)
 	{
-		AnimationPipeline->AdjustSettingsForContext(ImportType, ReimportAsset, InBaseNodeContainer);
+		AnimationPipeline->AdjustSettingsForContext(ContextParams);
 	}
 
 	//Store the adjusted settings
diff --git a/Engine/Plugins/Interchange/Runtime/Source/Pipelines/Private/InterchangeGenericMaterialPipeline.cpp b/Engine/Plugins/Interchange/Runtime/Source/Pipelines/Private/InterchangeGenericMaterialPipeline.cpp
--- a/Engine/Plugins/Interchange/Runtime/Source/Pipelines/Private/InterchangeGenericMaterialPipeline.cpp
+++ b/Engine/Plugins/Interchange/Runtime/Source/Pipelines/Private/InterchangeGenericMaterialPipeline.cpp
@@ -612,21 +612,23 @@ bool UInterchangeGenericMaterialPipeline::IsSettingsAreValid(TOptional<FText>& O
 	return Super::IsSettingsAreValid(OutInvalidReason);
 }
 
-void UInterchangeGenericMaterialPipeline::AdjustSettingsForContext(EInterchangePipelineContext ImportType, TObjectPtr<UObject> ReimportAsset, const UInterchangeBaseNodeContainer* InBaseNodeContainer)
+void UInterchangeGenericMaterialPipeline::AdjustSettingsForContext(const FInterchangePipelineContextParams& ContextParams)
 {
-	Super::AdjustSettingsForContext(ImportType, ReimportAsset, InBaseNodeContainer);
+	Super::AdjustSettingsForContext(ContextParams);
 
 	if (TexturePipeline)
 	{
-		TexturePipeline->AdjustSettingsForContext(ImportType, ReimportAsset, InBaseNodeContainer);
+		TexturePipeline->AdjustSettingsForContext(ContextParams);
 	}
 #if WITH_EDITOR
 	TArray<FString> HideCategories;
-	bool bIsObjectAMaterial = !ReimportAsset ? false : ReimportAsset->IsA(UMaterialInterface::StaticClass());
-	if (ImportType == EInterchangePipelineContext::AssetCustomLODImport
-		|| ImportType == EInterchangePipelineContext::AssetCustomLODReimport
-		|| ImportType == EInterchangePipelineContext::AssetAlternateSkinningImport
-		|| ImportType == EInterchangePipelineContext::AssetAlternateSkinningReimport)
+	bool bIsObjectAMaterial = !ContextParams.ReimportAsset ? false : ContextParams.ReimportAsset->IsA(UMaterialInterface::StaticClass());
+	if (ContextParams.ContextType == EInterchangePipelineContext::AssetCustomLODImport
+		|| ContextParams.ContextType == EInterchangePipelineContext::AssetCustomLODReimport
+		|| ContextParams.ContextType == EInterchangePipelineContext::AssetAlternateSkinningImport
+		|| ContextParams.ContextType == EInterchangePipelineContext::AssetAlternateSkinningReimport
+		|| ContextParams.ContextType == EInterchangePipelineContext::AssetCustomMorphTargetImport
+		|| ContextParams.ContextType == EInterchangePipelineContext::AssetCustomMorphTargetReImport)
 	{
 		bImportMaterials = false;
 		HideCategories.Add(UInterchangeGenericMaterialPipeline::GetPipelineCategory(nullptr));
@@ -639,7 +641,7 @@ void UInterchangeGenericMaterialPipeline::AdjustSettingsForContext(EInterchangeP
 		{
 			HidePropertiesOfCategory(OuterMostPipeline, this, HideCategoryName);
 		}
-		if (!bIsObjectAMaterial && ImportType == EInterchangePipelineContext::AssetReimport)
+		if (!bIsObjectAMaterial && ContextParams.ContextType == EInterchangePipelineContext::AssetReimport)
 		{
 			//When we re-import we hide all setting but search location, so we can find existing materials.
 			HideProperty(OuterMostPipeline, this, GET_MEMBER_NAME_CHECKED(UInterchangeGenericMaterialPipeline, bImportMaterials));
diff --git a/Engine/Plugins/Interchange/Runtime/Source/Pipelines/Private/InterchangeGenericMeshPipeline.cpp b/Engine/Plugins/Interchange/Runtime/Source/Pipelines/Private/InterchangeGenericMeshPipeline.cpp
--- a/Engine/Plugins/Interchange/Runtime/Source/Pipelines/Private/InterchangeGenericMeshPipeline.cpp
+++ b/Engine/Plugins/Interchange/Runtime/Source/Pipelines/Private/InterchangeGenericMeshPipeline.cpp
@@ -37,14 +37,14 @@ FString UInterchangeGenericMeshPipeline::GetPipelineCategory(UClass* AssetClass)
 	return TEXT("Static Meshes");
 }
 
-void UInterchangeGenericMeshPipeline::AdjustSettingsForContext(EInterchangePipelineContext ImportType, TObjectPtr<UObject> ReimportAsset, const UInterchangeBaseNodeContainer* InBaseNodeContainer)
+void UInterchangeGenericMeshPipeline::AdjustSettingsForContext(const FInterchangePipelineContextParams& ContextParams)
 {
-	Super::AdjustSettingsForContext(ImportType, ReimportAsset, InBaseNodeContainer);
+	Super::AdjustSettingsForContext(ContextParams);
 
 #if WITH_EDITOR
 
 	check(CommonSkeletalMeshesAndAnimationsProperties.IsValid());
-	if (ImportType == EInterchangePipelineContext::None)
+	if (ContextParams.ContextType == EInterchangePipelineContext::None)
 	{
 		//We do not change the setting if we are in editing context
 		return;
@@ -55,22 +55,22 @@ void UInterchangeGenericMeshPipeline::AdjustSettingsForContext(EInterchangePipel
 	bool bContainSkeletalMesh = false;
 	bool bContainStaticMeshAnimationNode = false;
 	bool bIgnoreStaticMesh = false;
-	GetMeshesInformationFromTranslatedData(InBaseNodeContainer, bAutoDetectConvertStaticMeshToSkeletalMesh, bContainStaticMesh, bContainSkeletalMesh, bContainStaticMeshAnimationNode, bIgnoreStaticMesh);
+	GetMeshesInformationFromTranslatedData(ContextParams.BaseNodeContainer, bAutoDetectConvertStaticMeshToSkeletalMesh, bContainStaticMesh, bContainSkeletalMesh, bContainStaticMeshAnimationNode, bIgnoreStaticMesh);
 
 	//Avoid creating physics asset when importing a LOD or the alternate skinning
-	if (ImportType == EInterchangePipelineContext::AssetCustomLODImport
-		|| ImportType == EInterchangePipelineContext::AssetCustomLODReimport
-		|| ImportType == EInterchangePipelineContext::AssetAlternateSkinningImport
-		|| ImportType == EInterchangePipelineContext::AssetAlternateSkinningReimport
-		|| ImportType == EInterchangePipelineContext::AssetCustomMorphTargetImport
-		|| ImportType == EInterchangePipelineContext::AssetCustomMorphTargetReImport)
+	if (ContextParams.ContextType == EInterchangePipelineContext::AssetCustomLODImport
+		|| ContextParams.ContextType == EInterchangePipelineContext::AssetCustomLODReimport
+		|| ContextParams.ContextType == EInterchangePipelineContext::AssetAlternateSkinningImport
+		|| ContextParams.ContextType == EInterchangePipelineContext::AssetAlternateSkinningReimport
+		|| ContextParams.ContextType == EInterchangePipelineContext::AssetCustomMorphTargetImport
+		|| ContextParams.ContextType == EInterchangePipelineContext::AssetCustomMorphTargetReImport)
 	{
 		bCreatePhysicsAsset = false;
 		PhysicsAsset = nullptr;
 		LodGroup = NAME_None;
 		
-		if (ImportType == EInterchangePipelineContext::AssetAlternateSkinningImport
-			|| ImportType == EInterchangePipelineContext::AssetAlternateSkinningReimport)
+		if (ContextParams.ContextType == EInterchangePipelineContext::AssetAlternateSkinningImport
+			|| ContextParams.ContextType == EInterchangePipelineContext::AssetAlternateSkinningReimport)
 		{
 			CommonMeshesProperties->ForceAllMeshAsType = EInterchangeForceMeshType::IFMT_SkeletalMesh;
 			CommonMeshesProperties->bAutoDetectMeshType = false;
@@ -89,8 +89,8 @@ void UInterchangeGenericMeshPipeline::AdjustSettingsForContext(EInterchangePipel
 			CommonSkeletalMeshesAndAnimationsProperties->Skeleton = nullptr;
 			CommonSkeletalMeshesAndAnimationsProperties->bImportOnlyAnimations = false;
 		}
-		else if (ImportType == EInterchangePipelineContext::AssetCustomMorphTargetImport
-			|| ImportType == EInterchangePipelineContext::AssetCustomMorphTargetReImport)
+		else if (ContextParams.ContextType == EInterchangePipelineContext::AssetCustomMorphTargetImport
+			|| ContextParams.ContextType == EInterchangePipelineContext::AssetCustomMorphTargetReImport)
 		{
 			//Custom morph target are imported has a combined static mesh
 			CommonMeshesProperties->ForceAllMeshAsType = EInterchangeForceMeshType::IFMT_StaticMesh;
@@ -111,6 +111,42 @@ void UInterchangeGenericMeshPipeline::AdjustSettingsForContext(EInterchangePipel
 			bGenerateDistanceFieldAsIfTwoSided = false;
 			bSupportFaceRemap = false;
 		}
+		else if (ContextParams.ContextType == EInterchangePipelineContext::AssetCustomLODImport
+			|| ContextParams.ContextType == EInterchangePipelineContext::AssetCustomLODReimport)
+		{
+			CommonMeshesProperties->bBakeMeshes = true;
+			CommonMeshesProperties->bBakePivotMeshes = false;
+			CommonMeshesProperties->bImportLods = false;
+			CommonMeshesProperties->bKeepSectionsSeparate = false;
+			CommonMeshesProperties->VertexColorImportOption = EInterchangeVertexColorImportOption::IVCIO_Ignore;
+			bBuildNanite = false;
+			LodGroup = NAME_None;
+			bImportCollision = false;
+			bImportCollisionAccordingToMeshName = false;
+			bGenerateLightmapUVs = false;
+			bGenerateDistanceFieldAsIfTwoSided = false;
+			bSupportFaceRemap = false;
+			//We are importing custom LODs
+			if (ContextParams.ImportObjectType)
+			{
+				//If we have a provided import object type we can make sure we import the correct type
+				if (ContextParams.ImportObjectType->IsChildOf<UStaticMesh>())
+				{
+					CommonMeshesProperties->bAutoDetectMeshType = false;
+					CommonMeshesProperties->ForceAllMeshAsType = EInterchangeForceMeshType::IFMT_StaticMesh;
+					bImportSkeletalMeshes = false;
+					bImportStaticMeshes = true;
+					bCombineStaticMeshes = true;
+				}
+				else if (ContextParams.ImportObjectType->IsChildOf<USkeletalMesh>())
+				{
+					CommonMeshesProperties->bAutoDetectMeshType = false;
+					CommonMeshesProperties->ForceAllMeshAsType = EInterchangeForceMeshType::IFMT_SkeletalMesh;
+					bImportSkeletalMeshes = true;
+					bImportStaticMeshes = false;
+				}
+			}
+		}
 	}
 	const FString CommonMeshesCategory = UInterchangeGenericCommonMeshesProperties::GetPipelineCategory(nullptr);
 	const FString StaticMeshesCategory = UInterchangeGenericMeshPipeline::GetPipelineCategory(UStaticMesh::StaticClass());
@@ -119,12 +155,12 @@ void UInterchangeGenericMeshPipeline::AdjustSettingsForContext(EInterchangePipel
 
 	TArray<FString> HideCategories;
 	TArray<FString> HideSubCategories;
-	if (ImportType == EInterchangePipelineContext::AssetReimport)
+	if (ContextParams.ContextType == EInterchangePipelineContext::AssetReimport)
 	{
 		CommonMeshesProperties->bAutoDetectMeshType = false;
 
 		HideSubCategories.Add(TEXT("Build"));
-		if (USkeletalMesh* SkeletalMesh = Cast<USkeletalMesh>(ReimportAsset))
+		if (USkeletalMesh* SkeletalMesh = Cast<USkeletalMesh>(ContextParams.ReimportAsset))
 		{
 			//Set the skeleton to the current asset skeleton
 			CommonSkeletalMeshesAndAnimationsProperties->Skeleton = SkeletalMesh->GetSkeleton();
@@ -140,7 +176,7 @@ void UInterchangeGenericMeshPipeline::AdjustSettingsForContext(EInterchangePipel
 				CommonMeshesProperties->ForceAllMeshAsType = EInterchangeForceMeshType::IFMT_SkeletalMesh;
 			}
 		}
-		else if (UStaticMesh* StaticMesh = Cast<UStaticMesh>(ReimportAsset))
+		else if (UStaticMesh* StaticMesh = Cast<UStaticMesh>(ContextParams.ReimportAsset))
 		{
 			HideCategories.Add(SkeletalMeshesCategory);
 			HideCategories.Add(CommonSkeletalMeshesAndAnimationCategory);
@@ -150,13 +186,13 @@ void UInterchangeGenericMeshPipeline::AdjustSettingsForContext(EInterchangePipel
 				CommonMeshesProperties->ForceAllMeshAsType = EInterchangeForceMeshType::IFMT_StaticMesh;
 			}
 		}
-		else if (UAnimSequence* AnimSequence = Cast<UAnimSequence>(ReimportAsset))
+		else if (UAnimSequence* AnimSequence = Cast<UAnimSequence>(ContextParams.ReimportAsset))
 		{
 			HideCategories.Add(StaticMeshesCategory);
 			HideCategories.Add(SkeletalMeshesCategory);
 			HideCategories.Add(CommonMeshesCategory);
 		}
-		else if (ReimportAsset)
+		else if (ContextParams.ReimportAsset)
 		{
 			HideCategories.Add(StaticMeshesCategory);
 			HideCategories.Add(SkeletalMeshesCategory);
diff --git a/Engine/Plugins/Interchange/Runtime/Source/Pipelines/Private/InterchangeGenericScenesPipeline.cpp b/Engine/Plugins/Interchange/Runtime/Source/Pipelines/Private/InterchangeGenericScenesPipeline.cpp
--- a/Engine/Plugins/Interchange/Runtime/Source/Pipelines/Private/InterchangeGenericScenesPipeline.cpp
+++ b/Engine/Plugins/Interchange/Runtime/Source/Pipelines/Private/InterchangeGenericScenesPipeline.cpp
@@ -218,11 +218,11 @@ namespace UE::Interchange::Private
 	}
 }
 
-void UInterchangeGenericLevelPipeline::AdjustSettingsForContext(EInterchangePipelineContext ImportType, TObjectPtr<UObject> ReimportAsset, const UInterchangeBaseNodeContainer* InBaseNodeContainer)
+void UInterchangeGenericLevelPipeline::AdjustSettingsForContext(const FInterchangePipelineContextParams& ContextParams)
 {
-	Super::AdjustSettingsForContext(ImportType, ReimportAsset, InBaseNodeContainer);
+	Super::AdjustSettingsForContext(ContextParams);
 
-	bIsReimportContext = ReimportAsset != nullptr;
+	bIsReimportContext = ContextParams.ReimportAsset != nullptr;
 }
 
 void UInterchangeGenericLevelPipeline::ExecutePipeline(UInterchangeBaseNodeContainer* InBaseNodeContainer, const TArray<UInterchangeSourceData*>& InSourceDatas, const FString& ContentBasePath)
diff --git a/Engine/Plugins/Interchange/Runtime/Source/Pipelines/Private/InterchangeGenericTexturePipeline.cpp b/Engine/Plugins/Interchange/Runtime/Source/Pipelines/Private/InterchangeGenericTexturePipeline.cpp
--- a/Engine/Plugins/Interchange/Runtime/Source/Pipelines/Private/InterchangeGenericTexturePipeline.cpp
+++ b/Engine/Plugins/Interchange/Runtime/Source/Pipelines/Private/InterchangeGenericTexturePipeline.cpp
@@ -115,17 +115,19 @@ FString UInterchangeGenericTexturePipeline::GetPipelineCategory(UClass* AssetCla
 	return TEXT("Textures");
 }
 
-void UInterchangeGenericTexturePipeline::AdjustSettingsForContext(EInterchangePipelineContext ImportType, TObjectPtr<UObject> ReimportAsset, const UInterchangeBaseNodeContainer* InBaseNodeContainer)
+void UInterchangeGenericTexturePipeline::AdjustSettingsForContext(const FInterchangePipelineContextParams& ContextParams)
 {
-	Super::AdjustSettingsForContext(ImportType, ReimportAsset, InBaseNodeContainer);
+	Super::AdjustSettingsForContext(ContextParams);
 #if WITH_EDITOR
 	TArray<FString> HideCategories;
-	bool bIsObjectATexture = !ReimportAsset ? false : ReimportAsset.IsA(UTexture::StaticClass());
-	if( (!bIsObjectATexture && ImportType == EInterchangePipelineContext::AssetReimport)
-		|| ImportType == EInterchangePipelineContext::AssetCustomLODImport
-		|| ImportType == EInterchangePipelineContext::AssetCustomLODReimport
-		|| ImportType == EInterchangePipelineContext::AssetAlternateSkinningImport
-		|| ImportType == EInterchangePipelineContext::AssetAlternateSkinningReimport)
+	bool bIsObjectATexture = !ContextParams.ReimportAsset ? false : ContextParams.ReimportAsset.IsA(UTexture::StaticClass());
+	if( (!bIsObjectATexture && ContextParams.ContextType == EInterchangePipelineContext::AssetReimport)
+		|| ContextParams.ContextType == EInterchangePipelineContext::AssetCustomLODImport
+		|| ContextParams.ContextType == EInterchangePipelineContext::AssetCustomLODReimport
+		|| ContextParams.ContextType == EInterchangePipelineContext::AssetAlternateSkinningImport
+		|| ContextParams.ContextType == EInterchangePipelineContext::AssetAlternateSkinningReimport
+		|| ContextParams.ContextType == EInterchangePipelineContext::AssetCustomMorphTargetImport
+		|| ContextParams.ContextType == EInterchangePipelineContext::AssetCustomMorphTargetReImport)
 	{
 		bImportTextures = false;
 		HideCategories.Add(UInterchangeGenericTexturePipeline::GetPipelineCategory(nullptr));
diff --git a/Engine/Plugins/Interchange/Runtime/Source/Pipelines/Private/InterchangeMaterialXPipeline.cpp b/Engine/Plugins/Interchange/Runtime/Source/Pipelines/Private/InterchangeMaterialXPipeline.cpp
--- a/Engine/Plugins/Interchange/Runtime/Source/Pipelines/Private/InterchangeMaterialXPipeline.cpp
+++ b/Engine/Plugins/Interchange/Runtime/Source/Pipelines/Private/InterchangeMaterialXPipeline.cpp
@@ -687,9 +687,9 @@ UInterchangeMaterialXPipeline::UInterchangeMaterialXPipeline()
 	}
 }
 
-void UInterchangeMaterialXPipeline::AdjustSettingsForContext(EInterchangePipelineContext ImportType, TObjectPtr<UObject> ReimportAsset, const UInterchangeBaseNodeContainer* InBaseNodeContainer)
+void UInterchangeMaterialXPipeline::AdjustSettingsForContext(const FInterchangePipelineContextParams& ContextParams)
 {
-	Super::AdjustSettingsForContext(ImportType, ReimportAsset, InBaseNodeContainer);
+	Super::AdjustSettingsForContext(ContextParams);
 
 	if (!MaterialXSettings->AreRequiredPackagesLoaded())
 	{
diff --git a/Engine/Plugins/Interchange/Runtime/Source/Pipelines/Private/InterchangeglTFPipeline.cpp b/Engine/Plugins/Interchange/Runtime/Source/Pipelines/Private/InterchangeglTFPipeline.cpp
--- a/Engine/Plugins/Interchange/Runtime/Source/Pipelines/Private/InterchangeglTFPipeline.cpp
+++ b/Engine/Plugins/Interchange/Runtime/Source/Pipelines/Private/InterchangeglTFPipeline.cpp
@@ -287,9 +287,9 @@ UInterchangeGLTFPipeline::UInterchangeGLTFPipeline()
 {
 }
 
-void UInterchangeGLTFPipeline::AdjustSettingsForContext(EInterchangePipelineContext ImportType, TObjectPtr<UObject> ReimportAsset, const UInterchangeBaseNodeContainer* InBaseNodeContainer)
+void UInterchangeGLTFPipeline::AdjustSettingsForContext(const FInterchangePipelineContextParams& ContextParams)
 {
-	Super::AdjustSettingsForContext(ImportType, ReimportAsset, InBaseNodeContainer);
+	Super::AdjustSettingsForContext(ContextParams);
 
 	TArray<FString> MaterialInstanceIssues = GLTFPipelineSettings->ValidateMaterialInstancesAndParameters();
 	for (const FString& MaterialInstanceIssue : MaterialInstanceIssues)
diff --git a/Engine/Plugins/Interchange/Runtime/Source/Pipelines/Public/InterchangeGenericAnimationPipeline.h b/Engine/Plugins/Interchange/Runtime/Source/Pipelines/Public/InterchangeGenericAnimationPipeline.h
--- a/Engine/Plugins/Interchange/Runtime/Source/Pipelines/Public/InterchangeGenericAnimationPipeline.h
+++ b/Engine/Plugins/Interchange/Runtime/Source/Pipelines/Public/InterchangeGenericAnimationPipeline.h
@@ -120,7 +120,7 @@ public:
 	UPROPERTY()
 	FString SourceAnimationName;
 
-	virtual void AdjustSettingsForContext(EInterchangePipelineContext ImportType, TObjectPtr<UObject> ReimportAsset, const UInterchangeBaseNodeContainer* InBaseNodeContainer) override;
+	virtual void AdjustSettingsForContext(const FInterchangePipelineContextParams& ContextParams) override;
 	
 protected:
 
diff --git a/Engine/Plugins/Interchange/Runtime/Source/Pipelines/Public/InterchangeGenericAssetsPipeline.h b/Engine/Plugins/Interchange/Runtime/Source/Pipelines/Public/InterchangeGenericAssetsPipeline.h
--- a/Engine/Plugins/Interchange/Runtime/Source/Pipelines/Public/InterchangeGenericAssetsPipeline.h
+++ b/Engine/Plugins/Interchange/Runtime/Source/Pipelines/Public/InterchangeGenericAssetsPipeline.h
@@ -98,7 +98,7 @@ public:
 	virtual bool IsSettingsAreValid(TOptional<FText>& OutInvalidReason) const override;
 
 
-	virtual void AdjustSettingsForContext(EInterchangePipelineContext ImportType, TObjectPtr<UObject> ReimportAsset, const UInterchangeBaseNodeContainer* InBaseNodeContainer) override;
+	virtual void AdjustSettingsForContext(const FInterchangePipelineContextParams& ContextParams) override;
 #if WITH_EDITOR
 	virtual void FilterPropertiesFromTranslatedData(UInterchangeBaseNodeContainer* InBaseNodeContainer) override;
 
diff --git a/Engine/Plugins/Interchange/Runtime/Source/Pipelines/Public/InterchangeGenericMaterialPipeline.h b/Engine/Plugins/Interchange/Runtime/Source/Pipelines/Public/InterchangeGenericMaterialPipeline.h
--- a/Engine/Plugins/Interchange/Runtime/Source/Pipelines/Public/InterchangeGenericMaterialPipeline.h
+++ b/Engine/Plugins/Interchange/Runtime/Source/Pipelines/Public/InterchangeGenericMaterialPipeline.h
@@ -101,7 +101,7 @@ public:
 	/** BEGIN UInterchangePipelineBase overrides */
 	virtual void PreDialogCleanup(const FName PipelineStackName) override;
 	virtual bool IsSettingsAreValid(TOptional<FText>& OutInvalidReason) const override;
-	virtual void AdjustSettingsForContext(EInterchangePipelineContext ImportType, TObjectPtr<UObject> ReimportAsset, const UInterchangeBaseNodeContainer* InBaseNodeContainer) override;
+	virtual void AdjustSettingsForContext(const FInterchangePipelineContextParams& ContextParams) override;
 #if WITH_EDITOR
 	virtual void FilterPropertiesFromTranslatedData(UInterchangeBaseNodeContainer* InBaseNodeContainer) override;
 	virtual bool IsPropertyChangeNeedRefresh(const FPropertyChangedEvent& PropertyChangedEvent) override;
diff --git a/Engine/Plugins/Interchange/Runtime/Source/Pipelines/Public/InterchangeGenericMeshPipeline.h b/Engine/Plugins/Interchange/Runtime/Source/Pipelines/Public/InterchangeGenericMeshPipeline.h
--- a/Engine/Plugins/Interchange/Runtime/Source/Pipelines/Public/InterchangeGenericMeshPipeline.h
+++ b/Engine/Plugins/Interchange/Runtime/Source/Pipelines/Public/InterchangeGenericMeshPipeline.h
@@ -225,7 +225,7 @@ public:
 	UPROPERTY(EditAnywhere, BlueprintReadWrite, Category = "Skeletal Meshes", meta = (SubCategory = "Build"))
 	int32 BoneInfluenceLimit = 0;
 
-	virtual void AdjustSettingsForContext(EInterchangePipelineContext ImportType, TObjectPtr<UObject> ReimportAsset, const UInterchangeBaseNodeContainer* InBaseNodeContainer) override;
+	virtual void AdjustSettingsForContext(const FInterchangePipelineContextParams& ContextParams) override;
 
 	virtual void PreDialogCleanup(const FName PipelineStackName) override;
 
diff --git a/Engine/Plugins/Interchange/Runtime/Source/Pipelines/Public/InterchangeGenericScenesPipeline.h b/Engine/Plugins/Interchange/Runtime/Source/Pipelines/Public/InterchangeGenericScenesPipeline.h
--- a/Engine/Plugins/Interchange/Runtime/Source/Pipelines/Public/InterchangeGenericScenesPipeline.h
+++ b/Engine/Plugins/Interchange/Runtime/Source/Pipelines/Public/InterchangeGenericScenesPipeline.h
@@ -46,7 +46,7 @@ public:
 	bool bDeleteMissingAssets = false;
 
 	/** BEGIN UInterchangePipelineBase overrides */
-	virtual void AdjustSettingsForContext(EInterchangePipelineContext ImportType, TObjectPtr<UObject> ReimportAsset, const UInterchangeBaseNodeContainer* InBaseNodeContainer) override;
+	virtual void AdjustSettingsForContext(const FInterchangePipelineContextParams& ContextParams) override;
 
 protected:
 
diff --git a/Engine/Plugins/Interchange/Runtime/Source/Pipelines/Public/InterchangeGenericTexturePipeline.h b/Engine/Plugins/Interchange/Runtime/Source/Pipelines/Public/InterchangeGenericTexturePipeline.h
--- a/Engine/Plugins/Interchange/Runtime/Source/Pipelines/Public/InterchangeGenericTexturePipeline.h
+++ b/Engine/Plugins/Interchange/Runtime/Source/Pipelines/Public/InterchangeGenericTexturePipeline.h
@@ -68,7 +68,7 @@ public:
 	bool bAllowNonPowerOfTwo = false;
 
 public:
-	virtual void AdjustSettingsForContext(EInterchangePipelineContext ImportType, TObjectPtr<UObject> ReimportAsset, const UInterchangeBaseNodeContainer* InBaseNodeContainer) override;
+	virtual void AdjustSettingsForContext(const FInterchangePipelineContextParams& ContextParams) override;
 	virtual void ExecutePipeline(UInterchangeBaseNodeContainer* InBaseNodeContainer, const TArray<UInterchangeSourceData*>& InSourceDatas, const FString& ContentBasePath) override;
 	virtual void ExecutePostFactoryPipeline(const UInterchangeBaseNodeContainer* BaseNodeContainer, const FString& NodeKey, UObject* CreatedAsset, bool bIsAReimport) override;
 
diff --git a/Engine/Plugins/Interchange/Runtime/Source/Pipelines/Public/InterchangeMaterialXPipeline.h b/Engine/Plugins/Interchange/Runtime/Source/Pipelines/Public/InterchangeMaterialXPipeline.h
--- a/Engine/Plugins/Interchange/Runtime/Source/Pipelines/Public/InterchangeMaterialXPipeline.h
+++ b/Engine/Plugins/Interchange/Runtime/Source/Pipelines/Public/InterchangeMaterialXPipeline.h
@@ -101,7 +101,7 @@ public:
 	TObjectPtr<UMaterialXPipelineSettings> MaterialXSettings;
 
 protected:
-	virtual void AdjustSettingsForContext(EInterchangePipelineContext ImportType, TObjectPtr<UObject> ReimportAsset, const UInterchangeBaseNodeContainer* InBaseNodeContainer) override;
+	virtual void AdjustSettingsForContext(const FInterchangePipelineContextParams& ContextParams) override;
 	virtual void ExecutePipeline(UInterchangeBaseNodeContainer* BaseNodeContainer, const TArray<UInterchangeSourceData*>& SourceDatas, const FString& ContentBasePath) override;
 
 	virtual bool CanExecuteOnAnyThread(EInterchangePipelineTask PipelineTask) override
diff --git a/Engine/Plugins/Interchange/Runtime/Source/Pipelines/Public/InterchangeglTFPipeline.h b/Engine/Plugins/Interchange/Runtime/Source/Pipelines/Public/InterchangeglTFPipeline.h
--- a/Engine/Plugins/Interchange/Runtime/Source/Pipelines/Public/InterchangeglTFPipeline.h
+++ b/Engine/Plugins/Interchange/Runtime/Source/Pipelines/Public/InterchangeglTFPipeline.h
@@ -59,7 +59,7 @@ public:
 	TObjectPtr<UGLTFPipelineSettings> GLTFPipelineSettings;
 
 protected:
-	virtual void AdjustSettingsForContext(EInterchangePipelineContext ImportType, TObjectPtr<UObject> ReimportAsset, const UInterchangeBaseNodeContainer* InBaseNodeContainer) override;
+	virtual void AdjustSettingsForContext(const FInterchangePipelineContextParams& ContextParams) override;
 	virtual void ExecutePipeline(UInterchangeBaseNodeContainer* BaseNodeContainer, const TArray<UInterchangeSourceData*>& SourceDatas, const FString& ContentBasePath) override;
 
 	virtual bool CanExecuteOnAnyThread(EInterchangePipelineTask PipelineTask) override
diff --git a/Engine/Source/Editor/UnrealEd/Private/SkinWeightsUtilities.cpp b/Engine/Source/Editor/UnrealEd/Private/SkinWeightsUtilities.cpp
--- a/Engine/Source/Editor/UnrealEd/Private/SkinWeightsUtilities.cpp
+++ b/Engine/Source/Editor/UnrealEd/Private/SkinWeightsUtilities.cpp
@@ -152,7 +152,10 @@ bool FSkinWeightsUtilities::ImportAlternateSkinWeight(USkeletalMesh* SkeletalMes
 			if (GenericPipeline)
 			{
 				GenericPipeline->ClearFlags(EObjectFlags::RF_Standalone | EObjectFlags::RF_Public);
-				GenericPipeline->AdjustSettingsForContext(bIsReimport ? EInterchangePipelineContext::AssetAlternateSkinningReimport : EInterchangePipelineContext::AssetAlternateSkinningImport, nullptr, nullptr);
+				FInterchangePipelineContextParams ContextParams;
+				ContextParams.ContextType = bIsReimport ? EInterchangePipelineContext::AssetAlternateSkinningReimport : EInterchangePipelineContext::AssetAlternateSkinningImport;
+				ContextParams.ImportObjectType = USkeletalMesh::StaticClass();
+				GenericPipeline->AdjustSettingsForContext(ContextParams);
 				ImportAssetParameters.OverridePipelines.Add(GenericPipeline);
 			}
 		}
diff --git a/Engine/Source/Runtime/Interchange/Core/Private/InterchangePipelineBase.cpp b/Engine/Source/Runtime/Interchange/Core/Private/InterchangePipelineBase.cpp
--- a/Engine/Source/Runtime/Interchange/Core/Private/InterchangePipelineBase.cpp
+++ b/Engine/Source/Runtime/Interchange/Core/Private/InterchangePipelineBase.cpp
@@ -37,15 +37,14 @@ void UInterchangePipelineBase::SaveSettings(const FName PipelineStackName)
 	GConfig->Flush(0);
 }
 
-void UInterchangePipelineBase::AdjustSettingsForContext(EInterchangePipelineContext ReimportType, TObjectPtr<UObject> ReimportAsset, const UInterchangeBaseNodeContainer* BaseNodeContainer)
+void UInterchangePipelineBase::AdjustSettingsForContext(const FInterchangePipelineContextParams& ContextParams)
 {
-	CachePipelineContext = ReimportType;
-	CacheReimportObject = ReimportAsset;
+	CacheContextParam = ContextParams;
 	CachePropertiesStates = PropertiesStates;
 
-	bAllowPropertyStatesEdition = (ReimportType == EInterchangePipelineContext::None);
+	bAllowPropertyStatesEdition = (ContextParams.ContextType == EInterchangePipelineContext::None);
 	bIsReimportContext = false;
-	switch (ReimportType)
+	switch (ContextParams.ContextType)
 	{
 	case EInterchangePipelineContext::AssetReimport:
 	case EInterchangePipelineContext::AssetAlternateSkinningReimport:
@@ -61,13 +60,12 @@ void UInterchangePipelineBase::AdjustSettingsForContext(EInterchangePipelineCont
 void UInterchangePipelineBase::AdjustSettingsFromCache()
 {
 	PropertiesStates = CachePropertiesStates;
-	AdjustSettingsForContext(CachePipelineContext, CacheReimportObject.Get(), nullptr);
+	AdjustSettingsForContext(CacheContextParam);
 }
 
 void UInterchangePipelineBase::TransferAdjustSettings(UInterchangePipelineBase* SourcePipeline)
 {
-	CachePipelineContext = SourcePipeline->CachePipelineContext;
-	CacheReimportObject = SourcePipeline->CacheReimportObject;
+	CacheContextParam = SourcePipeline->CacheContextParam;
 	CachePropertiesStates = SourcePipeline->CachePropertiesStates;
 	bAllowPropertyStatesEdition = SourcePipeline->bAllowPropertyStatesEdition;
 	bIsReimportContext = SourcePipeline->bIsReimportContext;
diff --git a/Engine/Source/Runtime/Interchange/Core/Public/InterchangePipelineBase.h b/Engine/Source/Runtime/Interchange/Core/Public/InterchangePipelineBase.h
--- a/Engine/Source/Runtime/Interchange/Core/Public/InterchangePipelineBase.h
+++ b/Engine/Source/Runtime/Interchange/Core/Public/InterchangePipelineBase.h
@@ -55,6 +55,18 @@ enum class EInterchangePipelineContext : uint8
 	AssetCustomMorphTargetReImport,
 };
 
+USTRUCT()
+struct FInterchangePipelineContextParams
+{
+	GENERATED_BODY()
+
+	EInterchangePipelineContext ContextType = EInterchangePipelineContext::None;
+	UClass* ImportObjectType = nullptr;
+	TObjectPtr<UObject> ReimportAsset = nullptr;
+	const UInterchangeBaseNodeContainer* BaseNodeContainer = nullptr;
+
+};
+
 USTRUCT(BlueprintType)
 struct FInterchangePipelinePropertyStatePerContext
 {
@@ -314,10 +326,9 @@ public:
 	 * The function is also called when we import or reimport custom LOD and alternate skinning.
 	 *
 	 * @Note - The function will set the context of the pipeline.
-	 * @Param ReimportType - Tells the pipeline what reimport type the user wants to achieve.
-	 * @Param ReimportAsset - This is an optional parameter which is set when reimporting an asset.
+	 * @Param ContextParams - Give all the context information to the pipeline for the current import.
 	 */
-	INTERCHANGECORE_API virtual void AdjustSettingsForContext(EInterchangePipelineContext ReimportType, TObjectPtr<UObject> ReimportAsset, const UInterchangeBaseNodeContainer* BaseNodeContainer);
+	INTERCHANGECORE_API virtual void AdjustSettingsForContext(const FInterchangePipelineContextParams& ContextParams);
 	INTERCHANGECORE_API virtual void AdjustSettingsFromCache();
 
 	/** Transfer the source pipeline adjust settings to this pipeline. */
@@ -593,8 +604,7 @@ protected:
 	TMap<FName, FInterchangePipelinePropertyStates> PropertiesStates;
 
 	mutable TMap<FName, FInterchangePipelinePropertyStates> CachePropertiesStates;
-	mutable EInterchangePipelineContext CachePipelineContext = EInterchangePipelineContext::None;
-	mutable TWeakObjectPtr<UObject> CacheReimportObject = nullptr;
+	mutable FInterchangePipelineContextParams CacheContextParam;
 
 	TArray<FInterchangeConflictInfo> ConflictInfos;
 };
diff --git a/Engine/Source/Runtime/Interchange/Engine/Private/InterchangeManager.cpp b/Engine/Source/Runtime/Interchange/Engine/Private/InterchangeManager.cpp
--- a/Engine/Source/Runtime/Interchange/Engine/Private/InterchangeManager.cpp
+++ b/Engine/Source/Runtime/Interchange/Engine/Private/InterchangeManager.cpp
@@ -1131,7 +1131,9 @@ bool UInterchangeManager::ConvertImportData(const UObject* SourceImportData, FIm
 				{
 					if (UInterchangePipelineBase* GeneratedPipeline = UE::Interchange::GeneratePipelineInstance(PipelinePath))
 					{
-						GeneratedPipeline->AdjustSettingsForContext(EInterchangePipelineContext::AssetImport, nullptr, nullptr);
+						FInterchangePipelineContextParams ContextParams;
+						ContextParams.ContextType = EInterchangePipelineContext::AssetImport;
+						GeneratedPipeline->AdjustSettingsForContext(ContextParams);
 						if (GeneratedPipeline->IsA(ConvertedPipelineClass))
 						{
 							//We found a match, so we will use the default pipeline stacks
@@ -1841,7 +1843,11 @@ UInterchangeManager::ImportInternal(const FString& ContentPath, const UInterchan
 				Context = bImportScene ? EInterchangePipelineContext::SceneImport : EInterchangePipelineContext::AssetImport;
 			}
 			Pipeline->ContentImportPath = ContentBasePath;
-			Pipeline->AdjustSettingsForContext(Context, TaskData.ReimportObject, BaseNodeContainer);
+			FInterchangePipelineContextParams ContextParams;
+			ContextParams.ContextType = Context;
+			ContextParams.ReimportAsset = TaskData.ReimportObject;
+			ContextParams.BaseNodeContainer = BaseNodeContainer;
+			Pipeline->AdjustSettingsForContext(ContextParams);
 			Pipeline->DestinationName = TaskData.DestinationName;
 		};
 
diff --git a/Engine/Source/Runtime/Interchange/Engine/Private/InterchangeMeshUtilities.cpp b/Engine/Source/Runtime/Interchange/Engine/Private/InterchangeMeshUtilities.cpp
--- a/Engine/Source/Runtime/Interchange/Engine/Private/InterchangeMeshUtilities.cpp
+++ b/Engine/Source/Runtime/Interchange/Engine/Private/InterchangeMeshUtilities.cpp
@@ -187,6 +187,7 @@ TFuture<bool> UInterchangeMeshUtilities::InternalImportCustomLod(TSharedPtr<TPro
 	EInterchangePipelineContext ImportType = EInterchangePipelineContext::AssetCustomLODImport;
 	bool bInvalidLodIndex = false;
 	UObject* SourceImportData = nullptr;
+	UClass* ObjectType = nullptr;
 	if (SkeletalMesh)
 	{
 		SourceImportData = SkeletalMesh->GetAssetImportData();
@@ -199,6 +200,7 @@ TFuture<bool> UInterchangeMeshUtilities::InternalImportCustomLod(TSharedPtr<TPro
 		{
 			bInvalidLodIndex = true;
 		}
+		ObjectType = USkeletalMesh::StaticClass();
 	}
 	else if (StaticMesh)
 	{
@@ -212,6 +214,7 @@ TFuture<bool> UInterchangeMeshUtilities::InternalImportCustomLod(TSharedPtr<TPro
 		{
 			bInvalidLodIndex = true;
 		}
+		ObjectType = UStaticMesh::StaticClass();
 	}
 	else
 	{
@@ -245,6 +248,9 @@ TFuture<bool> UInterchangeMeshUtilities::InternalImportCustomLod(TSharedPtr<TPro
 
 	FImportAssetParameters ImportAssetParameters;
 	ImportAssetParameters.bIsAutomated = true;
+	FInterchangePipelineContextParams ContextParams;
+	ContextParams.ContextType = ImportType;
+	ContextParams.ImportObjectType = ObjectType;
 	if (InterchangeAssetImportData)
 	{
 		TArray<UObject*> Pipelines = InterchangeAssetImportData->GetPipelines();
@@ -261,7 +267,7 @@ TFuture<bool> UInterchangeMeshUtilities::InternalImportCustomLod(TSharedPtr<TPro
 			}
 			if (ensure(GeneratedPipeline))
 			{
-				GeneratedPipeline->AdjustSettingsForContext(ImportType, nullptr, nullptr);
+				GeneratedPipeline->AdjustSettingsForContext(ContextParams);
 				ImportAssetParameters.OverridePipelines.Add(GeneratedPipeline);
 			}
 		}
@@ -277,7 +283,7 @@ TFuture<bool> UInterchangeMeshUtilities::InternalImportCustomLod(TSharedPtr<TPro
 			if (UInterchangePipelineBase* GenericPipeline = NewObject<UInterchangePipelineBase>(GetTransientPackage(), GenericPipelineClass))
 			{
 				GenericPipeline->ClearFlags(EObjectFlags::RF_Standalone | EObjectFlags::RF_Public);
-				GenericPipeline->AdjustSettingsForContext(ImportType, nullptr, nullptr);
+				GenericPipeline->AdjustSettingsForContext(ContextParams);
 				ImportAssetParameters.OverridePipelines.Add(GenericPipeline);
 			}
 		}
@@ -388,6 +394,9 @@ TFuture<bool> UInterchangeMeshUtilities::ImportMorphTarget(USkeletalMesh* Skelet
 
 	FImportAssetParameters ImportAssetParameters;
 	ImportAssetParameters.bIsAutomated = true;
+	FInterchangePipelineContextParams ContextParams;
+	ContextParams.ContextType = ImportType;
+	ContextParams.ImportObjectType = USkeletalMesh::StaticClass();
 	if (InterchangeAssetImportData)
 	{
 		TArray<UObject*> Pipelines = InterchangeAssetImportData->GetPipelines();
@@ -404,7 +413,7 @@ TFuture<bool> UInterchangeMeshUtilities::ImportMorphTarget(USkeletalMesh* Skelet
 			}
 			if (ensure(GeneratedPipeline))
 			{
-				GeneratedPipeline->AdjustSettingsForContext(ImportType, nullptr, nullptr);
+				GeneratedPipeline->AdjustSettingsForContext(ContextParams);
 				ImportAssetParameters.OverridePipelines.Add(GeneratedPipeline);
 			}
 		}
@@ -420,7 +429,7 @@ TFuture<bool> UInterchangeMeshUtilities::ImportMorphTarget(USkeletalMesh* Skelet
 			if (UInterchangePipelineBase* GenericPipeline = NewObject<UInterchangePipelineBase>(GetTransientPackage(), GenericPipelineClass))
 			{
 				GenericPipeline->ClearFlags(EObjectFlags::RF_Standalone | EObjectFlags::RF_Public);
-				GenericPipeline->AdjustSettingsForContext(ImportType, nullptr, nullptr);
+				GenericPipeline->AdjustSettingsForContext(ContextParams);
 				ImportAssetParameters.OverridePipelines.Add(GenericPipeline);
 			}
 		}
