Interchange: Animation rotation keys not imported correctly

#rnx

diff --git a/Engine/Plugins/Interchange/Runtime/Source/Parsers/Fbx/Private/FbxAnimation.cpp b/Engine/Plugins/Interchange/Runtime/Source/Parsers/Fbx/Private/FbxAnimation.cpp
--- a/Engine/Plugins/Interchange/Runtime/Source/Parsers/Fbx/Private/FbxAnimation.cpp
+++ b/Engine/Plugins/Interchange/Runtime/Source/Parsers/Fbx/Private/FbxAnimation.cpp
@@ -429,6 +429,7 @@ namespace UE::Interchange::Private
 					if (TransformChannelCurves[ChannelIndex])
 					{
 						Parameters.IsNodeAnimated = true;
+						break;
 					}
 				}
 			}
diff --git a/Engine/Plugins/Interchange/Runtime/Source/Parsers/Fbx/Private/FbxConvert.cpp b/Engine/Plugins/Interchange/Runtime/Source/Parsers/Fbx/Private/FbxConvert.cpp
--- a/Engine/Plugins/Interchange/Runtime/Source/Parsers/Fbx/Private/FbxConvert.cpp
+++ b/Engine/Plugins/Interchange/Runtime/Source/Parsers/Fbx/Private/FbxConvert.cpp
@@ -68,21 +68,9 @@ namespace UE
 				//Set the original framerate from the current fbx file
 				float FbxFramerate = FbxTime::GetFrameRate(TimeMode);
 
-				int32 AnimStackCount = SDKScene->GetSrcObjectCount<FbxAnimStack>();
-				for (int32 AnimStackIndex = 0; AnimStackIndex < AnimStackCount; AnimStackIndex++)
-				{
-					FbxAnimStack* CurrentAnimStack = SDKScene->GetSrcObject<FbxAnimStack>(AnimStackIndex);
-					int32 NumLayers = CurrentAnimStack->GetMemberCount();
-					for (int LayerIndex = 0; LayerIndex < NumLayers; LayerIndex++)
-					{
-						FbxAnimLayer* AnimLayer = (FbxAnimLayer*)CurrentAnimStack->GetMember(LayerIndex);
-
-						// always apply unroll filter
-						FbxAnimCurveFilterUnroll UnrollFilter;
-						UnrollFilter.Reset();
-						ApplyUnroll(SDKScene->GetRootNode(), AnimLayer, &UnrollFilter);
-					}
-				}
+				//Apply any curve filter here, we currently do not apply any
+				//The unroll curve filter was apply in legacy fbx importer if there was more then one FbxAnimStack.
+				//The unroll curve filter can obliterate curve keys if for example a key do a complete rotation (360 degree in euler)
 
 				FbxAxisSystem FileAxisSystem = SDKScene->GetGlobalSettings().GetAxisSystem();
 				FileSystemDirection = GetFileAxisDirection(FileAxisSystem);
@@ -156,36 +144,6 @@ namespace UE
 			{
 				return FString(UTF8_TO_TCHAR(Name));
 			}
-
-			void FFbxConvert::ApplyUnroll(FbxNode* Node, FbxAnimLayer* Layer, FbxAnimCurveFilterUnroll* UnrollFilter)
-			{
-				if (!ensure(Node) || !ensure(Layer) || !ensure(UnrollFilter))
-				{
-					return;
-				}
-
-				FbxAnimCurveNode* lCN = Node->LclRotation.GetCurveNode(Layer);
-				if (lCN)
-				{
-					FbxAnimCurve* lRCurve[3];
-					lRCurve[0] = lCN->GetCurve(0);
-					lRCurve[1] = lCN->GetCurve(1);
-					lRCurve[2] = lCN->GetCurve(2);
-
-
-					// Set bone rotation order
-					EFbxRotationOrder RotationOrder = eEulerXYZ;
-					Node->GetRotationOrder(FbxNode::eSourcePivot, RotationOrder);
-					UnrollFilter->SetRotationOrder((FbxEuler::EOrder)(RotationOrder));
-
-					UnrollFilter->Apply(lRCurve, 3);
-				}
-
-				for (int32 i = 0; i < Node->GetChildCount(); i++)
-				{
-					ApplyUnroll(Node->GetChild(i), Layer, UnrollFilter);
-				}
-			}
 		}//ns Private
 	}//ns Interchange
 }//ns UE
diff --git a/Engine/Plugins/Interchange/Runtime/Source/Parsers/Fbx/Private/FbxConvert.h b/Engine/Plugins/Interchange/Runtime/Source/Parsers/Fbx/Private/FbxConvert.h
--- a/Engine/Plugins/Interchange/Runtime/Source/Parsers/Fbx/Private/FbxConvert.h
+++ b/Engine/Plugins/Interchange/Runtime/Source/Parsers/Fbx/Private/FbxConvert.h
@@ -181,18 +181,6 @@ namespace UE
 
 			private:
 				
-				//////////////////////////////////////////////////////////////////////////
-				/** Scene Conversion Private Implementation Begin */
-
-				/**
-				 * The Unroll filter expects only rotation curves, we need to walk the scene and extract the
-				 * rotation curves from the nodes property. This can become time consuming but we have no choice.
-				 */
-				static void ApplyUnroll(FbxNode* pNode, FbxAnimLayer* pLayer, FbxAnimCurveFilterUnroll* pUnrollFilter);
-
-				/** Scene Conversion Private Implementation End */
-				//////////////////////////////////////////////////////////////////////////
-
 				template<typename ScalarType>
 				static void VerifyFiniteScalar(ScalarType& Value)
 				{
