Interchange: Avoid scanning asset registry temp folder path.
#rnx

diff --git a/Engine/Plugins/Interchange/Runtime/Source/Pipelines/Private/InterchangeGenericMaterialPipeline.cpp b/Engine/Plugins/Interchange/Runtime/Source/Pipelines/Private/InterchangeGenericMaterialPipeline.cpp
--- a/Engine/Plugins/Interchange/Runtime/Source/Pipelines/Private/InterchangeGenericMaterialPipeline.cpp
+++ b/Engine/Plugins/Interchange/Runtime/Source/Pipelines/Private/InterchangeGenericMaterialPipeline.cpp
@@ -216,12 +216,16 @@ namespace UE::Interchange::InterchangeGenericMaterialPipeline::Private
 		{
 			FPackageName::QueryRootContentPaths(ScanPaths);
 		}
-		else
+		else if (!BasePath.StartsWith(TEXT("/Temp"), ESearchCase::IgnoreCase)) //We must exclude Temp path to avoid asset registry scan path warnings
 		{
 			ScanPaths.Add(BasePath);
 		}
-		const bool bForceRescan = false;
-		AssetRegistry.ScanPathsSynchronous(ScanPaths, bForceRescan);
+
+		if (!ScanPaths.IsEmpty())
+		{
+			constexpr bool bForceRescan = false;
+			AssetRegistry.ScanPathsSynchronous(ScanPaths, bForceRescan);
+		}
 
 
 		FARFilter Filter;
