Interchange: Relaxe the find existing material to find any material (instance or material)even if the factory node class is one or the other.

#rnx

diff --git a/Engine/Plugins/Interchange/Runtime/Source/Pipelines/Private/InterchangeGenericMaterialPipeline.cpp b/Engine/Plugins/Interchange/Runtime/Source/Pipelines/Private/InterchangeGenericMaterialPipeline.cpp
--- a/Engine/Plugins/Interchange/Runtime/Source/Pipelines/Private/InterchangeGenericMaterialPipeline.cpp
+++ b/Engine/Plugins/Interchange/Runtime/Source/Pipelines/Private/InterchangeGenericMaterialPipeline.cpp
@@ -841,14 +841,12 @@ void UInterchangeGenericMaterialPipeline::ExecutePipeline(UInterchangeBaseNodeCo
 					//Make sure we have the correct type of material (can be material instance) before setting the custom object reference.
 					const bool bIsMaterial = MaterialBaseFactoryNode->IsA<UInterchangeMaterialFactoryNode>() && ExistingMaterial->IsA<UMaterial>();
 					const bool bIsMaterialinstance = MaterialBaseFactoryNode->IsA<UInterchangeMaterialInstanceFactoryNode>() && ExistingMaterial->IsA<UMaterialInstance>();
-					if (bIsMaterial || bIsMaterialinstance)
-					{
-						MaterialBaseFactoryNode->SetCustomReferenceObject(ExistingMaterial);
-						//Reimport can only be done on material instances
-						const bool bEnableReimport = !bIsMaterial && bIsMaterialinstance;
-						MaterialBaseFactoryNode->SetCustomIsMaterialImportEnabled(bEnableReimport);
-						MaterialBaseFactoryNode->SetEnabled(bEnableReimport);
-					}
+
+					MaterialBaseFactoryNode->SetCustomReferenceObject(ExistingMaterial);
+					//Reimport can only be done on material instances
+					const bool bEnableReimport = !bIsMaterial && bIsMaterialinstance;
+					MaterialBaseFactoryNode->SetCustomIsMaterialImportEnabled(bEnableReimport);
+					MaterialBaseFactoryNode->SetEnabled(bEnableReimport);
 				}
 			}
 		});
