Interchange:
- Add support of Displacement material connection

diff --git a/Engine/Plugins/Interchange/Runtime/Source/FactoryNodes/Private/InterchangeMaterialFactoryNode.cpp b/Engine/Plugins/Interchange/Runtime/Source/FactoryNodes/Private/InterchangeMaterialFactoryNode.cpp
--- a/Engine/Plugins/Interchange/Runtime/Source/FactoryNodes/Private/InterchangeMaterialFactoryNode.cpp
+++ b/Engine/Plugins/Interchange/Runtime/Source/FactoryNodes/Private/InterchangeMaterialFactoryNode.cpp
@@ -330,6 +330,21 @@ bool UInterchangeMaterialFactoryNode::ConnectOutputToCloth(const FString& Expres
 	return UInterchangeShaderPortsAPI::ConnectOuputToInputByName(this, UE::Interchange::Materials::Sheen::Parameters::SheenRoughness.ToString(), ExpressionNodeUid, OutputName);
 }
 
+bool UInterchangeMaterialFactoryNode::GetDisplacementConnection(FString& ExpressionNodeUid, FString& OutputName) const
+{
+	return UInterchangeShaderPortsAPI::GetInputConnection(this, UE::Interchange::Materials::Common::Parameters::Displacement.ToString(), ExpressionNodeUid, OutputName);
+}
+
+bool UInterchangeMaterialFactoryNode::ConnectToDisplacement(const FString& AttributeValue)
+{
+	return UInterchangeShaderPortsAPI::ConnectDefaultOuputToInput(this, UE::Interchange::Materials::Common::Parameters::Displacement.ToString(), AttributeValue);
+}
+
+bool UInterchangeMaterialFactoryNode::ConnectOutputToDisplacement(const FString& ExpressionNodeUid, const FString& OutputName)
+{
+	return UInterchangeShaderPortsAPI::ConnectOuputToInputByName(this, UE::Interchange::Materials::Common::Parameters::Displacement.ToString(), ExpressionNodeUid, OutputName);
+}
+
 bool UInterchangeMaterialFactoryNode::GetCustomShadingModel(TEnumAsByte<EMaterialShadingModel>& AttributeValue) const
 {
 	IMPLEMENT_NODE_ATTRIBUTE_GETTER(ShadingModel, TEnumAsByte<EMaterialShadingModel>);
diff --git a/Engine/Plugins/Interchange/Runtime/Source/FactoryNodes/Public/InterchangeMaterialFactoryNode.h b/Engine/Plugins/Interchange/Runtime/Source/FactoryNodes/Public/InterchangeMaterialFactoryNode.h
--- a/Engine/Plugins/Interchange/Runtime/Source/FactoryNodes/Public/InterchangeMaterialFactoryNode.h
+++ b/Engine/Plugins/Interchange/Runtime/Source/FactoryNodes/Public/InterchangeMaterialFactoryNode.h
@@ -217,6 +217,15 @@ public:
 	UFUNCTION(BlueprintCallable, Category = "Interchange | Node | MaterialFactory")
 	bool ConnectOutputToCloth(const FString& ExpressionNodeUid, const FString& OutputName);
 
+	UFUNCTION(BlueprintCallable, Category = "Interchange | Node | MaterialFactory")
+	bool GetDisplacementConnection(FString& ExpressionNodeUid, FString& OutputName) const;
+
+	UFUNCTION(BlueprintCallable, Category = "Interchange | Node | MaterialFactory")
+	bool ConnectToDisplacement(const FString& AttributeValue);
+
+	UFUNCTION(BlueprintCallable, Category = "Interchange | Node | MaterialFactory")
+	bool ConnectOutputToDisplacement(const FString& ExpressionNodeUid, const FString& OutputName);
+
 // Material parameters
 public:
 	UFUNCTION(BlueprintCallable, Category = "Interchange | Node | Material")
diff --git a/Engine/Plugins/Interchange/Runtime/Source/Import/Private/Material/InterchangeMaterialFactory.cpp b/Engine/Plugins/Interchange/Runtime/Source/Import/Private/Material/InterchangeMaterialFactory.cpp
--- a/Engine/Plugins/Interchange/Runtime/Source/Import/Private/Material/InterchangeMaterialFactory.cpp
+++ b/Engine/Plugins/Interchange/Runtime/Source/Import/Private/Material/InterchangeMaterialFactory.cpp
@@ -1713,6 +1713,29 @@ void UInterchangeMaterialFactory::SetupMaterial(UMaterial* Material, const FImpo
 		}
 	}
 
+	// Displacement
+	{
+		FString DisplacementUid;
+		FString OutputName;
+
+		if(MaterialFactoryNode->GetDisplacementConnection(DisplacementUid, OutputName))
+		{
+			const UInterchangeMaterialExpressionFactoryNode* DisplacementNode = Cast<UInterchangeMaterialExpressionFactoryNode>(Arguments.NodeContainer->GetNode(DisplacementUid));
+
+			if(DisplacementNode)
+			{
+				if(UMaterialExpression* DisplacementExpression = Builder.CreateExpressionsForNode(*DisplacementNode))
+				{
+					if(FExpressionInput* DisplacementInput = Material->GetExpressionInputForProperty(MP_Displacement))
+					{
+						DisplacementExpression->ConnectExpression(DisplacementInput, GetOutputIndex(*DisplacementExpression, OutputName));
+						Material->bEnableTessellation = true;
+					}
+				}
+			}
+		}
+	}
+
 	UMaterialEditingLibrary::LayoutMaterialExpressions(Material);
 }
 #endif // #if WITH_EDITOR
diff --git a/Engine/Plugins/Interchange/Runtime/Source/Nodes/Public/InterchangeMaterialDefinitions.h b/Engine/Plugins/Interchange/Runtime/Source/Nodes/Public/InterchangeMaterialDefinitions.h
--- a/Engine/Plugins/Interchange/Runtime/Source/Nodes/Public/InterchangeMaterialDefinitions.h
+++ b/Engine/Plugins/Interchange/Runtime/Source/Nodes/Public/InterchangeMaterialDefinitions.h
@@ -341,6 +341,7 @@ namespace Interchange
 				const FName BxDF = TEXT("BxDF"); // input/output of BSDF or BRDF or BXDF or BTDF data
 				const FName Refraction = TEXT("Refraction"); // input/output of BSDF or BRDF or BXDF or BTDF data
 				const FName Anisotropy = TEXT("Anisotropy"); // Type: float
+				const FName Displacement = TEXT("Displacement"); // Type: float
 			}
 		}
 
diff --git a/Engine/Plugins/Interchange/Runtime/Source/Pipelines/Private/InterchangeGenericMaterialPipeline.cpp b/Engine/Plugins/Interchange/Runtime/Source/Pipelines/Private/InterchangeGenericMaterialPipeline.cpp
--- a/Engine/Plugins/Interchange/Runtime/Source/Pipelines/Private/InterchangeGenericMaterialPipeline.cpp
+++ b/Engine/Plugins/Interchange/Runtime/Source/Pipelines/Private/InterchangeGenericMaterialPipeline.cpp
@@ -1768,6 +1768,22 @@ void UInterchangeGenericMaterialPipeline::HandleCommonParameters(const UIntercha
 			}
 		}
 	}
+
+	// Displacement
+	{
+		const bool bHasDisplacementInput = UInterchangeShaderPortsAPI::HasInput(ShaderGraphNode, Parameters::Displacement);
+
+		if(bHasDisplacementInput)
+		{
+			TTuple<UInterchangeMaterialExpressionFactoryNode*, FString> DisplacementExpressionFactoryNode =
+				CreateMaterialExpressionForInput(MaterialFactoryNode, ShaderGraphNode, Parameters::Displacement.ToString(), MaterialFactoryNode->GetUniqueID());
+
+			if(DisplacementExpressionFactoryNode.Get<0>())
+			{
+				MaterialFactoryNode->ConnectOutputToDisplacement(DisplacementExpressionFactoryNode.Get<0>()->GetUniqueID(), DisplacementExpressionFactoryNode.Get<1>());
+			}
+		}
+	}
 }
 
 void UInterchangeGenericMaterialPipeline::HandleFlattenNormalNode(const UInterchangeShaderNode* ShaderNode, UInterchangeBaseMaterialFactoryNode* MaterialFactoryNode,
