Interchange: fix crash if we have empty LODs

#rnx

diff --git a/Engine/Plugins/Interchange/Runtime/Source/Pipelines/Private/InterchangeGenericSkeletalMeshPipeline.cpp b/Engine/Plugins/Interchange/Runtime/Source/Pipelines/Private/InterchangeGenericSkeletalMeshPipeline.cpp
--- a/Engine/Plugins/Interchange/Runtime/Source/Pipelines/Private/InterchangeGenericSkeletalMeshPipeline.cpp
+++ b/Engine/Plugins/Interchange/Runtime/Source/Pipelines/Private/InterchangeGenericSkeletalMeshPipeline.cpp
@@ -529,7 +529,15 @@ void UInterchangeGenericMeshPipeline::AddLodDataToSkeletalMesh(const UInterchang
 
 	const FString SkeletalMeshUid = SkeletalMeshFactoryNode->GetUniqueID();
 	const FString SkeletonUid = SkeletonFactoryNode->GetUniqueID();
-	const int32 LodCount = NodeUidsPerLodIndex.Num();
+
+	int32 MaxLodIndex = 0;
+	for (const TPair<int32, TArray<FString>>& LodIndexAndNodeUids : NodeUidsPerLodIndex)
+	{
+		MaxLodIndex = FMath::Max(MaxLodIndex, LodIndexAndNodeUids.Key);
+	}
+
+	TArray<FString> EmptyLodData;
+	const int32 LodCount = MaxLodIndex + 1;
 	for (int32 LodIndex = 0; LodIndex < LodCount; ++LodIndex)
 	{
 		if (!CommonMeshesProperties->bImportLods && LodIndex > 0)
@@ -539,7 +547,7 @@ void UInterchangeGenericMeshPipeline::AddLodDataToSkeletalMesh(const UInterchang
 		}
 
 		//Copy the nodes unique id because we need to remove nested mesh if the option is to not import them
-		TArray<FString> NodeUids = NodeUidsPerLodIndex.FindChecked(LodIndex);
+		TArray<FString> NodeUids = NodeUidsPerLodIndex.Contains(LodIndex) ? NodeUidsPerLodIndex.FindChecked(LodIndex) : EmptyLodData;
 		if (!CommonSkeletalMeshesAndAnimationsProperties->bImportMeshesInBoneHierarchy)
 		{
 			UE::Interchange::SkeletalMeshGenericPipeline::RemoveNestedMeshNodes(BaseNodeContainer, SkeletonFactoryNode, NodeUids);
diff --git a/Engine/Plugins/Interchange/Runtime/Source/Pipelines/Private/InterchangeGenericStaticMeshPipeline.cpp b/Engine/Plugins/Interchange/Runtime/Source/Pipelines/Private/InterchangeGenericStaticMeshPipeline.cpp
--- a/Engine/Plugins/Interchange/Runtime/Source/Pipelines/Private/InterchangeGenericStaticMeshPipeline.cpp
+++ b/Engine/Plugins/Interchange/Runtime/Source/Pipelines/Private/InterchangeGenericStaticMeshPipeline.cpp
@@ -516,7 +516,14 @@ void UInterchangeGenericMeshPipeline::AddLodDataToStaticMesh(UInterchangeStaticM
 	check(CommonMeshesProperties.IsValid());
 	const FString StaticMeshFactoryUid = StaticMeshFactoryNode->GetUniqueID();
 
-	const int32 LodCount = NodeUidsPerLodIndex.Num();
+	int32 MaxLodIndex = 0;
+	for (const TPair<int32, TArray<FString>>& LodIndexAndNodeUids : NodeUidsPerLodIndex)
+	{
+		MaxLodIndex = FMath::Max(MaxLodIndex, LodIndexAndNodeUids.Key);
+	}
+
+	TArray<FString> EmptyLodData;
+	const int32 LodCount = MaxLodIndex+1;
 	for(int32 LodIndex = 0; LodIndex < LodCount; ++LodIndex)
 	{
 		if (!CommonMeshesProperties->bImportLods && LodIndex > 0)
@@ -525,7 +532,7 @@ void UInterchangeGenericMeshPipeline::AddLodDataToStaticMesh(UInterchangeStaticM
 			continue;
 		}
 
-		const TArray<FString>& NodeUids = NodeUidsPerLodIndex.FindChecked(LodIndex);
+		const TArray<FString>& NodeUids = NodeUidsPerLodIndex.Contains(LodIndex) ? NodeUidsPerLodIndex.FindChecked(LodIndex) : EmptyLodData;
 		// Create a lod data node with all the meshes for this LOD
 		const FString StaticMeshLodDataName = TEXT("LodData") + FString::FromInt(LodIndex);
 		const FString LODDataPrefix = TEXT("\\LodData") + (LodIndex > 0 ? FString::FromInt(LodIndex) : TEXT(""));
