Sequencer: Fix bug where UMovieSceneSequence::FindBindingFromObject could not distinguish between 2 object that were children of different instances, but the same type of component, as the names of the components might match.
Unfortunately to fix this, the previous version of FindBindingFromObject had to be deprecated, replaced with one that accepts a shared playback state so bindings can be properly found from a cache. Previous callers of the function have been fixed up.

diff --git a/Engine/Plugins/Experimental/DaySequence/Source/DaySequence/Private/DaySequence.cpp b/Engine/Plugins/Experimental/DaySequence/Source/DaySequence/Private/DaySequence.cpp
--- a/Engine/Plugins/Experimental/DaySequence/Source/DaySequence/Private/DaySequence.cpp
+++ b/Engine/Plugins/Experimental/DaySequence/Source/DaySequence/Private/DaySequence.cpp
@@ -298,9 +298,17 @@ void UDaySequence::LocateBoundObjects(const FGuid& ObjectId, UObject* Context, T
 	BindingReferences.ResolveBinding(ObjectId, Context, OutObjects);
 }
 
-FGuid UDaySequence::FindBindingFromObject(UObject* InObject, UObject* Context) const
+FGuid UDaySequence::FindBindingFromObject(UObject* InObject, TSharedRef<const FSharedPlaybackState> SharedPlaybackState) const
 {
-	return BindingReferences.FindBindingFromObject(InObject, Context);
+	if (InObject)
+	{
+		if (FMovieSceneEvaluationState* EvaluationState = SharedPlaybackState->FindCapability<FMovieSceneEvaluationState>())
+		{
+			FMovieSceneSequenceID SequenceID = EvaluationState->FindSequenceId(this);
+			return EvaluationState->FindCachedObjectId(*InObject, SequenceID, SharedPlaybackState);
+		}
+	}
+	return FGuid();
 }
 
 void UDaySequence::GatherExpiredObjects(const FMovieSceneObjectCache& InObjectCache, TArray<FGuid>& OutInvalidIDs) const
diff --git a/Engine/Plugins/Experimental/DaySequence/Source/DaySequence/Private/DaySequenceBindingReference.cpp b/Engine/Plugins/Experimental/DaySequence/Source/DaySequence/Private/DaySequenceBindingReference.cpp
--- a/Engine/Plugins/Experimental/DaySequence/Source/DaySequence/Private/DaySequenceBindingReference.cpp
+++ b/Engine/Plugins/Experimental/DaySequence/Source/DaySequence/Private/DaySequenceBindingReference.cpp
@@ -225,21 +225,6 @@ void FDaySequenceBindingReferences::ResolveBinding(const FGuid& ObjectId, UObjec
 	}
 }
 
-FGuid FDaySequenceBindingReferences::FindBindingFromObject(UObject* InObject, UObject* InContext) const
-{
-	FDaySequenceBindingReference Predicate(InObject, InContext);
-
-	for (const TPair<FGuid, FDaySequenceBindingReferenceArray>& Pair : BindingIdToReferences)
-	{
-		if (Pair.Value.References.Contains(Predicate))
-		{
-			return Pair.Key;
-		}
-	}
-
-	return FGuid();
-}
-
 void FDaySequenceBindingReferences::RemoveInvalidBindings(const TSet<FGuid>& ValidBindingIDs)
 {
 	for (auto It = BindingIdToReferences.CreateIterator(); It; ++It)
diff --git a/Engine/Plugins/Experimental/DaySequence/Source/DaySequence/Private/DaySequenceModifierComponent.cpp b/Engine/Plugins/Experimental/DaySequence/Source/DaySequence/Private/DaySequenceModifierComponent.cpp
--- a/Engine/Plugins/Experimental/DaySequence/Source/DaySequence/Private/DaySequenceModifierComponent.cpp
+++ b/Engine/Plugins/Experimental/DaySequence/Source/DaySequence/Private/DaySequenceModifierComponent.cpp
@@ -1010,8 +1010,10 @@ FGuid UDaySequenceModifierComponent::GetOrCreateProceduralBinding(UObject* Objec
 
 	UMovieScene* MovieScene = ProceduralDaySequence->GetMovieScene();
 
+	TSharedRef<const UE::MovieScene::FSharedPlaybackState> SharedPlaybackState = MovieSceneHelpers::CreateTransientSharedPlaybackState(TargetActor, ProceduralDaySequence);
+
 	// Find the main binding
-	FGuid RootGuid = ProceduralDaySequence->FindBindingFromObject(TargetActor, TargetActor);
+	FGuid RootGuid = ProceduralDaySequence->FindBindingFromObject(TargetActor, SharedPlaybackState);
 	if (!RootGuid.IsValid())
 	{
 		FString RootName = TargetActor->GetName();
@@ -1034,7 +1036,7 @@ FGuid UDaySequenceModifierComponent::GetOrCreateProceduralBinding(UObject* Objec
 	}
 
 	// If we're trying to animate a component within the actor, retrieve or create a child binding for that
-	FGuid ComponentGuid = ProceduralDaySequence->FindBindingFromObject(Component, TargetActor);
+	FGuid ComponentGuid = ProceduralDaySequence->FindBindingFromObject(Component, SharedPlaybackState);
 	if (!ComponentGuid.IsValid() && Component)
 	{
 		FString Name = Component->GetName();
diff --git a/Engine/Plugins/Experimental/DaySequence/Source/DaySequence/Private/ProceduralDaySequenceBuilder.cpp b/Engine/Plugins/Experimental/DaySequence/Source/DaySequence/Private/ProceduralDaySequenceBuilder.cpp
--- a/Engine/Plugins/Experimental/DaySequence/Source/DaySequence/Private/ProceduralDaySequenceBuilder.cpp
+++ b/Engine/Plugins/Experimental/DaySequence/Source/DaySequence/Private/ProceduralDaySequenceBuilder.cpp
@@ -355,7 +355,8 @@ FGuid UProceduralDaySequenceBuilder::GetOrCreateProceduralBinding(UObject* Objec
 	UMovieScene* MovieScene = ProceduralDaySequence->GetMovieScene();
 
 	// Find the main binding
-	FGuid RootGuid = ProceduralDaySequence->FindBindingFromObject(TargetActor, TargetActor);
+	TSharedRef<const UE::MovieScene::FSharedPlaybackState> SharedPlaybackState = MovieSceneHelpers::CreateTransientSharedPlaybackState(TargetActor, ProceduralDaySequence);
+	FGuid RootGuid = ProceduralDaySequence->FindBindingFromObject(TargetActor, SharedPlaybackState);
 	if (!RootGuid.IsValid())
 	{
 		FString RootName = TargetActor->GetName();
@@ -378,7 +379,7 @@ FGuid UProceduralDaySequenceBuilder::GetOrCreateProceduralBinding(UObject* Objec
 	}
 
 	// If we're trying to animate a component within the actor, retrieve or create a child binding for that
-	FGuid ComponentGuid = ProceduralDaySequence->FindBindingFromObject(Component, TargetActor);
+	FGuid ComponentGuid = ProceduralDaySequence->FindBindingFromObject(Component, SharedPlaybackState);
 	if (!ComponentGuid.IsValid() && Component)
 	{
 		FString Name = Component->GetName();
diff --git a/Engine/Plugins/Experimental/DaySequence/Source/DaySequence/Public/DaySequence.h b/Engine/Plugins/Experimental/DaySequence/Source/DaySequence/Public/DaySequence.h
--- a/Engine/Plugins/Experimental/DaySequence/Source/DaySequence/Public/DaySequence.h
+++ b/Engine/Plugins/Experimental/DaySequence/Source/DaySequence/Public/DaySequence.h
@@ -35,7 +35,7 @@ public:
 	virtual void BindPossessableObject(const FGuid& ObjectId, UObject& PossessedObject, UObject* Context) override;
 	virtual bool CanPossessObject(UObject& Object, UObject* InPlaybackContext) const override;
 	virtual void LocateBoundObjects(const FGuid& ObjectId, UObject* Context, TArray<UObject*, TInlineAllocator<1>>& OutObjects) const override;
-	virtual FGuid FindBindingFromObject(UObject* InObject, UObject* Context) const override;
+	virtual FGuid FindBindingFromObject(UObject* InObject, TSharedRef<const FSharedPlaybackState> SharedPlaybackState) const override;
 	virtual void GatherExpiredObjects(const FMovieSceneObjectCache& InObjectCache, TArray<FGuid>& OutInvalidIDs) const override;
 	virtual UMovieScene* GetMovieScene() const override;
 	virtual UObject* GetParentObject(UObject* Object) const override;
diff --git a/Engine/Plugins/Experimental/DaySequence/Source/DaySequence/Public/DaySequenceBindingReference.h b/Engine/Plugins/Experimental/DaySequence/Source/DaySequence/Public/DaySequenceBindingReference.h
--- a/Engine/Plugins/Experimental/DaySequence/Source/DaySequence/Public/DaySequenceBindingReference.h
+++ b/Engine/Plugins/Experimental/DaySequence/Source/DaySequence/Public/DaySequenceBindingReference.h
@@ -146,14 +146,6 @@ struct FDaySequenceBindingReferences
 	 */
 	void ResolveBinding(const FGuid& ObjectId, UObject* InContext, TArray<UObject*, TInlineAllocator<1>>& OutObjects) const;
 
-	/**
-	 * Find a binding for the specified ID using a given context
-	 *
-	 * @param InObject					The object to search for
-	 * @param InContext					A context in which InObject resides
-	 */
-	FGuid FindBindingFromObject(UObject* InObject, UObject* InContext) const;
-
 	/**
 	 * Const accessor for the currently bound anim instance IDs
 	 */
diff --git a/Engine/Plugins/Importers/USDImporter/Source/USDStage/Private/USDLevelSequenceHelper.cpp b/Engine/Plugins/Importers/USDImporter/Source/USDStage/Private/USDLevelSequenceHelper.cpp
--- a/Engine/Plugins/Importers/USDImporter/Source/USDStage/Private/USDLevelSequenceHelper.cpp
+++ b/Engine/Plugins/Importers/USDImporter/Source/USDStage/Private/USDLevelSequenceHelper.cpp
@@ -3859,7 +3859,9 @@ FGuid FUsdLevelSequenceHelperImpl::GetOrCreateComponentBinding(
 	// Make sure we always bind the parent actor too
 	if (AActor* Actor = ComponentToBind.GetOwner())
 	{
-		ActorBinding = Sequence.FindBindingFromObject(Actor, Actor->GetWorld());
+		TSharedRef<const UE::MovieScene::FSharedPlaybackState> SharedPlaybackState = MovieSceneHelpers::CreateTransientSharedPlaybackState(Actor, &Sequence);
+
+		ActorBinding = Sequence.FindBindingFromObject(Actor, SharedPlaybackState);
 		if (!ActorBinding.IsValid())
 		{
 			// We use the label here because that will always be named after the prim that caused the actor
diff --git a/Engine/Plugins/Importers/USDImporter/Source/USDStage/Private/USDStageActor.cpp b/Engine/Plugins/Importers/USDImporter/Source/USDStage/Private/USDStageActor.cpp
--- a/Engine/Plugins/Importers/USDImporter/Source/USDStage/Private/USDStageActor.cpp
+++ b/Engine/Plugins/Importers/USDImporter/Source/USDStage/Private/USDStageActor.cpp
@@ -4584,13 +4584,13 @@ void AUsdStageActor::OnObjectPropertyChanged(UObject* ObjectBeingModified, FProp
 						}
 
 						TArray<FGuid> BindingsToCheck;
-						BindingsToCheck.Add(Sequence->FindBindingFromObject(ComponentBeingModified, Context));
+						BindingsToCheck.Add(Sequence->FindBindingFromObject(ComponentBeingModified, PinnedSequencer->GetSharedPlaybackState()));
 						if (bIsRootComponent)
 						{
 							// Maybe all the sequence has is a track directly on the actor. That's still enough to
 							// supress a root component animation in case the property is just mirrored on the actor,
 							// so let's try checking for that
-							BindingsToCheck.Add(Sequence->FindBindingFromObject(OwnerActor, ActorContext));
+							BindingsToCheck.Add(Sequence->FindBindingFromObject(OwnerActor, PinnedSequencer->GetSharedPlaybackState()));
 						}
 
 						for (const FGuid& BindingGuid : BindingsToCheck)
diff --git a/Engine/Plugins/Interchange/Runtime/Source/Import/Private/Animation/InterchangeLevelSequenceFactory.cpp b/Engine/Plugins/Interchange/Runtime/Source/Import/Private/Animation/InterchangeLevelSequenceFactory.cpp
--- a/Engine/Plugins/Interchange/Runtime/Source/Import/Private/Animation/InterchangeLevelSequenceFactory.cpp
+++ b/Engine/Plugins/Interchange/Runtime/Source/Import/Private/Animation/InterchangeLevelSequenceFactory.cpp
@@ -632,7 +632,8 @@ namespace UE::Interchange::Private
 
 	FGuid FLevelSequenceHelper::BindActorToLevelSequence(AActor* Actor)
 	{
-		FGuid ActorBinding = LevelSequence.FindBindingFromObject(Actor, Actor->GetWorld());
+		TSharedRef<const UE::MovieScene::FSharedPlaybackState> SharedPlaybackState = MovieSceneHelpers::CreateTransientSharedPlaybackState(Actor, &LevelSequence);
+		FGuid ActorBinding = LevelSequence.FindBindingFromObject(Actor, SharedPlaybackState);
 		if(!ActorBinding.IsValid())
 		{
 			ActorBinding = MovieScene->AddPossessable(Actor->GetActorLabel(), Actor->GetClass());
@@ -644,9 +645,10 @@ namespace UE::Interchange::Private
 
 	FGuid FLevelSequenceHelper::BindComponentToLevelSequence(AActor* Actor)
 	{
+		TSharedRef<const UE::MovieScene::FSharedPlaybackState> SharedPlaybackState = MovieSceneHelpers::CreateTransientSharedPlaybackState(Actor, &LevelSequence);
 		FGuid ActorBinding = BindActorToLevelSequence(Actor);
 		USceneComponent* Component = Actor->GetDefaultAttachComponent();
-		FGuid ComponentBinding = LevelSequence.FindBindingFromObject(Component, Actor);
+		FGuid ComponentBinding = LevelSequence.FindBindingFromObject(Component, SharedPlaybackState);
 		if(!ComponentBinding.IsValid())
 		{
 			ComponentBinding = MovieScene->AddPossessable(Component->GetReadableName(), Component->GetClass());
diff --git a/Engine/Source/Runtime/LevelSequence/Private/LevelSequence.cpp b/Engine/Source/Runtime/LevelSequence/Private/LevelSequence.cpp
--- a/Engine/Source/Runtime/LevelSequence/Private/LevelSequence.cpp
+++ b/Engine/Source/Runtime/LevelSequence/Private/LevelSequence.cpp
@@ -634,7 +634,23 @@ void ULevelSequence::LocateBoundObjects(const FGuid& ObjectId, UObject* Context,
 
 FGuid ULevelSequence::FindBindingFromObject(UObject* InObject, UObject* Context) const
 {
-	return BindingReferences.FindBindingFromObject(InObject, Context);
+PRAGMA_DISABLE_DEPRECATION_WARNINGS
+	return UMovieSceneSequence::FindBindingFromObject(InObject, Context);
+PRAGMA_ENABLE_DEPRECATION_WARNINGS
+}
+
+
+FGuid ULevelSequence::FindBindingFromObject(UObject* InObject, TSharedRef<const UE::MovieScene::FSharedPlaybackState> SharedPlaybackState) const
+{
+	if (InObject)
+	{
+		if (FMovieSceneEvaluationState* EvaluationState = SharedPlaybackState->FindCapability<FMovieSceneEvaluationState>())
+		{
+			FMovieSceneSequenceID SequenceID = EvaluationState->FindSequenceId(this);
+			return EvaluationState->FindCachedObjectId(*InObject, SequenceID, SharedPlaybackState);
+		}
+	}
+	return FGuid();
 }
 
 void ULevelSequence::GatherExpiredObjects(const FMovieSceneObjectCache& InObjectCache, TArray<FGuid>& OutInvalidIDs) const
diff --git a/Engine/Source/Runtime/LevelSequence/Public/LevelSequence.h b/Engine/Source/Runtime/LevelSequence/Public/LevelSequence.h
--- a/Engine/Source/Runtime/LevelSequence/Public/LevelSequence.h
+++ b/Engine/Source/Runtime/LevelSequence/Public/LevelSequence.h
@@ -42,7 +42,11 @@ public:
 	// UMovieSceneSequence interface
 	LEVELSEQUENCE_API virtual void BindPossessableObject(const FGuid& ObjectId, UObject& PossessedObject, UObject* Context) override;
 	LEVELSEQUENCE_API virtual bool CanPossessObject(UObject& Object, UObject* InPlaybackContext) const override;
+
+	UE_DEPRECATED(5.5, "Use the version that takes a SharedPlaybackState")
 	LEVELSEQUENCE_API virtual FGuid FindBindingFromObject(UObject* InObject, UObject* Context) const override;
+	LEVELSEQUENCE_API virtual FGuid FindBindingFromObject(UObject* InObject, TSharedRef<const FSharedPlaybackState> SharedPlaybackState) const override;
+
 	LEVELSEQUENCE_API virtual void GatherExpiredObjects(const FMovieSceneObjectCache& InObjectCache, TArray<FGuid>& OutInvalidIDs) const override;
 	LEVELSEQUENCE_API virtual UMovieScene* GetMovieScene() const override;
 	LEVELSEQUENCE_API virtual UObject* GetParentObject(UObject* Object) const override;
diff --git a/Engine/Source/Runtime/MovieScene/Private/Evaluation/MovieSceneEvaluationState.cpp b/Engine/Source/Runtime/MovieScene/Private/Evaluation/MovieSceneEvaluationState.cpp
--- a/Engine/Source/Runtime/MovieScene/Private/Evaluation/MovieSceneEvaluationState.cpp
+++ b/Engine/Source/Runtime/MovieScene/Private/Evaluation/MovieSceneEvaluationState.cpp
@@ -647,6 +647,11 @@ UMovieSceneSequence* FMovieSceneEvaluationState::FindSequence(FMovieSceneSequenc
 	return Cache ? Cache->ObjectCache.GetSequence() : nullptr;
 }
 
+FMovieSceneSequenceID FMovieSceneEvaluationState::FindSequenceId(const UMovieSceneSequence* InSequence) const
+{
+	return FindSequenceId(const_cast<UMovieSceneSequence*>(InSequence));
+}
+
 FMovieSceneSequenceID FMovieSceneEvaluationState::FindSequenceId(UMovieSceneSequence* InSequence) const
 {
 	for (auto& Pair : ObjectCaches)
diff --git a/Engine/Source/Runtime/MovieScene/Private/MovieSceneCommonHelpers.cpp b/Engine/Source/Runtime/MovieScene/Private/MovieSceneCommonHelpers.cpp
--- a/Engine/Source/Runtime/MovieScene/Private/MovieSceneCommonHelpers.cpp
+++ b/Engine/Source/Runtime/MovieScene/Private/MovieSceneCommonHelpers.cpp
@@ -1002,9 +1002,9 @@ TSharedRef<UE::MovieScene::FSharedPlaybackState> MovieSceneHelpers::CreateTransi
 	CreateParams.PlaybackContext = WorldContext;
 	TSharedRef<FSharedPlaybackState> TransientPlaybackState = MakeShared<FSharedPlaybackState>(*Sequence, CreateParams);
 
-	FMovieSceneEvaluationState State;
-	TransientPlaybackState->AddCapabilityRaw(&State);
-	State.AssignSequence(MovieSceneSequenceID::Root, *Sequence, TransientPlaybackState);
+	TSharedRef<FMovieSceneEvaluationState> State = MakeShared<FMovieSceneEvaluationState>();
+	TransientPlaybackState->AddCapabilityShared(State);
+	State->AssignSequence(MovieSceneSequenceID::Root, *Sequence, TransientPlaybackState);
 
 	return TransientPlaybackState;
 }
diff --git a/Engine/Source/Runtime/MovieScene/Private/MovieSceneSequence.cpp b/Engine/Source/Runtime/MovieScene/Private/MovieSceneSequence.cpp
--- a/Engine/Source/Runtime/MovieScene/Private/MovieSceneSequence.cpp
+++ b/Engine/Source/Runtime/MovieScene/Private/MovieSceneSequence.cpp
@@ -87,6 +87,18 @@ void UMovieSceneSequence::LocateBoundObjects(const FGuid& ObjectId, const UE::Un
 	}
 }
 
+FGuid UMovieSceneSequence::FindBindingFromObject(UObject* InObject, UObject* Context) const
+{
+	if (InObject && Context)
+	{
+		TSharedRef<const UE::MovieScene::FSharedPlaybackState> SharedPlaybackState = MovieSceneHelpers::CreateTransientSharedPlaybackState(Context, const_cast<UMovieSceneSequence*>(this));
+
+		return FindBindingFromObject(InObject, SharedPlaybackState);
+	}
+
+	return FGuid();
+}
+
 void UMovieSceneSequence::PostLoad()
 {
 	UMovieSceneCompiledDataManager* PrecompiledData = UMovieSceneCompiledDataManager::GetPrecompiledData();
diff --git a/Engine/Source/Runtime/MovieScene/Public/Evaluation/MovieSceneEvaluationState.h b/Engine/Source/Runtime/MovieScene/Public/Evaluation/MovieSceneEvaluationState.h
--- a/Engine/Source/Runtime/MovieScene/Public/Evaluation/MovieSceneEvaluationState.h
+++ b/Engine/Source/Runtime/MovieScene/Public/Evaluation/MovieSceneEvaluationState.h
@@ -313,6 +313,7 @@ struct FMovieSceneEvaluationState : public UE::MovieScene::IPlaybackCapability
 	 * @param InSequence		The sequence to look up
 	 */
 	MOVIESCENE_API FMovieSceneSequenceID FindSequenceId(UMovieSceneSequence* InSequence) const;
+	MOVIESCENE_API FMovieSceneSequenceID FindSequenceId(const UMovieSceneSequence* InSequence) const;
 
 	/**
 	 * Attempt deduce the posessable or spawnable that relates to the specified object
diff --git a/Engine/Source/Runtime/MovieScene/Public/MovieSceneBindingReferences.h b/Engine/Source/Runtime/MovieScene/Public/MovieSceneBindingReferences.h
--- a/Engine/Source/Runtime/MovieScene/Public/MovieSceneBindingReferences.h
+++ b/Engine/Source/Runtime/MovieScene/Public/MovieSceneBindingReferences.h
@@ -204,6 +204,7 @@ struct FMovieSceneBindingReferences
 	 * @oaram StreamedLevelAssetPath    The path to the streamed level asset that contains the level sequence actor playing back the sequence. 'None' for any non-instance-level setups.
 	 * @param OutObjects				Array to populate with resolved object bindings
 	 */
+	UE_DEPRECATED(5.5, "This method has been deprecated as it produces errors for subobjects. Please call UMovieSceneSequence::FindBindingFromObject passing in a SharedPlaybackState")
 	MOVIESCENE_API FGuid FindBindingFromObject(UObject* InObject, UObject* InContext) const;
 
 	/**
diff --git a/Engine/Source/Runtime/MovieScene/Public/MovieSceneSequence.h b/Engine/Source/Runtime/MovieScene/Public/MovieSceneSequence.h
--- a/Engine/Source/Runtime/MovieScene/Public/MovieSceneSequence.h
+++ b/Engine/Source/Runtime/MovieScene/Public/MovieSceneSequence.h
@@ -186,14 +186,17 @@ public:
 	 */
 	MOVIESCENE_API FGuid FindPossessableObjectId(UObject& Object, UObject* Context) const;
 
+	UE_DEPRECATED(5.5, "Use the version that takes a SharedPlaybackState")
+	MOVIESCENE_API virtual FGuid FindBindingFromObject(UObject* InObject, UObject* Context) const;
+
 	/**
 	 * Optional method for efficient lookup of an object binding from an actual object in the world
 	 *
 	 * @param ObjectId				The unique identifier of the object.
-	 * @param Context				Optional context to use to find the required object (for instance, a parent spawnable object or its world)
+	 * @param SharedPlaybackState   A shared playback state for the binding- create a transient one if necessary.
 	 * @return The object's guid, or zero guid if the object is not a valid possessable in the current context
 	 */
-	virtual FGuid FindBindingFromObject(UObject* InObject, UObject* Context) const { return FGuid(); }
+	virtual FGuid FindBindingFromObject(UObject* InObject, TSharedRef<const FSharedPlaybackState> SharedPlaybackState) const { return FGuid(); }
 
 	/**
 	 * Called to validate the specified object cache by removing anything that should be deemed out of date
diff --git a/Engine/Source/Runtime/MovieSceneTracks/Private/Systems/MovieScenePredictionSystem.cpp b/Engine/Source/Runtime/MovieSceneTracks/Private/Systems/MovieScenePredictionSystem.cpp
--- a/Engine/Source/Runtime/MovieSceneTracks/Private/Systems/MovieScenePredictionSystem.cpp
+++ b/Engine/Source/Runtime/MovieSceneTracks/Private/Systems/MovieScenePredictionSystem.cpp
@@ -272,7 +272,7 @@ int32 UMovieSceneAsyncAction_SequencePrediction::ImportTransformEntities(UObject
 	if (UMovieSceneSequence* Sequence = CompiledEntry.GetSequence())
 	{
 		const bool bIsSpawnable = SpawnableAnnotation && SpawnableAnnotation->SequenceID == MovieSceneSequenceID::Root;
-		FGuid FoundID = bIsSpawnable ? SpawnableAnnotation->ObjectBindingID : Sequence->FindBindingFromObject(PredicateObject, ObjectContext);
+		FGuid FoundID = bIsSpawnable ? SpawnableAnnotation->ObjectBindingID : Sequence->FindBindingFromObject(PredicateObject, SequencePlayer->GetSharedPlaybackState());
 
 		if (FoundID.IsValid())
 		{
@@ -303,7 +303,7 @@ int32 UMovieSceneAsyncAction_SequencePrediction::ImportTransformEntities(UObject
 						if (SubDataID.IsValid())
 						{
 							const bool bIsSpawnable = SpawnableAnnotation && SpawnableAnnotation->SequenceID == SubSequenceIt->SequenceID;
-							FGuid FoundID = bIsSpawnable ? SpawnableAnnotation->ObjectBindingID : SubSequence->FindBindingFromObject(PredicateObject, ObjectContext);
+							FGuid FoundID = bIsSpawnable ? SpawnableAnnotation->ObjectBindingID : SubSequence->FindBindingFromObject(PredicateObject, SequencePlayer->GetSharedPlaybackState());
 
 							if (FoundID.IsValid())
 							{
