Fix typo in MaterialX settings

diff --git a/Engine/Plugins/Interchange/Runtime/Source/Pipelines/Public/InterchangeMaterialXPipeline.h b/Engine/Plugins/Interchange/Runtime/Source/Pipelines/Public/InterchangeMaterialXPipeline.h
--- a/Engine/Plugins/Interchange/Runtime/Source/Pipelines/Public/InterchangeMaterialXPipeline.h
+++ b/Engine/Plugins/Interchange/Runtime/Source/Pipelines/Public/InterchangeMaterialXPipeline.h
@@ -21,7 +21,7 @@ uint32 INTERCHANGEPIPELINES_API GetTypeHash(EInterchangeMaterialXSettings Key);
 
 bool INTERCHANGEPIPELINES_API operator==(EInterchangeMaterialXSettings Lhs, EInterchangeMaterialXSettings Rhs);
 
-UCLASS(config = Interchange, meta = (DisplayName = "Interchange MaterialX"))
+UCLASS(config = Interchange, meta = (DisplayName = "Interchange MaterialX", ToolTip = "Interchange MaterialX Pipeline Settings"))
 class INTERCHANGEPIPELINES_API UMaterialXPipelineSettings : public UDeveloperSettings
 {
 	GENERATED_BODY()
