[Interchange - FBX] Animation payload fetch critical section fix.

#rnx

diff --git a/Engine/Plugins/Interchange/Runtime/Source/Parsers/Fbx/Private/FbxAPI.cpp b/Engine/Plugins/Interchange/Runtime/Source/Parsers/Fbx/Private/FbxAPI.cpp
--- a/Engine/Plugins/Interchange/Runtime/Source/Parsers/Fbx/Private/FbxAPI.cpp
+++ b/Engine/Plugins/Interchange/Runtime/Source/Parsers/Fbx/Private/FbxAPI.cpp
@@ -250,6 +250,9 @@ namespace UE
 
 			bool FFbxParser::FetchAnimationBakeTransformPayload(const TArray<UE::Interchange::FAnimationPayloadQuery>& PayloadQueries, const FString& ResultFolder, FCriticalSection* ResultPayloadsCriticalSection, TAtomic<int64>& UniqueIdCounter, TMap<FString, FString>& ResultPayloads/*PayloadUniqueID to FilePath*/)
 			{
+				//Critical section to force payload to be fetch one by one with no concurrency.
+				FScopeLock Lock(&PayloadCriticalSection);
+
 				TMap<uint32, TArray<const UE::Interchange::FAnimationPayloadQuery*>> PayloadQueriesGrouped;
 
 				for (const UE::Interchange::FAnimationPayloadQuery& PayloadQuery : PayloadQueries)
