MaterialX:
- Adding support of displacement for UsdPreviewSurface

Interchange:
- Adding Occlusion for Substrate

#virtualized

diff --git a/Engine/Plugins/Interchange/Runtime/Source/Import/Private/Material/InterchangeMaterialFactory.cpp b/Engine/Plugins/Interchange/Runtime/Source/Import/Private/Material/InterchangeMaterialFactory.cpp
--- a/Engine/Plugins/Interchange/Runtime/Source/Import/Private/Material/InterchangeMaterialFactory.cpp
+++ b/Engine/Plugins/Interchange/Runtime/Source/Import/Private/Material/InterchangeMaterialFactory.cpp
@@ -1259,6 +1259,51 @@ void UInterchangeMaterialFactory::SetupMaterial(UMaterial* Material, const FImpo
 			}
 		}
 
+		// Displacement
+		if(UInterchangeShaderPortsAPI::HasInput(MaterialFactoryNode, SubstrateMaterial::Parameters::Displacement))
+		{
+			FString ExpressionNodeUid;
+			FString OutputName;
+
+			UInterchangeShaderPortsAPI::GetInputConnection(MaterialFactoryNode, SubstrateMaterial::Parameters::Displacement.ToString(), ExpressionNodeUid, OutputName);
+
+			const UInterchangeMaterialExpressionFactoryNode* Displacement = Cast<UInterchangeMaterialExpressionFactoryNode>(Arguments.NodeContainer->GetNode(ExpressionNodeUid));
+
+			if(Displacement)
+			{
+				if(UMaterialExpression* DisplacementExpression = Builder.CreateExpressionsForNode(*Displacement))
+				{
+					if(FExpressionInput* DisplacementInput = Material->GetExpressionInputForProperty(MP_Displacement))
+					{
+						DisplacementExpression->ConnectExpression(DisplacementInput, GetOutputIndex(*DisplacementExpression, OutputName));
+						Material->bEnableTessellation = true;
+					}
+				}
+			}
+		}
+
+		// Occlusion
+		if(UInterchangeShaderPortsAPI::HasInput(MaterialFactoryNode, SubstrateMaterial::Parameters::Occlusion))
+		{
+			FString ExpressionNodeUid;
+			FString OutputName;
+
+			UInterchangeShaderPortsAPI::GetInputConnection(MaterialFactoryNode, SubstrateMaterial::Parameters::Occlusion.ToString(), ExpressionNodeUid, OutputName);
+
+			const UInterchangeMaterialExpressionFactoryNode* Occlusion = Cast<UInterchangeMaterialExpressionFactoryNode>(Arguments.NodeContainer->GetNode(ExpressionNodeUid));
+
+			if(Occlusion)
+			{
+				if(UMaterialExpression* OcclusionExpression = Builder.CreateExpressionsForNode(*Occlusion))
+				{
+					if(FExpressionInput* OcclusionInput = Material->GetExpressionInputForProperty(MP_AmbientOcclusion))
+					{
+						OcclusionExpression->ConnectExpression(OcclusionInput, GetOutputIndex(*OcclusionExpression, OutputName));
+					}
+				}
+			}
+		}
+
 		UMaterialEditingLibrary::LayoutMaterialExpressions(Material);
 
 		return;
diff --git a/Engine/Plugins/Interchange/Runtime/Source/Import/Private/MaterialX/MaterialXUtils/MaterialXUsdPreviewSurfaceShader.cpp b/Engine/Plugins/Interchange/Runtime/Source/Import/Private/MaterialX/MaterialXUtils/MaterialXUsdPreviewSurfaceShader.cpp
--- a/Engine/Plugins/Interchange/Runtime/Source/Import/Private/MaterialX/MaterialXUtils/MaterialXUsdPreviewSurfaceShader.cpp
+++ b/Engine/Plugins/Interchange/Runtime/Source/Import/Private/MaterialX/MaterialXUtils/MaterialXUsdPreviewSurfaceShader.cpp
@@ -3,6 +3,7 @@
 #if WITH_EDITOR
 
 #include "MaterialXUsdPreviewSurfaceShader.h"
+#include "Engine/EngineTypes.h"
 
 namespace mx = MaterialX;
 
@@ -46,10 +47,27 @@ void FMaterialXUsdPreviewSurfaceShader::Translate(MaterialX::NodePtr UsdPreviewS
 		UInterchangeShaderPortsAPI::ConnectOuputToInputByName(ShaderGraphNode, PBRMR::Parameters::Refraction.ToString(), UsdPreviewSurfaceShaderNode->GetUniqueID(), PBRMR::Parameters::Refraction.ToString());
 		UInterchangeShaderPortsAPI::ConnectOuputToInputByName(ShaderGraphNode, ClearCoat::Parameters::ClearCoat.ToString(), UsdPreviewSurfaceShaderNode->GetUniqueID(), ClearCoat::Parameters::ClearCoat.ToString());
 		UInterchangeShaderPortsAPI::ConnectOuputToInputByName(ShaderGraphNode, ClearCoat::Parameters::ClearCoatRoughness.ToString(), UsdPreviewSurfaceShaderNode->GetUniqueID(), ClearCoat::Parameters::ClearCoatRoughness.ToString());
+
+		if(UInterchangeShaderPortsAPI::HasInput(UsdPreviewSurfaceShaderNode, UsdPreviewSurface::Parameters::Displacement))
+		{
+			UInterchangeShaderPortsAPI::ConnectOuputToInputByName(ShaderGraphNode, Common::Parameters::Displacement.ToString(), UsdPreviewSurfaceShaderNode->GetUniqueID(), Common::Parameters::Displacement.ToString());
+		}
 	}
 	else
 	{
-		UInterchangeShaderPortsAPI::ConnectDefaultOuputToInput(ShaderGraphNode, SubstrateMaterial::Parameters::FrontMaterial.ToString(), UsdPreviewSurfaceShaderNode->GetUniqueID());
+		UInterchangeShaderPortsAPI::ConnectOuputToInputByName(ShaderGraphNode, SubstrateMaterial::Parameters::FrontMaterial.ToString(), UsdPreviewSurfaceShaderNode->GetUniqueID(), SubstrateMaterial::Parameters::FrontMaterial.ToString());
+		UInterchangeShaderPortsAPI::ConnectOuputToInputByName(ShaderGraphNode, SubstrateMaterial::Parameters::Occlusion.ToString(), UsdPreviewSurfaceShaderNode->GetUniqueID(), SubstrateMaterial::Parameters::Occlusion.ToString());
+
+		if(UInterchangeShaderPortsAPI::HasInput(UsdPreviewSurfaceShaderNode, UsdPreviewSurface::Parameters::Displacement))
+		{
+			UInterchangeShaderPortsAPI::ConnectOuputToInputByName(ShaderGraphNode, SubstrateMaterial::Parameters::Displacement.ToString(), UsdPreviewSurfaceShaderNode->GetUniqueID(), SubstrateMaterial::Parameters::Displacement.ToString());
+		}
+
+		if(UInterchangeShaderPortsAPI::HasInput(UsdPreviewSurfaceShaderNode, UsdPreviewSurface::Parameters::Opacity))
+		{
+			ShaderGraphNode->SetCustomBlendMode(EBlendMode::BLEND_TranslucentColoredTransmittance);
+			UInterchangeShaderPortsAPI::ConnectOuputToInputByName(ShaderGraphNode, PBRMR::Parameters::Opacity.ToString(), UsdPreviewSurfaceShaderNode->GetUniqueID(), PBRMR::Parameters::Opacity.ToString());
+		}
 	}
 }
 #endif
\ No newline at end of file
diff --git a/Engine/Plugins/Interchange/Runtime/Source/Nodes/Public/InterchangeMaterialDefinitions.h b/Engine/Plugins/Interchange/Runtime/Source/Nodes/Public/InterchangeMaterialDefinitions.h
--- a/Engine/Plugins/Interchange/Runtime/Source/Nodes/Public/InterchangeMaterialDefinitions.h
+++ b/Engine/Plugins/Interchange/Runtime/Source/Nodes/Public/InterchangeMaterialDefinitions.h
@@ -642,8 +642,9 @@ namespace Interchange
 			{
 				const FName FrontMaterial = TEXT("Front Material");
 				const FName OpacityMask = TEXT("Opacity Mask");
-			}
-			
+				const FName Displacement = TEXT("Displacement");
+				const FName Occlusion = TEXT("Occlusion");
+			}			
 		}
 	}
 }
diff --git a/Engine/Plugins/Interchange/Runtime/Source/Pipelines/Private/InterchangeGenericMaterialPipeline.cpp b/Engine/Plugins/Interchange/Runtime/Source/Pipelines/Private/InterchangeGenericMaterialPipeline.cpp
--- a/Engine/Plugins/Interchange/Runtime/Source/Pipelines/Private/InterchangeGenericMaterialPipeline.cpp
+++ b/Engine/Plugins/Interchange/Runtime/Source/Pipelines/Private/InterchangeGenericMaterialPipeline.cpp
@@ -3542,6 +3542,30 @@ bool UInterchangeGenericMaterialPipeline::HandleSubstrate(const UInterchangeShad
 			}
 		}
 
+		if(UInterchangeShaderPortsAPI::HasInput(ShaderGraphNode, SubstrateMaterial::Parameters::Displacement))
+		{
+			TTuple<UInterchangeMaterialExpressionFactoryNode*, FString> DisplacementFactoryNode =
+				CreateMaterialExpressionForInput(MaterialFactoryNode, ShaderGraphNode, SubstrateMaterial::Parameters::Displacement.ToString(), MaterialFactoryNode->GetUniqueID());
+			ensure(DisplacementFactoryNode.Get<0>());
+
+			if(DisplacementFactoryNode.Get<0>())
+			{
+				UInterchangeShaderPortsAPI::ConnectOuputToInputByName(MaterialFactoryNode, SubstrateMaterial::Parameters::Displacement.ToString(), DisplacementFactoryNode.Get<0>()->GetUniqueID(), DisplacementFactoryNode.Get<1>());
+			}
+		}
+
+		if(UInterchangeShaderPortsAPI::HasInput(ShaderGraphNode, SubstrateMaterial::Parameters::Occlusion))
+		{
+			TTuple<UInterchangeMaterialExpressionFactoryNode*, FString> OcclusionFactoryNode =
+				CreateMaterialExpressionForInput(MaterialFactoryNode, ShaderGraphNode, SubstrateMaterial::Parameters::Occlusion.ToString(), MaterialFactoryNode->GetUniqueID());
+			ensure(OcclusionFactoryNode.Get<0>());
+
+			if(OcclusionFactoryNode.Get<0>())
+			{
+				UInterchangeShaderPortsAPI::ConnectOuputToInputByName(MaterialFactoryNode, SubstrateMaterial::Parameters::Occlusion.ToString(), OcclusionFactoryNode.Get<0>()->GetUniqueID(), OcclusionFactoryNode.Get<1>());
+			}
+		}
+
 		if(EBlendMode BlendMode; ShaderGraphNode->GetCustomBlendMode(reinterpret_cast<int&>(BlendMode)))
 		{
 			MaterialFactoryNode->SetCustomBlendMode(BlendMode);
