USD: Interchange translator for lights and cameras.

diff --git a/Engine/Plugins/Interchange/Runtime/Source/Import/Private/Usd/InterchangeUsdTranslator.cpp b/Engine/Plugins/Interchange/Runtime/Source/Import/Private/Usd/InterchangeUsdTranslator.cpp
--- a/Engine/Plugins/Interchange/Runtime/Source/Import/Private/Usd/InterchangeUsdTranslator.cpp
+++ b/Engine/Plugins/Interchange/Runtime/Source/Import/Private/Usd/InterchangeUsdTranslator.cpp
@@ -1,4 +1,4 @@
-// Copyright Epic Games, Inc. All Rights Reserved.
+// Copyright Epic Games, Inc. All Rights Reserved. 
 
 #include "Usd/InterchangeUsdTranslator.h"
 
@@ -14,6 +14,7 @@
 #include "USDConversionUtils.h"
 #include "USDGeomMeshConversion.h"
 #include "USDLog.h"
+#include "USDLightConversion.h"	
 #include "USDMaterialUtils.h"
 #include "USDPrimConversion.h"
 #include "USDShadeConversion.h"
@@ -39,6 +40,7 @@
 #if USE_USD_SDK
 #include "USDIncludesStart.h"
 #include "pxr/usd/usdGeom/tokens.h"
+#include "pxr/usd/usdLux/tokens.h"
 #include "pxr/usd/usdPhysics/tokens.h"
 #include "pxr/usd/usdShade/tokens.h"
 #include "USDIncludesEnd.h"
@@ -436,50 +438,96 @@ namespace UE::InterchangeUsdTranslator::Private
 
 	void AddLightNode(const UE::FUsdPrim& Prim, UInterchangeBaseNodeContainer& NodeContainer)
 	{
+#if USE_USD_SDK
+
 		FString NodeUid = LightPrefix + Prim.GetPrimPath().GetString();
 		FString NodeName(Prim.GetName().ToString());
 
-		// #ueent_todo: Retrieve light attributes and set them below
+		// Ref. UsdToUnreal::ConvertLight
+		static const FString IntensityToken = UsdToUnreal::ConvertToken(pxr::UsdLuxTokens->inputsIntensity);
+		static const FString ExposureToken = UsdToUnreal::ConvertToken(pxr::UsdLuxTokens->inputsExposure);
+		static const FString ColorToken = UsdToUnreal::ConvertToken(pxr::UsdLuxTokens->inputsColor);
 
-		if (Prim.IsA(TEXT("DistantLight")))
+		float Intensity = UsdUtils::GetAttributeValue<float>(Prim, IntensityToken);
+		float Exposure = UsdUtils::GetAttributeValue<float>(Prim, ExposureToken);
+		FLinearColor Color = UsdUtils::GetAttributeValue<FLinearColor>(Prim, ColorToken);
+
+		const bool bSRGB = true;
+		Color.ToFColor(bSRGB);
+
+		static const FString TemperatureToken = UsdToUnreal::ConvertToken(pxr::UsdLuxTokens->inputsColorTemperature);
+		static const FString UseTemperatureToken = UsdToUnreal::ConvertToken(pxr::UsdLuxTokens->inputsEnableColorTemperature);
+
+		float Temperature = UsdUtils::GetAttributeValue<float>(Prim, TemperatureToken);
+		bool UseTemperature = UsdUtils::GetAttributeValue<bool>(Prim, UseTemperatureToken);
+
+		// "Shadow enabled" currently not supported
+
+		auto SetBaseLightProperties = [&NodeUid, &NodeName, Color, Temperature, UseTemperature](UInterchangeBaseLightNode* LightNode)
 		{
-			UInterchangeDirectionalLightNode* LightNode = NewObject<UInterchangeDirectionalLightNode>(&NodeContainer);
 			LightNode->InitializeNode(NodeUid, NodeName, EInterchangeNodeContainerType::TranslatedAsset);
 			LightNode->SetAssetName(NodeName);
+				
+			LightNode->SetCustomLightColor(Color);
+			LightNode->SetCustomTemperature(Temperature);
+			LightNode->SetCustomUseTemperature(UseTemperature);
+		};
+
+		static const FString RadiusToken = UsdToUnreal::ConvertToken(pxr::UsdLuxTokens->inputsRadius);
 
-			// LightNode->SetCustomLightColor(Color);
-			// LightNode->SetCustomIntensity(Intensity);
+		if (Prim.IsA(TEXT("DistantLight")))
+		{
+			UInterchangeDirectionalLightNode* LightNode = NewObject<UInterchangeDirectionalLightNode>(&NodeContainer);
+			SetBaseLightProperties(LightNode);
+
+			Intensity = UsdToUnreal::ConvertLightIntensityAttr(Intensity, Exposure);
+			LightNode->SetCustomIntensity(Intensity);
+
+			// LightSourceAngle currently not supported by UInterchangeDirectionalLightNode
+			//float Angle = UsdUtils::GetAttributeValue<float>(Prim, TEXTVIEW("inputs:angle"));
 
 			NodeContainer.AddNode(LightNode);
 		}
 		else if (Prim.IsA(TEXT("SphereLight")))
 		{
+			const FUsdStageInfo StageInfo(Prim.GetStage());
+
+			const float Radius = UsdUtils::GetAttributeValue<float>(Prim, RadiusToken);
+			const float SourceRadius = UsdToUnreal::ConvertDistance(StageInfo, Radius); // currently not supported
+
 			if (Prim.HasAPI(TEXT("ShapingAPI")))
 			{
 				UInterchangeSpotLightNode* LightNode = NewObject<UInterchangeSpotLightNode>(&NodeContainer);
-				LightNode->InitializeNode(NodeUid, NodeName, EInterchangeNodeContainerType::TranslatedAsset);
-				LightNode->SetAssetName(NodeName);
+				SetBaseLightProperties(LightNode);
 
 				LightNode->SetCustomIntensityUnits(EInterchangeLightUnits::Lumens);
-				// LightNode->SetCustomLightColor(Color));
-				// LightNode->SetCustomIntensity(Intensity);
 
-				// LightNode->SetCustomInnerConeAngle(FMath::RadiansToDegrees(InnerConeAngle));
-				// LightNode->SetCustomOuterConeAngle(FMath::RadiansToDegrees(OuterConeAngle));
+				static const FString ConeAngleToken = UsdToUnreal::ConvertToken(pxr::UsdLuxTokens->inputsShapingConeAngle);
+				static const FString ConeSoftnessToken = UsdToUnreal::ConvertToken(pxr::UsdLuxTokens->inputsShapingConeSoftness);
+
+				float ConeAngle = UsdUtils::GetAttributeValue<float>(Prim, ConeAngleToken);
+				float ConeSoftness = UsdUtils::GetAttributeValue<float>(Prim, ConeSoftnessToken);
+
+				float InnerConeAngle = 0.0f;
+				const float OuterConeAngle = UsdToUnreal::ConvertConeAngleSoftnessAttr(ConeAngle, ConeSoftness, InnerConeAngle);
+
+				Intensity = UsdToUnreal::ConvertLuxShapingAPIIntensityAttr(Intensity, Exposure, Radius, ConeAngle, ConeSoftness, StageInfo);
+				LightNode->SetCustomIntensity(Intensity);
+
+				LightNode->SetCustomInnerConeAngle(InnerConeAngle);
+				LightNode->SetCustomOuterConeAngle(OuterConeAngle);
 
 				NodeContainer.AddNode(LightNode);
 			}
 			else
 			{
 				UInterchangePointLightNode* LightNode = NewObject<UInterchangePointLightNode>(&NodeContainer);
-				LightNode->InitializeNode(NodeUid, NodeName, EInterchangeNodeContainerType::TranslatedAsset);
-				LightNode->SetAssetName(NodeName);
+				SetBaseLightProperties(LightNode);
 
 				LightNode->SetCustomIntensityUnits(EInterchangeLightUnits::Lumens);
-				// LightNode->SetCustomLightColor(FLinearColor(Color));
-				// LightNode->SetCustomIntensity(Intensity);
 
-				// LightNode->SetCustomAttenuationRadius(AttenuationRadius);
+				Intensity = UsdToUnreal::ConvertSphereLightIntensityAttr(Intensity, Exposure, Radius, StageInfo);
+				LightNode->SetCustomIntensity(Intensity);
 
 				NodeContainer.AddNode(LightNode);
 			}
@@ -487,22 +535,46 @@ namespace UE::InterchangeUsdTranslator::Private
 		else if (Prim.IsA(TEXT("RectLight")) || Prim.IsA(TEXT("DiskLight")))
 		{
 			UInterchangeRectLightNode* LightNode = NewObject<UInterchangeRectLightNode>(&NodeContainer);
-			LightNode->InitializeNode(NodeUid, NodeName, EInterchangeNodeContainerType::TranslatedAsset);
-			LightNode->SetAssetName(NodeName);
+			SetBaseLightProperties(LightNode);
+
+			LightNode->SetCustomIntensityUnits(EInterchangeLightUnits::Lumens);
+
+			static const FString WidthToken = UsdToUnreal::ConvertToken(pxr::UsdLuxTokens->inputsWidth);
+			static const FString HeightToken = UsdToUnreal::ConvertToken(pxr::UsdLuxTokens->inputsHeight);
+
+			float Width = UsdUtils::GetAttributeValue<float>(Prim, WidthToken);
+			float Height = UsdUtils::GetAttributeValue<float>(Prim, HeightToken);
+
+			const FUsdStageInfo StageInfo(Prim.GetStage());
+
+			if (Prim.IsA(TEXT("RectLight")))
+			{
+				Width = UsdToUnreal::ConvertDistance(StageInfo, Width);
+				Height = UsdToUnreal::ConvertDistance(StageInfo, Height);
+				Intensity = UsdToUnreal::ConvertRectLightIntensityAttr(Intensity, Exposure, Width, Height, StageInfo);
+			}
+			else
+			{
+				float Radius = UsdUtils::GetAttributeValue<float>(Prim, RadiusToken);
+				Width = UsdToUnreal::ConvertDistance(StageInfo, Radius) * 2.f;
+				Height = Width;
+
+				Intensity = UsdToUnreal::ConvertDiskLightIntensityAttr(Intensity, Exposure, Radius, StageInfo);
+			}
+			LightNode->SetCustomIntensity(Intensity);
+			LightNode->SetCustomSourceWidth(Width);
+			LightNode->SetCustomSourceHeight(Height);
+
 			NodeContainer.AddNode(LightNode);
 		}
 		// #ueent_todo:
 		// DomeLight -> SkyLight
+#endif	  // USE_USD_SDK
 	}
 
 	void AddCameraNode(const UE::FUsdPrim& Prim, UInterchangeBaseNodeContainer& NodeContainer)
 	{
-		// FUsdCamera Camera;
-		// if (!UsdToUnreal::ConvertGeomCamera(Prim, Camera, UsdUtils::GetEarliestTimeCode()))
-		//{
-		//	return;
-		// }
-
+#if USE_USD_SDK
 		FString NodeUid = CameraPrefix + Prim.GetPrimPath().GetString();
 		FString NodeName(Prim.GetName().ToString());
 
@@ -510,9 +582,28 @@ namespace UE::InterchangeUsdTranslator::Private
 		CameraNode->InitializeNode(NodeUid, NodeName, EInterchangeNodeContainerType::TranslatedAsset);
 		NodeContainer.AddNode(CameraNode);
 
-		// CameraNode->SetCustomFocalLength(Camera.FocalLength);
-		// CameraNode->SetCustomSensorHeight(Camera.SensorHeight);
-		// CameraNode->SetCustomSensorWidth(Camera.SensorWidth);
+		// ref. UsdToUnreal::ConvertGeomCamera
+		UE::FUsdStage Stage = Prim.GetStage();
+		FUsdStageInfo StageInfo(Stage);
+
+		static const FString FocalLengthToken = UsdToUnreal::ConvertToken(pxr::UsdGeomTokens->focalLength);
+		static const FString HorizontalApertureToken = UsdToUnreal::ConvertToken(pxr::UsdGeomTokens->horizontalAperture);
+		static const FString VerticalApertureToken = UsdToUnreal::ConvertToken(pxr::UsdGeomTokens->verticalAperture);
+
+		float FocalLength = UsdUtils::GetAttributeValue<float>(Prim, FocalLengthToken);
+		FocalLength = UsdToUnreal::ConvertDistance(StageInfo, FocalLength);
+		CameraNode->SetCustomFocalLength(FocalLength);
+
+		float SensorWidth = UsdUtils::GetAttributeValue<float>(Prim, HorizontalApertureToken);
+		SensorWidth = UsdToUnreal::ConvertDistance(StageInfo, SensorWidth);
+		CameraNode->SetCustomSensorWidth(SensorWidth);
+
+		float SensorHeight = UsdUtils::GetAttributeValue<float>(Prim, VerticalApertureToken);
+		SensorHeight = UsdToUnreal::ConvertDistance(StageInfo, SensorHeight);
+		CameraNode->SetCustomSensorHeight(SensorHeight);
+
+		// Focus distance and FStop not currently supported
+#endif	  // USE_USD_SDK
 	}
 
 	void AddMeshNode(const UE::FUsdPrim& Prim, UInterchangeUSDTranslatorImpl* TranslatorImpl, UInterchangeBaseNodeContainer& NodeContainer)
@@ -663,7 +754,7 @@ namespace UE::InterchangeUsdTranslator::Private
 		FMeshDescription& OutMeshDescription
 	)
 	{
-		UE::FUsdPrim Prim = UsdStage.GetPrimAtPath(UE::FSdfPath{*PrimPath});
+		UE::FUsdPrim Prim = UsdStage.GetPrimAtPath(UE::FSdfPath{ *PrimPath });
 		if (!Prim)
 		{
 			return false;
@@ -920,9 +1011,9 @@ TOptional<UE::Interchange::FImportBlockedImage> UInterchangeUSDTranslator::GetBl
 	TextureGroup TextureGroup;
 	bool bDecoded = DecodeTexturePayloadKey(PayloadKey, FilePath, TextureGroup);
 	if (!bDecoded)
-	{
-		return {};
-	}
+{
+	return {};
+}
 
 	AlternateTexturePath = FilePath;
 
diff --git a/Engine/Plugins/Runtime/USDCore/Source/UnrealUSDWrapper/Private/UsdWrappers/UsdAttribute.cpp b/Engine/Plugins/Runtime/USDCore/Source/UnrealUSDWrapper/Private/UsdWrappers/UsdAttribute.cpp
--- a/Engine/Plugins/Runtime/USDCore/Source/UnrealUSDWrapper/Private/UsdWrappers/UsdAttribute.cpp
+++ b/Engine/Plugins/Runtime/USDCore/Source/UnrealUSDWrapper/Private/UsdWrappers/UsdAttribute.cpp
@@ -365,6 +365,67 @@ namespace UE
 		return FUsdPrim(Impl->PxrUsdAttribute.Get().GetPrim());
 #else
 		return FUsdPrim();
+#endif	  // #if USE_USD_SDK
+	}
+
+	template <typename T>
+	bool FUsdAttribute::Get(T& Value, TOptional<double> Time) const 
+	{
+#if USE_USD_SDK
+		FScopedUsdAllocs UsdAllocs;
+
+		pxr::UsdTimeCode TimeCode = Time.IsSet() ? Time.GetValue() : pxr::UsdTimeCode::Default();
+		return Impl->PxrUsdAttribute.Get().Get(&Value, TimeCode);
+#else
+		return false;
+#endif	  // #if USE_USD_SDK
+	}
+
+	template UNREALUSDWRAPPER_API bool FUsdAttribute::Get(float& Value, TOptional<double> Time) const;
+	template UNREALUSDWRAPPER_API bool FUsdAttribute::Get(bool& Value, TOptional<double> Time) const;
+
+	template <>
+	UNREALUSDWRAPPER_API bool FUsdAttribute::Get(FLinearColor& Value, TOptional<double> Time) const
+	{
+#if USE_USD_SDK
+		FScopedUsdAllocs UsdAllocs;
+
+		pxr::UsdTimeCode TimeCode = Time.IsSet() ? Time.GetValue() : pxr::UsdTimeCode::Default();
+		pxr::GfVec3f PxrValue;
+		bool bSuccess = Impl->PxrUsdAttribute.Get().Get(&PxrValue, TimeCode);
+		if (bSuccess)
+		{
+			Value = FLinearColor(PxrValue[0], PxrValue[1], PxrValue[2]);
+		}
+		return bSuccess;
+#else
+		return false;
 #endif	  // #if USE_USD_SDK
 	}
 }	 // namespace UE
+
+namespace UsdUtils
+{
+	template<typename ValueType>
+	ValueType GetAttributeValue(const UE::FUsdPrim& Prim, FStringView AttributeName, TOptional<double> Time)
+	{
+		if (!Prim)
+		{
+			return {};
+		}
+
+		UE::FUsdAttribute Attribute = Prim.GetAttribute(AttributeName.GetData());
+
+		ValueType Value{};
+		if (Attribute)
+		{
+			Attribute.Get(Value, Time);
+		}
+
+		return Value;
+	}
+
+	template UNREALUSDWRAPPER_API float GetAttributeValue(const UE::FUsdPrim& Prim, FStringView AttributeName, TOptional<double> Time);
+	template UNREALUSDWRAPPER_API bool GetAttributeValue(const UE::FUsdPrim& Prim, FStringView AttributeName, TOptional<double> Time);
+	template UNREALUSDWRAPPER_API FLinearColor GetAttributeValue(const UE::FUsdPrim& Prim, FStringView AttributeName, TOptional<double> Time);
+}
\ No newline at end of file
diff --git a/Engine/Plugins/Runtime/USDCore/Source/UnrealUSDWrapper/Public/UsdWrappers/UsdAttribute.h b/Engine/Plugins/Runtime/USDCore/Source/UnrealUSDWrapper/Public/UsdWrappers/UsdAttribute.h
--- a/Engine/Plugins/Runtime/USDCore/Source/UnrealUSDWrapper/Public/UsdWrappers/UsdAttribute.h
+++ b/Engine/Plugins/Runtime/USDCore/Source/UnrealUSDWrapper/Public/UsdWrappers/UsdAttribute.h
@@ -4,6 +4,7 @@
 
 #include "Misc/Optional.h"
 #include "Templates/UniquePtr.h"
+#include "Containers/StringView.h"
 
 #if USE_USD_SDK
 #include "USDIncludesStart.h"
@@ -90,6 +91,9 @@ namespace UE
 		bool Get(UE::FVtValue& Value, TOptional<double> Time = {}) const;
 		bool Set(const UE::FVtValue& Value, TOptional<double> Time = {}) const;
 
+		template <typename T>
+		bool Get(T& Value, TOptional<double> Time = {}) const;
+
 		bool Clear() const;
 		bool ClearAtTime(double Time) const;
 
@@ -102,3 +106,10 @@ namespace UE
 		TUniquePtr<Internal::FUsdAttributeImpl> Impl;
 	};
 }	 // namespace UE
+
+namespace UsdUtils
+{
+	// Returns the attribute value for the given prim and attribute name. If the attribute doesn't exist, it returns the default value for the type.
+	template<typename ValueType>
+	UNREALUSDWRAPPER_API ValueType GetAttributeValue(const UE::FUsdPrim& Prim, FStringView AttributeName, TOptional<double> Time = {});
+}
\ No newline at end of file
