USD: Added cvar to control support for USD Interchange level import (disabled by default).

diff --git a/Engine/Plugins/Interchange/Runtime/Source/Import/Private/Usd/InterchangeUsdTranslator.cpp b/Engine/Plugins/Interchange/Runtime/Source/Import/Private/Usd/InterchangeUsdTranslator.cpp
--- a/Engine/Plugins/Interchange/Runtime/Source/Import/Private/Usd/InterchangeUsdTranslator.cpp
+++ b/Engine/Plugins/Interchange/Runtime/Source/Import/Private/Usd/InterchangeUsdTranslator.cpp
@@ -58,6 +58,13 @@ static FAutoConsoleVariableRef CVarInterchangeEnableUSDImport(
 	TEXT("Whether USD support is enabled.")
 );
 
+static bool GInterchangeEnableUSDLevelImport = false;
+static FAutoConsoleVariableRef CVarInterchangeEnableUSDLevelImport(
+	TEXT("Interchange.FeatureFlags.Import.USD.ToLevel"),
+	GInterchangeEnableUSDLevelImport,
+	TEXT("Whether support for USD level import is enabled.")
+);
+
 namespace UE::InterchangeUsdTranslator::Private
 {
 	const static FString CameraPrefix = TEXT("\\Camera\\");
@@ -106,7 +113,7 @@ UInterchangeUSDTranslator::UInterchangeUSDTranslator()
 
 EInterchangeTranslatorType UInterchangeUSDTranslator::GetTranslatorType() const
 {
-	return EInterchangeTranslatorType::Scenes;
+	return GInterchangeEnableUSDLevelImport ? EInterchangeTranslatorType::Scenes : EInterchangeTranslatorType::Assets;
 }
 
 EInterchangeTranslatorAssetType UInterchangeUSDTranslator::GetSupportedAssetTypes() const
