USD: Add support for translating skeleton, skeletal meshes and morph targets via Interchange;

Add a whole bunch of UsdSkel wrapper classes that we can use from Interchange;
Fix copy constructors of valid FUsdGeomXformable instances returning invalid prims;
Fix prims inside of Scopes not being translated as assets (e.g. Material prims);
Revert named material slot change from cl 34269265 as naming slots after their current assignments works better with Interchange.

diff --git a/Engine/Plugins/Importers/USDImporter/Source/USDSchemas/Private/USDSkelRootTranslator.cpp b/Engine/Plugins/Importers/USDImporter/Source/USDSchemas/Private/USDSkelRootTranslator.cpp
--- a/Engine/Plugins/Importers/USDImporter/Source/USDSchemas/Private/USDSkelRootTranslator.cpp
+++ b/Engine/Plugins/Importers/USDImporter/Source/USDSchemas/Private/USDSkelRootTranslator.cpp
@@ -24,6 +24,7 @@
 #include "UsdWrappers/SdfLayer.h"
 #include "UsdWrappers/SdfPath.h"
 #include "UsdWrappers/UsdPrim.h"
+#include "UsdWrappers/UsdSkelSkinningQuery.h"
 
 #include "Animation/AnimBlueprint.h"
 #include "Animation/AnimInstance.h"
@@ -526,7 +527,7 @@ namespace UsdSkelRootTranslatorImpl
 																	  &StageInfo,
 																	  Options](const pxr::UsdGeomMesh& LODMesh, int32 LODIndex)
 		{
-			pxr::UsdSkelSkinningQuery SkinningQuery = UsdUtils::CreateSkinningQuery(LODMesh, SkelQuery);
+			pxr::UsdSkelSkinningQuery SkinningQuery = UsdUtils::CreateSkinningQuery(LODMesh.GetPrim(), SkelQuery);
 			if (!SkinningQuery)
 			{
 				return true;	// Continue trying other LODs
diff --git a/Engine/Plugins/Importers/USDImporter/Source/USDSchemas/Private/USDSkelSkeletonTranslator.cpp b/Engine/Plugins/Importers/USDImporter/Source/USDSchemas/Private/USDSkelSkeletonTranslator.cpp
--- a/Engine/Plugins/Importers/USDImporter/Source/USDSchemas/Private/USDSkelSkeletonTranslator.cpp
+++ b/Engine/Plugins/Importers/USDImporter/Source/USDSchemas/Private/USDSkelSkeletonTranslator.cpp
@@ -26,6 +26,7 @@
 #include "UsdWrappers/SdfLayer.h"
 #include "UsdWrappers/SdfPath.h"
 #include "UsdWrappers/UsdPrim.h"
+#include "UsdWrappers/UsdSkelSkinningQuery.h"
 
 #include "Animation/AnimBlueprint.h"
 #include "Animation/AnimInstance.h"
@@ -584,7 +585,7 @@ namespace UsdSkelSkeletonTranslatorImpl
 				return true;	// Continue trying other LODs
 			}
 
-			pxr::UsdSkelSkinningQuery SkinningQuery = UsdUtils::CreateSkinningQuery(LODMesh, SkeletonQuery);
+			pxr::UsdSkelSkinningQuery SkinningQuery = UsdUtils::CreateSkinningQuery(LODMesh.GetPrim(), SkeletonQuery);
 			if (!SkinningQuery)
 			{
 				return true;
diff --git a/Engine/Plugins/Interchange/Runtime/Source/Import/Private/Usd/InterchangeUsdTranslator.cpp b/Engine/Plugins/Interchange/Runtime/Source/Import/Private/Usd/InterchangeUsdTranslator.cpp
--- a/Engine/Plugins/Interchange/Runtime/Source/Import/Private/Usd/InterchangeUsdTranslator.cpp
+++ b/Engine/Plugins/Interchange/Runtime/Source/Import/Private/Usd/InterchangeUsdTranslator.cpp
@@ -1,4 +1,4 @@
-// Copyright Epic Games, Inc. All Rights Reserved. 
+// Copyright Epic Games, Inc. All Rights Reserved.
 
 #include "Usd/InterchangeUsdTranslator.h"
 
@@ -7,20 +7,27 @@
 #include "UsdWrappers/UsdAttribute.h"
 #include "UsdWrappers/UsdGeomXformable.h"
 #include "UsdWrappers/UsdPrim.h"
+#include "UsdWrappers/UsdRelationship.h"
+#include "UsdWrappers/UsdSkelBlendShape.h"
+#include "UsdWrappers/UsdSkelBlendShapeQuery.h"
+#include "UsdWrappers/UsdSkelCache.h"
+#include "UsdWrappers/UsdSkelInbetweenShape.h"
+#include "UsdWrappers/UsdSkelSkeletonQuery.h"
+#include "UsdWrappers/UsdSkelSkinningQuery.h"
 #include "UsdWrappers/UsdStage.h"
 #include "UsdWrappers/UsdTyped.h"
-#include "UsdWrappers/VtValue.h"
 
+#include "USDClassesModule.h"
 #include "USDConversionUtils.h"
 #include "USDGeomMeshConversion.h"
+#include "USDLightConversion.h"
 #include "USDLog.h"
-#include "USDLightConversion.h"	
 #include "USDMaterialUtils.h"
 #include "USDPrimConversion.h"
 #include "USDShadeConversion.h"
+#include "USDSkeletalDataConversion.h"
 #include "USDStageOptions.h"
 #include "USDTypesConversion.h"
-#include "USDValueConversion.h"
 
 #include "Async/Async.h"
 #include "HAL/IConsoleManager.h"
@@ -33,6 +40,7 @@
 #include "InterchangeTexture2DNode.h"
 #include "InterchangeTranslatorHelper.h"
 #include "Internationalization/Regex.h"
+#include "Rendering/SkeletalMeshLODImporterData.h"
 #include "StaticMeshAttributes.h"
 #include "StaticMeshOperations.h"
 #include "UDIMUtilities.h"
@@ -44,7 +52,6 @@
 #include "pxr/usd/usdPhysics/tokens.h"
 #include "pxr/usd/usdShade/tokens.h"
 #include "USDIncludesEnd.h"
-
 #endif	  // USE_USD_SDK
 
 #include UE_INLINE_GENERATED_CPP_BY_NAME(InterchangeUsdTranslator)
@@ -71,6 +78,19 @@ namespace UE::InterchangeUsdTranslator::Private
 	const static FString LightPrefix = TEXT("\\Light\\");
 	const static FString MaterialPrefix = TEXT("\\Material\\");
 	const static FString MeshPrefix = TEXT("\\Mesh\\");
+	const static FString MorphTargetPrefix = TEXT("\\MorphTarget\\");
+
+	// Information intended to be passed down from parent to children (by value) as we traverse the stage
+	struct FTraversalInfo
+	{
+		UInterchangeBaseNode* ParentNode = nullptr;
+
+		TSharedPtr<UE::FUsdSkelCache> FurthestSkelCache;
+		UE::FUsdPrim ClosestParentSkelRoot;
+
+		UE::FUsdSkelSkeletonQuery ActiveSkelQuery;
+		TSharedPtr<TArray<FString>> SkelJointNames;	   // Needed for skel mesh payloads
+	};
 
 	// Small container that we can use Pimpl with so we don't have to include too many USD includes on the header file.
 	//
@@ -91,8 +111,15 @@ namespace UE::InterchangeUsdTranslator::Private
 		// we can reuse it (otherwise we have to keep converting the FNames into Tokens all the time)
 		UsdToUnreal::FUsdMeshConversionOptions CachedMeshConversionOptions;
 #endif
+
+		// When traversing we'll generate FTraversalInfo objects. If we need to (e.g. for skinned meshes),
+		// we'll store the info for that translated node here, so we don't have to recompute it when returning
+		// the payload data.
+		// Note: We only do this when needed: This shouldn't have data for every prim in the stage.
+		TMap<FString, FTraversalInfo> NodeUidToCachedTraversalInfo;
+		mutable FRWLock CachedTraversalInfoLock;
 	};
-}
+}	 // namespace UE::InterchangeUsdTranslator::Private
 
 UInterchangeUsdTranslatorSettings::UInterchangeUsdTranslatorSettings()
 	: GeometryPurpose((int32)(EUsdPurpose::Default | EUsdPurpose::Proxy | EUsdPurpose::Render | EUsdPurpose::Guide))
@@ -131,9 +158,9 @@ TArray<FString> UInterchangeUSDTranslator::GetSupportedFormats() const
 	return Extensions;
 }
 
+#if USE_USD_SDK
 namespace UE::InterchangeUsdTranslator::Private
 {
-#if USE_USD_SDK
 	FString EncodeTexturePayloadKey(const UsdToUnreal::FTextureParameterValue& Value)
 	{
 		// Encode the compression settings onto the payload key as we need to move that into the
@@ -474,7 +501,7 @@ namespace UE::InterchangeUsdTranslator::Private
 		{
 			LightNode->InitializeNode(NodeUid, NodeName, EInterchangeNodeContainerType::TranslatedAsset);
 			LightNode->SetAssetName(NodeName);
-				
+
 			LightNode->SetCustomLightColor(Color);
 			LightNode->SetCustomTemperature(Temperature);
 			LightNode->SetCustomUseTemperature(UseTemperature);
@@ -491,7 +518,7 @@ namespace UE::InterchangeUsdTranslator::Private
 			LightNode->SetCustomIntensity(Intensity);
 
 			// LightSourceAngle currently not supported by UInterchangeDirectionalLightNode
-			//float Angle = UsdUtils::GetAttributeValue<float>(Prim, TEXTVIEW("inputs:angle"));
+			// float Angle = UsdUtils::GetAttributeValue<float>(Prim, TEXTVIEW("inputs:angle"));
 
 			NodeContainer.AddNode(LightNode);
 		}
@@ -500,7 +527,7 @@ namespace UE::InterchangeUsdTranslator::Private
 			const FUsdStageInfo StageInfo(Prim.GetStage());
 
 			const float Radius = UsdUtils::GetAttributeValue<float>(Prim, RadiusToken);
-			const float SourceRadius = UsdToUnreal::ConvertDistance(StageInfo, Radius); // currently not supported
+			const float SourceRadius = UsdToUnreal::ConvertDistance(StageInfo, Radius);	   // currently not supported
 
 			if (Prim.HasAPI(TEXT("ShapingAPI")))
 			{
@@ -613,7 +640,74 @@ namespace UE::InterchangeUsdTranslator::Private
 #endif	  // USE_USD_SDK
 	}
 
-	void AddMeshNode(const UE::FUsdPrim& Prim, UInterchangeUSDTranslatorImpl* TranslatorImpl, UInterchangeBaseNodeContainer& NodeContainer)
+	void AddMorphTargetNodes(
+		const UE::FUsdPrim& Prim,
+		UInterchangeUSDTranslatorImpl* TranslatorImpl,
+		UInterchangeMeshNode& MeshNode,
+		UInterchangeBaseNodeContainer& NodeContainer
+	)
+	{
+		UE::FUsdSkelBlendShapeQuery Query{Prim};
+		if (!Query)
+		{
+			return;
+		}
+
+		const FString MeshPrimPath = Prim.GetPrimPath().GetString();
+
+		TFunction<void(const FString&, int32, const FString&)> AddMorphTargetNode =
+			[&MeshNode, &MeshPrimPath, &NodeContainer](const FString& MorphTargetName, int32 BlendShapeIndex, const FString& InbetweenName)
+		{
+			// Note: We identify a blend shape by its Mesh prim path and the blend shape index, even though
+			// the blend shape itself is a full standalone prim. This is for two reasons:
+			//  - We need to also read the Mesh prim's mesh data when emitting the payload, so having the Mesh path on the payload key is handy;
+			//  - It could be possible for different meshes to share the same BlendShape (possibly?), so we really want a separate version of
+			//    a blend shape for each mesh that uses it.
+			//
+			// Despite of that though, we'll just use the blendshape (+inbetween) name for MorphTargetName (so not anything's full path),
+			// so that users can get different blendshapes across the model to combine into a single morph target. Interchange has
+			// an import option to let you control whether they become separate morph targets or not anyway
+			// ("Merge Morph Targets with Same Name")
+			FString IDString = FString::Printf(TEXT("%s\\%d\\%s"), *MeshPrimPath, BlendShapeIndex, *InbetweenName);
+			FString MorphTargetUid = MorphTargetPrefix + IDString;
+
+			UInterchangeMeshNode* MorphTargetMeshNode = NewObject<UInterchangeMeshNode>(&NodeContainer);
+			MorphTargetMeshNode->InitializeNode(MorphTargetUid, MorphTargetName, EInterchangeNodeContainerType::TranslatedAsset);
+			MorphTargetMeshNode->SetPayLoadKey(IDString, EInterchangeMeshPayLoadType::MORPHTARGET);
+			MorphTargetMeshNode->SetMorphTarget(true);
+			MorphTargetMeshNode->SetMorphTargetName(MorphTargetName);
+			NodeContainer.AddNode(MorphTargetMeshNode);
+			MeshNode.SetMorphTargetDependencyUid(MorphTargetUid);
+		};
+
+		for (size_t Index = 0; Index < Query.GetNumBlendShapes(); ++Index)
+		{
+			UE::FUsdSkelBlendShape BlendShape = Query.GetBlendShape(Index);
+			if (!BlendShape)
+			{
+				continue;
+			}
+
+			UE::FUsdPrim BlendShapePrim = BlendShape.GetPrim();
+			FString MorphTargetName = BlendShapePrim.GetName().ToString();
+			AddMorphTargetNode(MorphTargetName, Index, FString{});
+
+			TArray<UE::FUsdSkelInbetweenShape> Inbetweens = BlendShape.GetInbetweens();
+			for (const UE::FUsdSkelInbetweenShape& Inbetween : Inbetweens)
+			{
+				FString InbetweenName = Inbetween.GetAttr().GetName().ToString();
+				FString InnerMorphTargetName = MorphTargetName + TEXT("_") + InbetweenName;
+				AddMorphTargetNode(InnerMorphTargetName, Index, InbetweenName);
+			}
+		}
+	}
+
+	void AddMeshNode(
+		const UE::FUsdPrim& Prim,
+		UInterchangeUSDTranslatorImpl* TranslatorImpl,
+		UInterchangeBaseNodeContainer& NodeContainer,
+		const FTraversalInfo& Info
+	)
 	{
 		FString PrimPath = Prim.GetPrimPath().GetString();
 		FString NodeUid = MeshPrefix + PrimPath;
@@ -625,10 +719,33 @@ namespace UE::InterchangeUsdTranslator::Private
 			return;
 		}
 
+		// Fill in the MeshNode itself
 		UInterchangeMeshNode* MeshNode = NewObject<UInterchangeMeshNode>(&NodeContainer);
 		MeshNode->InitializeNode(NodeUid, NodeName, EInterchangeNodeContainerType::TranslatedAsset);
-		MeshNode->SetPayLoadKey(PrimPath, EInterchangeMeshPayLoadType::STATIC);
 		MeshNode->SetAssetName(NodeName);
+		const bool bIsSkinned = static_cast<bool>(Info.ClosestParentSkelRoot) && Prim.HasAPI(TEXT("SkelBindingAPI"));
+		if (bIsSkinned)
+		{
+			MeshNode->SetSkinnedMesh(true);
+			MeshNode->SetPayLoadKey(PrimPath, EInterchangeMeshPayLoadType::SKELETAL);
+			if (Info.ActiveSkelQuery)
+			{
+				MeshNode->SetSkeletonDependencyUid(Info.ActiveSkelQuery.GetSkeleton().GetPrimPath().GetString());
+			}
+
+			AddMorphTargetNodes(Prim, TranslatorImpl, *MeshNode, NodeContainer);
+
+			// When returning the payload data later, we'll need at the very least our SkeletonQuery, so
+			// here we store the Info object into the Impl
+			{
+				FWriteScopeLock ScopedInfoWriteLock{TranslatorImpl->CachedTraversalInfoLock};
+				TranslatorImpl->NodeUidToCachedTraversalInfo.Add(NodeUid, Info);
+			}
+		}
+		else
+		{
+			MeshNode->SetPayLoadKey(PrimPath, EInterchangeMeshPayLoadType::STATIC);
+		}
 
 		// Material assignments
 		{
@@ -644,7 +761,13 @@ namespace UE::InterchangeUsdTranslator::Private
 
 			for (const UsdUtils::FUsdPrimMaterialSlot& Slot : Assignments.Slots)
 			{
-				const FString& SlotName = Slot.SlotName;
+				// Use the material prim path/display color desc as the material slot name, because Interchange
+				// already has a mechanism to merge material slots with the same name. Using the material name itself
+				// as the slot name has Interchange combine slots with identical materials, which works fine. If we
+				// were to use GeomSubset names or prim names in here though, it's possible that two similarly named
+				// slots in different skeletal mesh chunks (but with different materials!) could get merged together,
+				// which is not what we want
+				const FString& SlotName = Slot.MaterialSource;
 
 				// Get the Uid of the material instance that we'll end up assigning to this slot
 				FString MaterialInstanceUid;
@@ -681,12 +804,135 @@ namespace UE::InterchangeUsdTranslator::Private
 		NodeContainer.AddNode(MeshNode);
 	}
 
+	void AddSkeletonNodes(
+		const UE::FUsdPrim& Prim,
+		UInterchangeUSDTranslatorImpl* TranslatorImpl,
+		UInterchangeSceneNode& SkeletonPrimNode,
+		UInterchangeBaseNodeContainer& NodeContainer,
+		FTraversalInfo Info
+	)
+	{
+		// If we're not inside of a SkelRoot, the skeleton shouldn't really do anything
+		if (!Info.FurthestSkelCache.IsValid())
+		{
+			return;
+		}
+
+		// By the time we get here we've already emitted a scene node for the skeleton prim itself, so we just
+		// need to emit a node hierarchy that mirrors the joints.
+
+		// Make the prim node into an Interchange joint/bone itself. By doing this we solve three issues:
+		//  - It becomes easy to identify our SkeletonDependencyUid when parsing Mesh nodes: It's just the skeleton prim path
+		//    (as opposed to having to target the translated node of the first root joint of the skeleton);
+		//  - We automatically handle USD skeletons with multiple root bones: We'll only ever have one "true"
+		//    root bone anyway: The SkeletonPrimNode itself;
+		//  - If a skeleton has no bones at all somehow, we'll still make one "bone" for it (this node).
+		SkeletonPrimNode.AddSpecializedType(UE::Interchange::FSceneNodeStaticData::GetJointSpecializeTypeString());
+		SkeletonPrimNode.SetCustomBindPoseLocalTransform(&NodeContainer, FTransform::Identity);
+		SkeletonPrimNode.SetCustomTimeZeroLocalTransform(&NodeContainer, FTransform::Identity);
+
+		// Convert the skeleton bones/joints into ConvertedData
+		UE::FUsdSkelSkeletonQuery SkelQuery = Info.FurthestSkelCache->GetSkelQuery(Prim);
+		const bool bEnsureAtLeastOneBone = false;
+		const bool bEnsureSingleRootBone = false;
+		UsdToUnreal::FUsdSkeletonData ConvertedData;
+		const bool bSuccess = UsdToUnreal::ConvertSkeleton(SkelQuery, ConvertedData, bEnsureAtLeastOneBone, bEnsureSingleRootBone);
+		if (!bSuccess)
+		{
+			return;
+		}
+
+		// Recursively traverse ConvertedData spawning the joint translated nodes
+		TFunction<void(const UsdToUnreal::FUsdSkeletonData::FBone&, UInterchangeSceneNode&, const FString&)> RecursiveTraverseBones = nullptr;
+		RecursiveTraverseBones = [&RecursiveTraverseBones,
+								  &ConvertedData,
+								  &NodeContainer	//
+		](const UsdToUnreal::FUsdSkeletonData::FBone& Bone, UInterchangeSceneNode& ParentNode, const FString& BonePath)
+		{
+			// Concatenate a full "bone path" here for uniqueness, because Bone.Name is just the name of this
+			// single bone/joint itself (e.g. "Elbow")
+			const FString BoneNodeUid = BonePath + TEXT("\\") + Bone.Name;
+
+			UInterchangeSceneNode* BoneNode = NewObject<UInterchangeSceneNode>(&NodeContainer);
+			BoneNode->InitializeNode(BoneNodeUid, Bone.Name, EInterchangeNodeContainerType::TranslatedScene);
+			BoneNode->SetCustomLocalTransform(&NodeContainer, Bone.LocalRestTransform);
+
+			BoneNode->AddSpecializedType(UE::Interchange::FSceneNodeStaticData::GetJointSpecializeTypeString());
+			BoneNode->SetCustomBindPoseLocalTransform(&NodeContainer, Bone.LocalBindTransform);
+			BoneNode->SetCustomTimeZeroLocalTransform(&NodeContainer, Bone.LocalRestTransform);
+
+			NodeContainer.AddNode(BoneNode);
+			NodeContainer.SetNodeParentUid(BoneNodeUid, ParentNode.GetUniqueID());
+
+			for (int32 ChildIndex : Bone.ChildIndices)
+			{
+				RecursiveTraverseBones(ConvertedData.Bones[ChildIndex], *BoneNode, BoneNodeUid);
+			}
+		};
+
+		// Start traversing from the root bones (we may have more than one)
+		const FString SkeletonPrimNodeUid = SkeletonPrimNode.GetUniqueID();
+		for (const UsdToUnreal::FUsdSkeletonData::FBone& Bone : ConvertedData.Bones)
+		{
+			if (Bone.ParentIndex == INDEX_NONE)
+			{
+				RecursiveTraverseBones(Bone, SkeletonPrimNode, SkeletonPrimNodeUid);
+			}
+		}
+
+		// Cache our joint names in order, as this is needed when generating skeletal mesh payloads
+		Info.SkelJointNames = MakeShared<TArray<FString>>();
+		Info.SkelJointNames->Reserve(ConvertedData.Bones.Num());
+		for (const UsdToUnreal::FUsdSkeletonData::FBone& Bone : ConvertedData.Bones)
+		{
+			Info.SkelJointNames->Add(Bone.Name);
+		}
+		{
+			FWriteScopeLock ScopedInfoWriteLock{TranslatorImpl->CachedTraversalInfoLock};
+			TranslatorImpl->NodeUidToCachedTraversalInfo.Add(SkeletonPrimNodeUid, Info);
+		}
+	}
+
+	void UpdateTraversalInfo(FTraversalInfo& Info, const UE::FUsdPrim& CurrentPrim)
+	{
+		if (CurrentPrim.IsA(TEXT("SkelRoot")))
+		{
+			if (!Info.ClosestParentSkelRoot)
+			{
+				// The root-most skel cache should handle any nested UsdSkel prims as well
+				Info.FurthestSkelCache = MakeShared<UE::FUsdSkelCache>();
+
+				const bool bTraverseInstanceProxies = true;
+				Info.FurthestSkelCache->Populate(CurrentPrim, bTraverseInstanceProxies);
+			}
+
+			Info.ClosestParentSkelRoot = CurrentPrim;
+		}
+
+		if (Info.ClosestParentSkelRoot && CurrentPrim.HasAPI(TEXT("SkelBindingAPI")))
+		{
+			UE::FUsdStage Stage = CurrentPrim.GetStage();
+
+			if (UE::FUsdRelationship SkelRel = CurrentPrim.GetRelationship(TEXT("skel:skeleton")))
+			{
+				TArray<UE::FSdfPath> Targets;
+				if (SkelRel.GetTargets(Targets) && Targets.Num() > 0)
+				{
+					UE::FUsdPrim TargetSkeleton = Stage.GetPrimAtPath(Targets[0]);
+					if (TargetSkeleton && TargetSkeleton.IsA(TEXT("Skeleton")))
+					{
+						Info.ActiveSkelQuery = Info.FurthestSkelCache->GetSkelQuery(TargetSkeleton);
+					}
+				}
+			}
+		}
+	}
+
 	void Traverse(
-		const UE::FUsdStage& UsdStage,
 		const UE::FUsdPrim& Prim,
 		UInterchangeUSDTranslatorImpl* TranslatorImpl,
 		UInterchangeBaseNodeContainer& NodeContainer,
-		UInterchangeBaseNode* ParentNode
+		FTraversalInfo Info
 	)
 	{
 		// Ignore prim subtrees from disabled purposes
@@ -697,9 +943,12 @@ namespace UE::InterchangeUsdTranslator::Private
 			return;
 		}
 
-		FString NodeUid = Prim.GetPrimPath().GetString();
+		FString SceneNodeUid = Prim.GetPrimPath().GetString();
 		FString DisplayLabel(Prim.GetName().ToString());
 
+		// Do this before generating other nodes as they may need the updated info
+		UpdateTraversalInfo(Info, Prim);
+
 		// Generate asset node if applicable
 		const FString* Prefix = nullptr;
 		if (Prim.IsA(TEXT("Material")))
@@ -710,7 +959,7 @@ namespace UE::InterchangeUsdTranslator::Private
 		else if (Prim.IsA(TEXT("Mesh")))
 		{
 			Prefix = &MeshPrefix;
-			AddMeshNode(Prim, TranslatorImpl, NodeContainer);
+			AddMeshNode(Prim, TranslatorImpl, NodeContainer, Info);
 		}
 		else if (Prim.IsA(TEXT("Camera")))
 		{
@@ -723,45 +972,112 @@ namespace UE::InterchangeUsdTranslator::Private
 			AddLightNode(Prim, NodeContainer);
 		}
 
-		// Generate scene node if we're an Xformable
+		// Only prims that require rendering (and have a renderable parent) get a scene node.
+		// This includes Xforms but also Scopes, which are not Xformable
 		UInterchangeSceneNode* SceneNode = nullptr;
-		FTransform Transform = FTransform::Identity;
-		if (UsdToUnreal::ConvertXformable(UsdStage, UE::FUsdTyped(Prim), Transform, UsdUtils::GetDefaultTimeCode(), nullptr))
+		if (Prim.IsA(TEXT("Imageable")) && (Info.ParentNode || Prim.GetParent().IsPseudoRoot()))
 		{
 			SceneNode = NewObject<UInterchangeSceneNode>(&NodeContainer);
-			SceneNode->InitializeNode(NodeUid, DisplayLabel, EInterchangeNodeContainerType::TranslatedScene);
-			SceneNode->SetCustomLocalTransform(&NodeContainer, Transform);
+			SceneNode->InitializeNode(SceneNodeUid, DisplayLabel, EInterchangeNodeContainerType::TranslatedScene);
 			NodeContainer.AddNode(SceneNode);
 
+			// If we're an Xformable, get our transform
+			FTransform Transform = FTransform::Identity;
+			bool bResetTransformStack = false;
+			if (UsdToUnreal::ConvertXformable(
+					Prim.GetStage(),
+					UE::FUsdTyped(Prim),
+					Transform,
+					UsdUtils::GetEarliestTimeCode(),
+					&bResetTransformStack
+				))
+			{
+				SceneNode->SetCustomLocalTransform(&NodeContainer, Transform);
+			}
+
+			// Skeleton joints are separate scene nodes in Interchange, so we need to emit that node hierarchy now
+			if (Prim.IsA(TEXT("Skeleton")))
+			{
+				AddSkeletonNodes(Prim, TranslatorImpl, *SceneNode, NodeContainer, Info);
+			}
+
 			// Connect scene node and asset node
 			if (Prefix)
 			{
-				const FString AssetNodeUid = *Prefix + NodeUid;
+				const FString AssetNodeUid = *Prefix + SceneNodeUid;
 				SceneNode->SetCustomAssetInstanceUid(AssetNodeUid);
 			}
 
 			// Connect parent and child scene nodes
-			if (ParentNode)
+			if (Info.ParentNode)
 			{
-				NodeContainer.SetNodeParentUid(SceneNode->GetUniqueID(), ParentNode->GetUniqueID());
+				NodeContainer.SetNodeParentUid(SceneNode->GetUniqueID(), Info.ParentNode->GetUniqueID());
 			}
 		}
 
+		// Note: This has the effect of effectively shutting down the generation of scene nodes
+		// below any prim that is not a least an Imageable, as we check for a valid parent before
+		// generating one
+		Info.ParentNode = SceneNode;
+
 		// Recurse into child prims
 		for (const FUsdPrim& ChildPrim : Prim.GetChildren())
 		{
-			Traverse(UsdStage, ChildPrim, TranslatorImpl, NodeContainer, SceneNode);
+			Traverse(ChildPrim, TranslatorImpl, NodeContainer, Info);
 		}
 	}
 
-	bool GetStaticMeshPayloadDataForPayLoadKey(
-		const UE::FUsdStage& UsdStage,
-		const FString& PrimPath,
+	void FixSkeletalMeshDescriptionColors(FMeshDescription& MeshDescription)
+	{
+		// FSkeletalMeshImportData::GetMeshDescription() will reinterpret our Wedge FColors as linear, and put those
+		// sRGB values disguised as linear into the mesh description. This also seems to disagree with the patch on
+		// cl 32791826, so here we have to fix that up and get our mesh description colors to be actually linear...
+		//
+		// This will hopefully go away once we have our own skinned mesh to FMeshDescription conversion function.
+		//
+		// Note: Weirdly enough skeletal meshes seem to put linear colors on VertexColor output, while static meshes
+		// put sRGB colors? Maybe this is why the comment above the change on 32791826 mentions to remove the ToFColor on
+		// StaticMeshBuilder? This is overall very confusing
+		FStaticMeshAttributes Attributes(MeshDescription);
+		TVertexInstanceAttributesRef<FVector4f> VertexColor = Attributes.GetVertexInstanceColors();
+		for (FVertexInstanceID VertexInstanceID : MeshDescription.VertexInstances().GetElementIDs())
+		{
+			const FColor ActualSRGB = FLinearColor(VertexColor[VertexInstanceID]).ToFColor(false);
+			VertexColor[VertexInstanceID] = FLinearColor{ActualSRGB};
+		}
+	}
+
+	void FixMaterialSlotNames(FMeshDescription& MeshDescription, const TArray<UsdUtils::FUsdPrimMaterialSlot>& MeshAssingmentSlots)
+	{
+		// Fixup material slot names to match the material that is assigned. For Interchange it is better to have the material
+		// slot names match what is assigned into them, as it will use those names to "merge identical slots" depending on the
+		// import options.
+		//
+		// Note: These names must also match what is set via MeshNode->SetSlotMaterialDependencyUid(SlotName, MaterialUid)
+		FStaticMeshAttributes StaticMeshAttributes(MeshDescription);
+		for (int32 MaterialSlotIndex = 0; MaterialSlotIndex < StaticMeshAttributes.GetPolygonGroupMaterialSlotNames().GetNumElements();
+			 ++MaterialSlotIndex)
+		{
+			int32 MaterialIndex = 0;
+			LexFromString(MaterialIndex, *StaticMeshAttributes.GetPolygonGroupMaterialSlotNames()[MaterialSlotIndex].ToString());
+
+			if (MeshAssingmentSlots.IsValidIndex(MaterialIndex))
+			{
+				const FString Source = MeshAssingmentSlots[MaterialIndex].MaterialSource;
+				StaticMeshAttributes.GetPolygonGroupMaterialSlotNames()[MaterialSlotIndex] = *Source;
+			}
+		}
+	}
+
+	bool GetStaticMeshPayloadData(
+		const FString& PayloadKey,
+		const UInterchangeUSDTranslatorImpl& Impl,
 		const UsdToUnreal::FUsdMeshConversionOptions& Options,
 		FMeshDescription& OutMeshDescription
 	)
 	{
-		UE::FUsdPrim Prim = UsdStage.GetPrimAtPath(UE::FSdfPath{ *PrimPath });
+		const FString& PrimPath = PayloadKey;
+		UE::FUsdPrim Prim = Impl.UsdStage.GetPrimAtPath(UE::FSdfPath{*PrimPath});
 		if (!Prim)
 		{
 			return false;
@@ -770,11 +1086,163 @@ namespace UE::InterchangeUsdTranslator::Private
 		// TODO: We can't do much with these yet: There will be used to generate primvar-compatible
 		// versions of the materials that are assigned to this mesh, whenever we get a pipeline
 		UsdUtils::FUsdPrimMaterialAssignmentInfo TempMaterialInfo;
-		return UsdToUnreal::ConvertGeomMesh(Prim, OutMeshDescription, TempMaterialInfo, Options);
+		bool bSuccess = UsdToUnreal::ConvertGeomMesh(Prim, OutMeshDescription, TempMaterialInfo, Options);
+		if (!bSuccess)
+		{
+			return false;
+		}
+
+		FixMaterialSlotNames(OutMeshDescription, TempMaterialInfo.Slots);
+
+		return true;
 	}
 
-#endif
+	bool GetSkeletalMeshPayloadData(
+		const FString& PayloadKey,
+		const UInterchangeUSDTranslatorImpl& Impl,
+		const UsdToUnreal::FUsdMeshConversionOptions& Options,
+		FMeshDescription& OutMeshDescription,
+		TArray<FString>& OutJointNames
+	)
+	{
+		const FString& PrimPath = PayloadKey;
+		UE::FUsdPrim Prim = Impl.UsdStage.GetPrimAtPath(UE::FSdfPath{*PrimPath});
+		if (!Prim)
+		{
+			return false;
+		}
+
+		const FString& MeshNodeUid = MeshPrefix + Prim.GetPrimPath().GetString();
+
+		// Read these variables from the data we cached during traversal for translation
+		TSharedPtr<TArray<FString>> JointNames = nullptr;
+		UE::FUsdSkelSkeletonQuery SkelQuery;
+		{
+			FReadScopeLock ReadLock{Impl.CachedTraversalInfoLock};
+
+			const FTraversalInfo* MeshInfo = Impl.NodeUidToCachedTraversalInfo.Find(MeshNodeUid);
+			if (!MeshInfo)
+			{
+				return false;
+			}
+			SkelQuery = MeshInfo->ActiveSkelQuery;
+			if (!SkelQuery)
+			{
+				return false;
+			}
+
+			// The above fields are associated to the mesh *asset* node Uid (hence the prefix),
+			// while the joint names are associated to the skeleton *scene* node Uid, so no prefix
+			const FString SkeletonNodeUid = SkelQuery.GetSkeleton().GetPrimPath().GetString();
+			const FTraversalInfo* SkeletonInfo = Impl.NodeUidToCachedTraversalInfo.Find(SkeletonNodeUid);
+			if (!SkeletonInfo)
+			{
+				return false;
+			}
+			JointNames = SkeletonInfo->SkelJointNames;
+			if (!JointNames)
+			{
+				return false;
+			}
+		}
+
+		UE::FUsdSkelSkinningQuery SkinningQuery = UsdUtils::CreateSkinningQuery(Prim, SkelQuery);
+		if (!SkinningQuery)
+		{
+			return false;
+		}
+
+		FSkeletalMeshImportData SkelMeshImportData;
+		UsdUtils::FUsdPrimMaterialAssignmentInfo TempMaterialInfo;
+		bool bSuccess = UsdToUnreal::ConvertSkinnedMesh(SkinningQuery, SkelQuery, SkelMeshImportData, TempMaterialInfo, Options);
+		if (!bSuccess)
+		{
+			return false;
+		}
+
+		// TODO: Swap this code path with some function to directly convert a skinned USD mesh to MeshDescription.
+		// We need that on the other USD workflows as well, not only here...
+		//
+		// Note: This is also doubly bad because it internally recomputes tangents and normals, which will also
+		// be done by Interchange later..
+		const USkeletalMesh* UnusedSkelMesh = nullptr;
+		FSkeletalMeshBuildSettings* UnusedBuildSettings = nullptr;
+		bSuccess = SkelMeshImportData.GetMeshDescription(UnusedSkelMesh, UnusedBuildSettings, OutMeshDescription);
+		if (!bSuccess)
+		{
+			return false;
+		}
+
+		FixSkeletalMeshDescriptionColors(OutMeshDescription);
+
+		FixMaterialSlotNames(OutMeshDescription, TempMaterialInfo.Slots);
+
+		OutJointNames = *JointNames;
+
+		return true;
+	}
+
+	bool GetMorphTargetPayloadData(
+		const FString& PayloadKey,
+		const UInterchangeUSDTranslatorImpl& Impl,
+		const UsdToUnreal::FUsdMeshConversionOptions& Options,
+		FMeshDescription& OutMeshDescription,
+		FString& OutMorphTargetName
+	)
+	{
+		// Our payloadkey should be something like "/MySkelRoot/MyMeshPrim\2\BlendShapeName_inbetweens:someName"
+		const bool bCullEmpty = false;
+		TArray<FString> PayloadKeyTokens;
+		PayloadKey.ParseIntoArray(PayloadKeyTokens, TEXT("\\"), bCullEmpty);
+		if (PayloadKeyTokens.Num() != 3)
+		{
+			return false;
+		}
+
+		const FString& MeshPrimPath = PayloadKeyTokens[0];
+		const FString& BlendShapeIndexStr = PayloadKeyTokens[1];
+		const FString& InbetweenName = PayloadKeyTokens[2];
+
+		int32 BlendShapeIndex = INDEX_NONE;
+		bool bLexed = LexTryParseString(BlendShapeIndex, *BlendShapeIndexStr);
+		if (!bLexed)
+		{
+			return false;
+		}
+
+		UE::FUsdPrim MeshPrim = Impl.UsdStage.GetPrimAtPath(FSdfPath{*MeshPrimPath});
+		UE::FUsdSkelBlendShapeQuery Query{MeshPrim};
+		if (!Query)
+		{
+			return false;
+		}
+
+		UE::FUsdSkelBlendShape BlendShape = Query.GetBlendShape(BlendShapeIndex);
+		if (!BlendShape)
+		{
+			return false;
+		}
+
+		// TODO: This is extremely slow, as it will reimport the mesh for every single morph target!
+		// It seems to be what the other translators do, however. We need some form of FMeshDescription caching here
+		TArray<FString> UnusedJointNames;
+		bool bConverted = GetSkeletalMeshPayloadData(MeshPrimPath, Impl, Options, OutMeshDescription, UnusedJointNames);
+		if (!bConverted || OutMeshDescription.IsEmpty())
+		{
+			return false;
+		}
+
+		OutMorphTargetName = BlendShape.GetPrim().GetName().ToString();
+		if (!InbetweenName.IsEmpty())
+		{
+			OutMorphTargetName += TEXT("_") + InbetweenName;
+		}
+
+		const float Weight = 1.0f;
+		return UsdUtils::ApplyBlendShape(OutMeshDescription, BlendShape.GetPrim(), Weight, InbetweenName);
+	}
 }	 // namespace UE::InterchangeUsdTranslator::Private
+#endif	  // USE_USD_SDK
 
 bool UInterchangeUSDTranslator::Translate(UInterchangeBaseNodeContainer& NodeContainer) const
 {
@@ -826,11 +1294,10 @@ bool UInterchangeUSDTranslator::Translate(UInterchangeBaseNodeContainer& NodeCon
 																	 : UnrealToUsd::ConvertToken(*Settings->MaterialPurpose.ToString()).Get();
 
 	// Traverse stage and emit translated nodes
-	FUsdPrim RootPrim = ImplPtr->UsdStage.GetPseudoRoot();
-	for (const FUsdPrim& Prim : RootPrim.GetChildren())
+	FTraversalInfo Info;
+	for (const FUsdPrim& Prim : ImplPtr->UsdStage.GetPseudoRoot().GetChildren())
 	{
-		UInterchangeBaseNode* ParentNode = nullptr;
-		Traverse(ImplPtr->UsdStage, Prim, ImplPtr, NodeContainer, ParentNode);
+		Traverse(Prim, ImplPtr, NodeContainer, Info);
 	}
 
 	return true;
@@ -935,17 +1402,37 @@ TFuture<TOptional<UE::Interchange::FMeshPayloadData>> UInterchangeUSDTranslator:
 			{
 				case EInterchangeMeshPayLoadType::STATIC:
 				{
-					bSuccess = UE::InterchangeUsdTranslator::Private::GetStaticMeshPayloadDataForPayLoadKey(
-						ImplPtr->UsdStage,
+					bSuccess = UE::InterchangeUsdTranslator::Private::GetStaticMeshPayloadData(
 						PayloadKey.UniqueId,
+						*ImplPtr,
 						OptionsCopy,
 						MeshPayloadData.MeshDescription
 					);
 					break;
 				}
-				case EInterchangeMeshPayLoadType::SKELETAL:		  // Fallthrough
-				case EInterchangeMeshPayLoadType::MORPHTARGET:	  // Fallthrough
-				case EInterchangeMeshPayLoadType::NONE:			  // Fallthrough
+				case EInterchangeMeshPayLoadType::SKELETAL:
+				{
+					bSuccess = UE::InterchangeUsdTranslator::Private::GetSkeletalMeshPayloadData(
+						PayloadKey.UniqueId,
+						*ImplPtr,
+						OptionsCopy,
+						MeshPayloadData.MeshDescription,
+						MeshPayloadData.JointNames
+					);
+					break;
+				}
+				case EInterchangeMeshPayLoadType::MORPHTARGET:
+				{
+					bSuccess = UE::InterchangeUsdTranslator::Private::GetMorphTargetPayloadData(
+						PayloadKey.UniqueId,
+						*ImplPtr,
+						OptionsCopy,
+						MeshPayloadData.MeshDescription,
+						MeshPayloadData.MorphTargetName
+					);
+					break;
+				}
+				case EInterchangeMeshPayLoadType::NONE:	   // Fallthrough
 				default:
 					break;
 			}
@@ -1018,9 +1505,9 @@ TOptional<UE::Interchange::FImportBlockedImage> UInterchangeUSDTranslator::GetBl
 	TextureGroup TextureGroup;
 	bool bDecoded = DecodeTexturePayloadKey(PayloadKey, FilePath, TextureGroup);
 	if (!bDecoded)
-{
-	return {};
-}
+	{
+		return {};
+	}
 
 	AlternateTexturePath = FilePath;
 
diff --git a/Engine/Plugins/Runtime/USDCore/Source/USDUtilities/Private/USDGeomMeshConversion.cpp b/Engine/Plugins/Runtime/USDCore/Source/USDUtilities/Private/USDGeomMeshConversion.cpp
--- a/Engine/Plugins/Runtime/USDCore/Source/USDUtilities/Private/USDGeomMeshConversion.cpp
+++ b/Engine/Plugins/Runtime/USDCore/Source/USDUtilities/Private/USDGeomMeshConversion.cpp
@@ -2307,8 +2307,8 @@ namespace UE::UsdGeomMeshConversion::Private
 					FPolygonGroupID NewPolygonGroup = OutMeshDescription.CreatePolygonGroup();
 					PolygonGroupMapping.Add(CombinedMaterialIndex, NewPolygonGroup);
 
-					const FString& SlotName = InMeshData.LocalMaterialInfo.Slots[LocalMaterialIndex].SlotName;
-					MaterialSlotNames[NewPolygonGroup] = *SlotName;
+					// This is important for runtime, where the material slots are matched to LOD sections based on their material slot name
+					MaterialSlotNames[NewPolygonGroup] = *LexToString(NewPolygonGroup.GetValue());
 				}
 			}
 
@@ -3476,7 +3476,6 @@ UsdUtils::FUsdPrimMaterialAssignmentInfo UsdUtils::GetPrimMaterialAssignments(
 		{
 			const pxr::UsdGeomSubset& GeomSubset = GeomSubsets[GeomSubsetIndex];
 			pxr::UsdPrim GeomSubsetPrim = GeomSubset.GetPrim();
-			FString SlotName = UsdToUnreal::ConvertString(GeomSubsetPrim.GetName());
 			FString GeomSubsetPath = UsdToUnreal::ConvertPath(GeomSubsetPrim.GetPath());
 			bool bHasAssignment = false;
 
@@ -3492,7 +3491,6 @@ UsdUtils::FUsdPrimMaterialAssignmentInfo UsdUtils::GetPrimMaterialAssignments(
 						if (TOptional<FString> UnrealMaterial = UsdUtils::GetUnrealSurfaceOutput(ShadeMaterial.GetPrim()))
 						{
 							FUsdPrimMaterialSlot& Slot = Result.Slots.Emplace_GetRef();
-							Slot.SlotName = SlotName;
 							Slot.MaterialSource = UnrealMaterial.GetValue();
 							Slot.AssignmentType = UsdUtils::EPrimAssignmentType::UnrealMaterial;
 							Slot.bMeshIsDoubleSided = bIsDoubleSided;
@@ -3508,7 +3506,6 @@ UsdUtils::FUsdPrimMaterialAssignmentInfo UsdUtils::GetPrimMaterialAssignments(
 					if (TOptional<FString> UnrealMaterial = FetchFirstUEMaterialFromAttribute(GeomSubsetPrim, TimeCode))
 					{
 						FUsdPrimMaterialSlot& Slot = Result.Slots.Emplace_GetRef();
-						Slot.SlotName = SlotName;
 						Slot.MaterialSource = UnrealMaterial.GetValue();
 						Slot.AssignmentType = UsdUtils::EPrimAssignmentType::UnrealMaterial;
 						Slot.bMeshIsDoubleSided = bIsDoubleSided;
@@ -3524,7 +3521,6 @@ UsdUtils::FUsdPrimMaterialAssignmentInfo UsdUtils::GetPrimMaterialAssignments(
 				if (TOptional<FString> BoundMaterial = FetchMaterialByComputingBoundMaterial(GeomSubsetPrim))
 				{
 					FUsdPrimMaterialSlot& Slot = Result.Slots.Emplace_GetRef();
-					Slot.SlotName = SlotName;
 					Slot.MaterialSource = BoundMaterial.GetValue();
 					Slot.AssignmentType = UsdUtils::EPrimAssignmentType::MaterialPrim;
 					Slot.bMeshIsDoubleSided = bIsDoubleSided;
@@ -3541,7 +3537,6 @@ UsdUtils::FUsdPrimMaterialAssignmentInfo UsdUtils::GetPrimMaterialAssignments(
 			if (!bHasAssignment)
 			{
 				FUsdPrimMaterialSlot& Slot = Result.Slots.Emplace_GetRef();
-				Slot.SlotName = SlotName;
 				Slot.PrimPaths.Add(GeomSubsetPath);
 				bHasAssignment = true;
 			}
@@ -3590,8 +3585,6 @@ UsdUtils::FUsdPrimMaterialAssignmentInfo UsdUtils::GetPrimMaterialAssignments(
 	bool bHasMainAssignment = false;
 	if (bNeedsMainAssignment)
 	{
-		const FString MainSlotName = TEXT("Main");
-
 		// Priority 1: Material is an unreal asset
 		if (RenderContext == UnrealIdentifiers::Unreal)
 		{
@@ -3604,7 +3597,6 @@ UsdUtils::FUsdPrimMaterialAssignmentInfo UsdUtils::GetPrimMaterialAssignments(
 				if (TOptional<FString> UnrealMaterial = UsdUtils::GetUnrealSurfaceOutput(ShadeMaterial.GetPrim()))
 				{
 					FUsdPrimMaterialSlot& Slot = Result.Slots.Emplace_GetRef();
-					Slot.SlotName = MainSlotName;
 					Slot.MaterialSource = UnrealMaterial.GetValue();
 					Slot.AssignmentType = UsdUtils::EPrimAssignmentType::UnrealMaterial;
 					Slot.bMeshIsDoubleSided = bIsDoubleSided;
@@ -3620,7 +3612,6 @@ UsdUtils::FUsdPrimMaterialAssignmentInfo UsdUtils::GetPrimMaterialAssignments(
 				if (TOptional<FString> UnrealMaterial = FetchFirstUEMaterialFromAttribute(UsdPrim, TimeCode))
 				{
 					FUsdPrimMaterialSlot& Slot = Result.Slots.Emplace_GetRef();
-					Slot.SlotName = MainSlotName;
 					Slot.MaterialSource = UnrealMaterial.GetValue();
 					Slot.AssignmentType = UsdUtils::EPrimAssignmentType::UnrealMaterial;
 					Slot.bMeshIsDoubleSided = bIsDoubleSided;
@@ -3637,7 +3628,6 @@ UsdUtils::FUsdPrimMaterialAssignmentInfo UsdUtils::GetPrimMaterialAssignments(
 			if (TOptional<FString> BoundMaterial = FetchMaterialByComputingBoundMaterial(UsdPrim))
 			{
 				FUsdPrimMaterialSlot& Slot = Result.Slots.Emplace_GetRef();
-				Slot.SlotName = MainSlotName;
 				Slot.MaterialSource = BoundMaterial.GetValue();
 				Slot.AssignmentType = UsdUtils::EPrimAssignmentType::MaterialPrim;
 				Slot.bMeshIsDoubleSided = bIsDoubleSided;
@@ -3656,7 +3646,6 @@ UsdUtils::FUsdPrimMaterialAssignmentInfo UsdUtils::GetPrimMaterialAssignments(
 			if (DisplayColor)
 			{
 				FUsdPrimMaterialSlot& Slot = Result.Slots.Emplace_GetRef();
-				Slot.SlotName = MainSlotName;
 				Slot.MaterialSource = DisplayColor.GetValue().ToString();
 				Slot.AssignmentType = UsdUtils::EPrimAssignmentType::DisplayColor;
 				Slot.bMeshIsDoubleSided = bIsDoubleSided;
diff --git a/Engine/Plugins/Runtime/USDCore/Source/USDUtilities/Private/USDSkeletalDataConversion.cpp b/Engine/Plugins/Runtime/USDCore/Source/USDUtilities/Private/USDSkeletalDataConversion.cpp
--- a/Engine/Plugins/Runtime/USDCore/Source/USDUtilities/Private/USDSkeletalDataConversion.cpp
+++ b/Engine/Plugins/Runtime/USDCore/Source/USDUtilities/Private/USDSkeletalDataConversion.cpp
@@ -2,9 +2,7 @@
 
 #include "USDSkeletalDataConversion.h"
 
-#include "Engine/SkinnedAssetCommon.h"
 #include "UnrealUSDWrapper.h"
-#include "UObject/Package.h"
 #include "USDAttributeUtils.h"
 #include "USDClassesModule.h"
 #include "USDConversionUtils.h"
@@ -22,6 +20,8 @@
 #include "UsdWrappers/SdfPath.h"
 #include "UsdWrappers/UsdAttribute.h"
 #include "UsdWrappers/UsdPrim.h"
+#include "UsdWrappers/UsdSkelSkeletonQuery.h"
+#include "UsdWrappers/UsdSkelSkinningQuery.h"
 #include "UsdWrappers/UsdStage.h"
 
 #include "Animation/AnimCurveTypes.h"
@@ -30,6 +30,7 @@
 #include "Async/ParallelFor.h"
 #include "BoneWeights.h"
 #include "ControlRig.h"
+#include "Engine/SkinnedAssetCommon.h"
 #include "Evaluation/MovieSceneSequenceTransform.h"
 #include "IMovieScenePlayer.h"
 #include "Interfaces/ITargetPlatformManagerModule.h"
@@ -46,6 +47,8 @@
 #include "Rendering/SkeletalMeshRenderData.h"
 #include "Rigs/RigHierarchyElements.h"
 #include "Sequencer/MovieSceneControlRigParameterSection.h"
+#include "StaticMeshAttributes.h"
+#include "UObject/Package.h"
 
 #if WITH_EDITOR
 #include "Animation/DebugSkelMeshComponent.h"
@@ -1205,30 +1208,45 @@ namespace UnrealToUsdImpl
 	}
 }	 // namespace UnrealToUsdImpl
 
-bool UsdToUnreal::ConvertSkeleton(const pxr::UsdSkelSkeletonQuery& SkeletonQuery, FSkeletalMeshImportData& SkelMeshImportData)
+bool UsdToUnreal::ConvertSkeleton(
+	const pxr::UsdSkelSkeletonQuery& InSkeletonQuery,
+	FUsdSkeletonData& OutConvertedData,
+	bool bEnsureAtLeastOneBone,
+	bool bEnsureSingleRootBone
+)
 {
-	TRACE_CPUPROFILER_EVENT_SCOPE(UsdToUnreal::ConvertSkeleton);
-
-	using namespace pxr;
+	TRACE_CPUPROFILER_EVENT_SCOPE(UsdToUnreal::ConvertSkeletonToTempData);
 
-	TArray<FString> JointNames;
-	TArray<int32> ParentJointIndices;
+	// Note: In here "joint" and "bone" are used interchangeably
 
-	uint32 RootBoneCount = 0;
+	if (!InSkeletonQuery)
+	{
+		return false;
+	}
 
 	// Retrieve the joint names and parent indices from the skeleton topology
 	// GetJointOrder already orders them from parent-to-child
-	VtArray<TfToken> JointOrder = SkeletonQuery.GetJointOrder();
-	const UsdSkelTopology& SkelTopology = SkeletonQuery.GetTopology();
-	for (uint32 Index = 0; Index < SkelTopology.GetNumJoints(); ++Index)
+	pxr::VtArray<pxr::TfToken> JointOrder = InSkeletonQuery.GetJointOrder();
+	const pxr::UsdSkelTopology& SkelTopology = InSkeletonQuery.GetTopology();
+	const int32 NumBones = SkelTopology.GetNumJoints();
+	if (NumBones > MAX_BONES)
 	{
-		SdfPath JointPath(JointOrder[Index]);
+		return false;
+	}
 
+	// Fill in everything but transforms
+	uint32 RootBoneCount = 0;
+	OutConvertedData.Bones.Reset();
+	OutConvertedData.Bones.SetNum(NumBones);
+	for (uint32 Index = 0; Index < SkelTopology.GetNumJoints(); ++Index)
+	{
+		pxr::SdfPath JointPath{JointOrder[Index]};
 		FString JointName = UsdToUnreal::ConvertString(JointPath.GetName());
-		JointNames.Add(JointName);
+		int32 ParentIndex = SkelTopology.GetParent(Index);
 
-		int ParentIndex = SkelTopology.GetParent(Index);
-		ParentJointIndices.Add(ParentIndex);
+		FUsdSkeletonData::FBone& Bone = OutConvertedData.Bones[Index];
+		Bone.Name = JointName;
+		Bone.ParentIndex = ParentIndex;
 
 		if (ParentIndex == -1)
 		{
@@ -1237,9 +1255,9 @@ bool UsdToUnreal::ConvertSkeleton(const pxr::UsdSkelSkeletonQuery& SkeletonQuery
 	}
 
 	// Skeleton has no joints: Generate a dummy single "Root" bone skeleton
-	if (JointNames.Num() == 0)
+	if (NumBones == 0)
 	{
-		FString SkeletonPrimPath = UsdToUnreal::ConvertPath(SkeletonQuery.GetPrim().GetPath());
+		FString SkeletonPrimPath = UsdToUnreal::ConvertPath(InSkeletonQuery.GetPrim().GetPath());
 
 		FUsdLogManager::LogMessage(
 			EMessageSeverity::Warning,
@@ -1247,48 +1265,46 @@ bool UsdToUnreal::ConvertSkeleton(const pxr::UsdSkelSkeletonQuery& SkeletonQuery
 				LOCTEXT(
 					"NoBonesInSkeleton",
 					"Skeleton prim '{0}' has no joints! "
-					"A new skeleton with a single 'Root' bone will be generated as USkeletalMeshes require valid skeletons. "
+					"A new skeleton with a single root joint will be generated as USkeletalMeshes require valid skeletons. "
 					"Note that this new skeleton may be written back to the USD stage when exporting the corresponding asset."
 				),
 				FText::FromString(SkeletonPrimPath)
 			)
 		);
-
-		SkelMeshImportData.RefBonesBinary.AddZeroed(1);
-
-		SkeletalMeshImportData::FBone& Bone = SkelMeshImportData.RefBonesBinary.Last();
+	}
+	if (bEnsureAtLeastOneBone)
+	{
+		FUsdSkeletonData::FBone& Bone = OutConvertedData.Bones.Emplace_GetRef();
 		Bone.Name = TEXT("Root");
 		Bone.ParentIndex = INDEX_NONE;
-		Bone.NumChildren = 0;
-		Bone.BonePos.Transform = FTransform3f::Identity;
-		Bone.BonePos.Length = 1.0f;
-		Bone.BonePos.XSize = 100.0f;
-		Bone.BonePos.YSize = 100.0f;
-		Bone.BonePos.ZSize = 100.0f;
 		return true;
 	}
 
-	if (JointNames.Num() > MAX_BONES)
+	// Get the rest transforms that will end up on the skeleton data.
+	// Note that we'll want to use the data in restTransforms, so that the skeletal mesh assumes the pose described
+	// by them when non-animated. See the gigantic comment inside ConvertSkinnedMesh.
+	pxr::VtArray<pxr::GfMatrix4d> JointLocalRestTransforms;
+	const bool bAtRest = true;
+	bool bTransformsComputed = InSkeletonQuery.ComputeJointLocalTransforms(&JointLocalRestTransforms, pxr::UsdTimeCode::EarliestTime(), bAtRest);
+	if (NumBones != JointLocalRestTransforms.size())
 	{
 		return false;
 	}
 
-	// Retrieve the bone transforms to be used as the reference pose
-	TArray<FTransform> BoneTransforms;
+	pxr::VtArray<pxr::GfMatrix4d> JointWorldBindTransforms;
+	bTransformsComputed &= InSkeletonQuery.GetJointWorldBindTransforms(&JointWorldBindTransforms);
 
-	// Note that for the FReferenceSkeleton we'll want to use the data in restTransforms, so that the skeletal mesh
-	// assumes the pose described by them when non-animated. See the gigantic comment inside ConvertSkinnedMesh.
-	VtArray<GfMatrix4d> JointLocalRestTransforms;
-	const bool bAtRest = true;
-	bool bJointTransformsComputed = SkeletonQuery.ComputeJointLocalTransforms(&JointLocalRestTransforms, UsdTimeCode::EarliestTime(), bAtRest);
-	if (bJointTransformsComputed)
+	pxr::VtArray<pxr::GfMatrix4d> JointLocalBindTransforms;
+	bTransformsComputed &= pxr::UsdSkelComputeJointLocalTransforms(SkelTopology, JointWorldBindTransforms, &JointLocalBindTransforms);
+
+	if (bTransformsComputed)
 	{
-		UsdStageWeakPtr Stage = SkeletonQuery.GetSkeleton().GetPrim().GetStage();
+		pxr::UsdStageWeakPtr Stage = InSkeletonQuery.GetSkeleton().GetPrim().GetStage();
 		const FUsdStageInfo StageInfo(Stage);
 
 		for (uint32 Index = 0; Index < JointLocalRestTransforms.size(); ++Index)
 		{
-			const GfMatrix4d& UsdMatrix = JointLocalRestTransforms[Index];
+			FUsdSkeletonData::FBone& Bone = OutConvertedData.Bones[Index];
 
 			// Here we use DecomposeWithUniformReflection instead of the previous UsdToUnreal::ConvertMatrix(StageInfo, UsdMatrix)
 			// call, because internally that would have done the matrix decomposition via FTransform::SetFromMatrix.
@@ -1302,74 +1318,98 @@ bool UsdToUnreal::ConvertSkeleton(const pxr::UsdSkelSkeletonQuery& SkeletonQuery
 			// Note that FBX secretly does this as well, because the FBX SDK's Matrix.GetT(), Matrix.GetQ() and Matrix.GetS()
 			// (used within UnFbx::FFbxImporter::ImportBones) seem to behave the same way and flip all axes when a reflection is
 			// detected.
-			FTransform BoneTransform;
 			{
-				FMatrix Matrix = UsdToUnreal::ConvertMatrix(UsdMatrix);
-				BoneTransform = UsdUtils::DecomposeWithUniformReflection(Matrix);
-				BoneTransform = UsdUtils::ConvertTransformToUESpace(StageInfo, BoneTransform);
+				const pxr::GfMatrix4d& UsdRestTransform = JointLocalRestTransforms[Index];
+				FMatrix Matrix = UsdToUnreal::ConvertMatrix(UsdRestTransform);
+				FTransform RestTransform = UsdUtils::DecomposeWithUniformReflection(Matrix);
+				RestTransform = UsdUtils::ConvertTransformToUESpace(StageInfo, RestTransform);
+				Bone.LocalRestTransform = RestTransform;
+			}
+			{
+				const pxr::GfMatrix4d& UsdBindTransform = JointLocalBindTransforms[Index];
+				FMatrix Matrix = UsdToUnreal::ConvertMatrix(UsdBindTransform);
+				FTransform BindTransform = UsdUtils::DecomposeWithUniformReflection(Matrix);
+				BindTransform = UsdUtils::ConvertTransformToUESpace(StageInfo, BindTransform);
+				Bone.LocalBindTransform = BindTransform;
 			}
-			BoneTransforms.Add(BoneTransform);
 		}
 	}
 
-	if (JointNames.Num() != BoneTransforms.Num())
-	{
-		return false;
-	}
-
 	// If we have more than one root bone, let's create a new "true root bone" and add the
 	// previously root bones as children of it
-	if (RootBoneCount > 1)
+	if (bEnsureSingleRootBone && RootBoneCount > 1)
 	{
-		for (int32& Index : ParentJointIndices)
+		TSet<FString> BoneNames;
+		for (FUsdSkeletonData::FBone& Bone : OutConvertedData.Bones)
 		{
+			BoneNames.Add(Bone.Name);
+
 			// Have previously root bones point at the new bone we'll add soon
-			if (Index == INDEX_NONE)
+			if (Bone.ParentIndex == INDEX_NONE)
 			{
-				Index = 0;
+				Bone.ParentIndex = 0;
 			}
 			// All other index references have to move one over since we'll push
 			// a new root bone into the start of the array
 			else
 			{
-				Index += 1;
+				Bone.ParentIndex += 1;
 			}
 		}
 
-		const uint64 FirstIndex = 0;
-		const int32 RootParentIndex = INDEX_NONE;
-
-		FString UniqueNewRootName = UsdUnreal::ObjectUtils::GetUniqueName(TEXT("Root"), TSet<FString>{JointNames});
-
-		JointNames.Insert(UniqueNewRootName, FirstIndex);
-		BoneTransforms.Insert(FTransform::Identity, FirstIndex);
-		ParentJointIndices.Insert(RootParentIndex, FirstIndex);
+		FUsdSkeletonData::FBone TrueRoot;
+		TrueRoot.Name = UsdUnreal::ObjectUtils::GetUniqueName(TEXT("Root"), BoneNames);
+		TrueRoot.ParentIndex = INDEX_NONE;
+		TrueRoot.LocalRestTransform = FTransform::Identity;
+		OutConvertedData.Bones.Insert(TrueRoot, 0);
 	}
 
-	// Store the retrieved data as bones into the SkeletalMeshImportData
-	SkelMeshImportData.RefBonesBinary.AddZeroed(JointNames.Num());
-	for (int32 Index = 0; Index < JointNames.Num(); ++Index)
+	// Fill in child indices (easier now so we don't have to remap them for multiple root bones)
+	for (int32 Index = 0; Index < NumBones; ++Index)
 	{
-		SkeletalMeshImportData::FBone& Bone = SkelMeshImportData.RefBonesBinary[Index];
-
-		Bone.Name = JointNames[Index];
-		Bone.ParentIndex = ParentJointIndices[Index];
-		// Increment the number of children each time a bone is referenced as a parent bone; the root has a parent index of -1
+		FUsdSkeletonData::FBone& Bone = OutConvertedData.Bones[Index];
 		if (Bone.ParentIndex >= 0)
 		{
-			// The joints are ordered from parent-to-child so the parent will already have been added to the array
-			SkeletalMeshImportData::FBone& ParentBone = SkelMeshImportData.RefBonesBinary[Bone.ParentIndex];
-			++ParentBone.NumChildren;
+			FUsdSkeletonData::FBone& ParentBone = OutConvertedData.Bones[Bone.ParentIndex];
+			ParentBone.ChildIndices.Add(Index);
 		}
+	}
+
+	return true;
+}
+
+bool UsdToUnreal::ConvertSkeleton(const pxr::UsdSkelSkeletonQuery& SkeletonQuery, FSkeletalMeshImportData& SkelMeshImportData)
+{
+	TRACE_CPUPROFILER_EVENT_SCOPE(UsdToUnreal::ConvertSkeleton);
 
-		SkeletalMeshImportData::FJointPos& JointMatrix = Bone.BonePos;
-		JointMatrix.Transform = FTransform3f(BoneTransforms[Index]);
+	FUsdSkeletonData TempData;
+	const bool bEnsureAtLeastOneBone = true;
+	const bool bEnsureSingleRootBone = true;
+	const bool bSuccess = ConvertSkeleton(SkeletonQuery, TempData, bEnsureAtLeastOneBone, bEnsureSingleRootBone);
+	if (!bSuccess)
+	{
+		return false;
+	}
+
+	// Store the retrieved data as bones into the SkeletalMeshImportData
+	const int32 NumBones = TempData.Bones.Num();
+	SkelMeshImportData.RefBonesBinary.AddZeroed(NumBones);
+	for (int32 Index = 0; Index < NumBones; ++Index)
+	{
+		FUsdSkeletonData::FBone& InBone = TempData.Bones[Index];
+		SkeletalMeshImportData::FBone& OutBone = SkelMeshImportData.RefBonesBinary[Index];
+
+		OutBone.Name = InBone.Name;
+		OutBone.ParentIndex = InBone.ParentIndex;
+		OutBone.NumChildren = InBone.ChildIndices.Num();
 
 		// Not sure if Length and X/Y/Z Size need to be set, there are no equivalents in USD
+		SkeletalMeshImportData::FJointPos& JointMatrix = OutBone.BonePos;
 		JointMatrix.Length = 1.f;
 		JointMatrix.XSize = 100.f;
 		JointMatrix.YSize = 100.f;
 		JointMatrix.ZSize = 100.f;
+		JointMatrix.Transform = FTransform3f(InBone.LocalRestTransform);
 	}
 
 	return true;
@@ -2429,7 +2469,7 @@ bool UsdToUnreal::ConvertSkelAnim(
 					int32 LODIndex
 				)
 			{
-				pxr::UsdSkelSkinningQuery SkinningQuery = UsdUtils::CreateSkinningQuery(LODMesh, InUsdSkeletonQuery);
+				pxr::UsdSkelSkinningQuery SkinningQuery = UsdUtils::CreateSkinningQuery(LODMesh.GetPrim(), InUsdSkeletonQuery);
 				if (!SkinningQuery)
 				{
 					return true;	// Continue trying other LODs
@@ -3102,12 +3142,20 @@ pxr::UsdSkelSkinningQuery UsdUtils::CreateSkinningQuery(const pxr::UsdGeomMesh&
 		return {};
 	}
 
-	const pxr::UsdSkelAnimQuery& AnimQuery = SkeletonQuery.GetAnimQuery();
+	return CreateSkinningQuery(SkinnedPrim, SkeletonQuery);
+}
 
-	pxr::UsdSkelBindingAPI SkelBindingAPI{SkinnedPrim};
+UE::FUsdSkelSkinningQuery UsdUtils::CreateSkinningQuery(const pxr::UsdPrim& SkinnedMeshPrim, const pxr::UsdSkelSkeletonQuery& SkeletonQuery)
+{
+	pxr::UsdSkelBindingAPI SkelBindingAPI{SkinnedMeshPrim};
+	const pxr::UsdSkelAnimQuery& AnimQuery = SkeletonQuery.GetAnimQuery();
+	if (!SkelBindingAPI || !AnimQuery)
+	{
+		return {};
+	}
 
-	return pxr::UsdSkelSkinningQuery(
-		SkinnedPrim,
+	return UE::FUsdSkelSkinningQuery{pxr::UsdSkelSkinningQuery(
+		SkinnedMeshPrim,
 		SkeletonQuery ? SkeletonQuery.GetJointOrder() : pxr::VtTokenArray(),
 		AnimQuery ? AnimQuery.GetBlendShapeOrder() : pxr::VtTokenArray(),
 		SkelBindingAPI.GetJointIndicesAttr(),
@@ -3117,7 +3165,7 @@ pxr::UsdSkelSkinningQuery UsdUtils::CreateSkinningQuery(const pxr::UsdGeomMesh&
 		SkelBindingAPI.GetJointsAttr(),
 		SkelBindingAPI.GetBlendShapesAttr(),
 		SkelBindingAPI.GetBlendShapeTargetsRel()
-	);
+	)};
 }
 
 void UsdUtils::BindAnimationSource(pxr::UsdPrim& Prim, const pxr::UsdPrim& AnimationSource)
@@ -3248,6 +3296,209 @@ bool UsdUtils::GetSkelQueries(
 	return InOutSkelCache->ComputeSkelBinding(InSkelRootPrim, InSkeletonPrim, &OutSkelBinding, pxr::UsdTraverseInstanceProxies());
 }
 
+bool UsdUtils::ApplyBlendShape(
+	FMeshDescription& InOutMeshDescription,
+	const pxr::UsdPrim& InBlendShapePrim,
+	float Weight,
+	const FString InInbetweenName
+)
+{
+	TRACE_CPUPROFILER_EVENT_SCOPE(UsdToUnreal::ConvertBlendShape);
+
+	FScopedUsdAllocs Allocs;
+
+	pxr::UsdSkelBlendShape UsdBlendShape{InBlendShapePrim};
+	if (!UsdBlendShape)
+	{
+		return false;
+	}
+
+	pxr::UsdStageRefPtr Stage = InBlendShapePrim.GetStage();
+	if (!Stage)
+	{
+		return false;
+	}
+
+	FUsdStageInfo StageInfo{Stage};
+
+	// Collect blend shape deltas
+	pxr::VtArray<pxr::GfVec3f> PositionOffsets;
+	pxr::VtArray<pxr::GfVec3f> NormalOffsets;
+	pxr::VtArray<int> PointIndices;
+	{
+		pxr::UsdAttribute IndicesAttr = UsdBlendShape.GetPointIndicesAttr();
+		IndicesAttr.Get(&PointIndices);
+
+		if (!InInbetweenName.IsEmpty())
+		{
+			if (pxr::UsdSkelInbetweenShape Inbetween = UsdBlendShape.GetInbetween(UnrealToUsd::ConvertToken(*InInbetweenName).Get()))
+			{
+				Inbetween.GetOffsets(&PositionOffsets);
+				Inbetween.GetNormalOffsets(&NormalOffsets);
+			}
+			else
+			{
+				UE_LOG(
+					LogUsd,
+					Warning,
+					TEXT("Failed to find inbetween '%s' when applying blend shape prim '%s' to a mesh description"),
+					*InInbetweenName,
+					*UsdToUnreal::ConvertPath(InBlendShapePrim.GetPrimPath())
+				);
+			}
+		}
+		else
+		{
+			if (pxr::UsdAttribute OffsetsAttr = UsdBlendShape.GetOffsetsAttr())
+			{
+				OffsetsAttr.Get(&PositionOffsets);
+			}
+
+			if (pxr::UsdAttribute NormalsAttr = UsdBlendShape.GetNormalOffsetsAttr())
+			{
+				NormalsAttr.Get(&NormalOffsets);
+			}
+		}
+	}
+
+	FStaticMeshAttributes Attributes{InOutMeshDescription};
+	TVertexAttributesRef<FVector3f> MeshPositions = Attributes.GetVertexPositions();
+	TVertexInstanceAttributesRef<FVector3f> MeshInstanceNormals = Attributes.GetVertexInstanceNormals();
+
+	bool bApplied = false;
+
+	// We have one value for each vertex of the mesh description
+	if (PointIndices.size() == 0)
+	{
+		// Position offsets
+		if (PositionOffsets.size() == (size_t)MeshPositions.GetNumElements())
+		{
+			for (uint32 OffsetIndex = 0; OffsetIndex < PositionOffsets.size(); ++OffsetIndex)
+			{
+				const FVector UEOffset = UsdToUnreal::ConvertVector(StageInfo, PositionOffsets[OffsetIndex]);
+				MeshPositions[OffsetIndex] += FVector3f{UEOffset * Weight};
+			}
+
+			bApplied = true;
+		}
+		else
+		{
+			FString InbetweenText = InInbetweenName.IsEmpty() ? FString::Printf(TEXT(" (inbetween '%s')"), *InInbetweenName) : TEXT("");
+			UE_LOG(
+				LogUsd,
+				Warning,
+				TEXT("Failed to apply position offsets from BlendShape '%s'%s: Expected MeshDescription to have %d vertex positions, but it has %u!"),
+				*UsdToUnreal::ConvertPath(InBlendShapePrim.GetPrimPath()),
+				*InbetweenText,
+				PositionOffsets.size(),
+				MeshPositions.GetNumElements()
+			);
+		}
+
+		// Normal offsets
+		if (NormalOffsets.size() == (size_t)MeshPositions.GetNumElements())
+		{
+			for (uint32 VertexIndex = 0; VertexIndex < NormalOffsets.size(); ++VertexIndex)
+			{
+				const FVector UENormal = UsdToUnreal::ConvertVector(StageInfo, NormalOffsets[VertexIndex]).GetSafeNormal();
+
+				TArrayView<const FVertexInstanceID> Instances = InOutMeshDescription.GetVertexVertexInstanceIDs(VertexIndex);
+				for (FVertexInstanceID InstanceID : Instances)
+				{
+					MeshInstanceNormals[InstanceID] = (MeshInstanceNormals[InstanceID] + Weight * FVector3f{UENormal}).GetSafeNormal();
+				}
+			}
+
+			bApplied = true;
+		}
+		else if (NormalOffsets.size() != 0)
+		{
+			FString InbetweenText = InInbetweenName.IsEmpty() ? FString::Printf(TEXT(" (inbetween '%s')"), *InInbetweenName) : TEXT("");
+			UE_LOG(
+				LogUsd,
+				Warning,
+				TEXT("Failed to apply normal offsets from BlendShape '%s'%s: Expected MeshDescription to have %d vertices, but it has %u!"),
+				*UsdToUnreal::ConvertPath(InBlendShapePrim.GetPrimPath()),
+				*InbetweenText,
+				NormalOffsets.size(),
+				MeshPositions.GetNumElements()
+			);
+		}
+	}
+	// We have values for only a few vertices of the mesh description
+	else if (PointIndices.size() > 0)
+	{
+		// Position offsets
+		if (PointIndices.size() == PositionOffsets.size())
+		{
+			for (uint32 OffsetIndex = 0; OffsetIndex < PositionOffsets.size(); ++OffsetIndex)
+			{
+				int TargetPointIndex = PointIndices[OffsetIndex];
+				if (TargetPointIndex >= 0 && TargetPointIndex < MeshPositions.GetNumElements())
+				{
+					const FVector UEOffset = UsdToUnreal::ConvertVector(StageInfo, PositionOffsets[OffsetIndex]);
+					MeshPositions[TargetPointIndex] += FVector3f{UEOffset * Weight};
+				}
+			}
+
+			bApplied = true;
+		}
+		else
+		{
+			FString InbetweenText = InInbetweenName.IsEmpty() ? FString::Printf(TEXT(" (inbetween '%s')"), *InInbetweenName) : TEXT("");
+			UE_LOG(
+				LogUsd,
+				Warning,
+				TEXT(
+					"Failed to apply indexed position offsets from BlendShape '%s'%s: The blend shape has %u offsets, but %u indices! (those should match)"
+				),
+				*UsdToUnreal::ConvertPath(InBlendShapePrim.GetPrimPath()),
+				*InbetweenText,
+				PositionOffsets.size(),
+				PointIndices.size()
+			);
+		}
+
+		// Normal offsets
+		if (PointIndices.size() == NormalOffsets.size())
+		{
+			for (uint32 NormalIndex = 0; NormalIndex < NormalOffsets.size(); ++NormalIndex)
+			{
+				int TargetPointIndex = PointIndices[NormalIndex];
+				if (TargetPointIndex >= 0 && TargetPointIndex < MeshPositions.GetNumElements())
+				{
+					const FVector UENormal = UsdToUnreal::ConvertVector(StageInfo, NormalOffsets[NormalIndex]).GetSafeNormal();
+
+					TArrayView<const FVertexInstanceID> Instances = InOutMeshDescription.GetVertexVertexInstanceIDs(TargetPointIndex);
+					for (FVertexInstanceID InstanceID : Instances)
+					{
+						MeshInstanceNormals[InstanceID] = (MeshInstanceNormals[InstanceID] + Weight * FVector3f{UENormal}).GetSafeNormal();
+					}
+				}
+			}
+
+			bApplied = true;
+		}
+		else if (NormalOffsets.size() != 0)
+		{
+			FString InbetweenText = InInbetweenName.IsEmpty() ? FString::Printf(TEXT(" (inbetween '%s')"), *InInbetweenName) : TEXT("");
+			UE_LOG(
+				LogUsd,
+				Warning,
+				TEXT(
+					"Failed to apply indexed normal offsets from BlendShape '%s'%s: The blend shape has %u offsets, but %u indices! (those should match)"
+				),
+				*UsdToUnreal::ConvertPath(InBlendShapePrim.GetPrimPath()),
+				*InbetweenText,
+				NormalOffsets.size(),
+				PointIndices.size()
+			);
+		}
+	}
+
+	return bApplied;
+}
+
 #endif	  // USE_USD_SDK
 
 #if USE_USD_SDK && WITH_EDITOR
diff --git a/Engine/Plugins/Runtime/USDCore/Source/USDUtilities/Public/USDGeomMeshConversion.h b/Engine/Plugins/Runtime/USDCore/Source/USDUtilities/Public/USDGeomMeshConversion.h
--- a/Engine/Plugins/Runtime/USDCore/Source/USDUtilities/Public/USDGeomMeshConversion.h
+++ b/Engine/Plugins/Runtime/USDCore/Source/USDUtilities/Public/USDGeomMeshConversion.h
@@ -273,10 +273,6 @@ namespace UsdUtils
 	 */
 	struct FUsdPrimMaterialSlot
 	{
-		// This carries the name of the UsdGeomSubset if applicable, or just "Main" for the single material slot
-		// directly on the Mesh prim
-		FString SlotName;
-
 		/**
 		 * Path to the prims that contain this material assignment (e.g. '/Root/my_cube' or '/Root/my_cube/geomsubset_0'.
 		 * This is a set because this single material slot in UE may be merged from N identical prims.
@@ -312,7 +308,6 @@ namespace UsdUtils
 
 		friend FArchive& operator<<(FArchive& Ar, FUsdPrimMaterialSlot& Slot)
 		{
-			Ar << Slot.SlotName;
 			Ar << Slot.MaterialSource;
 			Ar << Slot.AssignmentType;
 			Ar << Slot.bMeshIsDoubleSided;
diff --git a/Engine/Plugins/Runtime/USDCore/Source/USDUtilities/Public/USDSkeletalDataConversion.h b/Engine/Plugins/Runtime/USDCore/Source/USDUtilities/Public/USDSkeletalDataConversion.h
--- a/Engine/Plugins/Runtime/USDCore/Source/USDUtilities/Public/USDSkeletalDataConversion.h
+++ b/Engine/Plugins/Runtime/USDCore/Source/USDUtilities/Public/USDSkeletalDataConversion.h
@@ -53,6 +53,11 @@ namespace UsdUtils
 	struct FUsdPrimMaterialAssignmentInfo;
 }
 
+namespace UE
+{
+	class FUsdSkelSkinningQuery;
+}
+
 #endif	  // #if USE_USD_SDK
 
 namespace UsdUtils
@@ -121,6 +126,11 @@ namespace UsdUtils
 
 #if USE_USD_SDK
 	/** Allows creation of a skinning query from the underlying skinned mesh and skeleton. Adapted from the USD SDK implementation */
+	USDUTILITIES_API UE::FUsdSkelSkinningQuery CreateSkinningQuery(
+		const pxr::UsdPrim& SkinnedMeshPrim,
+		const pxr::UsdSkelSkeletonQuery& SkeletonQuery
+	);
+	UE_DEPRECATED(5.5, "Please use the other overload of CreateSkinningQuery, that can better work with the UsdWrapper types and no-RTTI modules")
 	USDUTILITIES_API pxr::UsdSkelSkinningQuery CreateSkinningQuery(
 		const pxr::UsdGeomMesh& SkinnedMesh,
 		const pxr::UsdSkelSkeletonQuery& SkeletonQuery
@@ -161,12 +171,60 @@ namespace UsdUtils
 		pxr::UsdSkelSkeletonQuery& OutSkeletonQuery,
 		pxr::UsdSkelCache* InOutSkelCache = nullptr
 	);
+
+	/**
+	 * Converts and applies the blend shape deltas from InBlendShapePrim onto the mesh described by InOutMeshDescription.
+	 * Note that we expect InOutMeshDescription and InBlendShapePrim to at least have a matching number of vertices/points.
+	 * @param InOutMeshDescription - The mesh data to deform with the blend shape
+	 * @param InBlendShapePrim - The UsdSkelBlendShape prim to deform the mesh with
+	 * @param Weight - Scalar that multiplies the deltas before they're applied to the mesh deformation:
+	 *                 0.0f means no effect, 1.0f adds the deltas at full strength, etc.
+	 * @param InInbetweenName - The name of the inbetween shape of InBlendShapePrim to apply instead, if any.
+	 *                          If left empty (default) the main blend shape is applied instead.
+	 * @return Whether the conversion was successful or not.
+	 */
+	USDUTILITIES_API bool ApplyBlendShape(
+		FMeshDescription& InOutMeshDescription,
+		const pxr::UsdPrim& InBlendShapePrim,
+		float Weight = 1.0f,
+		const FString InInbetweenName = {}
+	);
 #endif	  // USE_USD_SDK
 }	 // namespace UsdUtils
 
 #if USE_USD_SDK && WITH_EDITOR
 namespace UsdToUnreal
 {
+	struct FUsdSkeletonData
+	{
+		struct FBone
+		{
+			FString Name;
+			FTransform LocalRestTransform = FTransform::Identity;
+			FTransform LocalBindTransform = FTransform::Identity;
+
+			TArray<int32> ChildIndices;
+			int32 ParentIndex = INDEX_NONE;	   // Root bones remain with INDEX_NONE
+		};
+
+		TArray<FBone> Bones;
+	};
+
+	/**
+	 * Extracts skeleton data from UsdSkeletonQuery and places the results in ConvertedData.
+	 * @param UsdSkeletonQuery - SkeletonQuery with the data to convert
+	 * @param ConvertedData - Output parameter that will be filled with the converted data
+	 * @param bEnsureAtLeastOneBone - Whether to ensure the resulting ConvertedData contains at least one bone (adding one if needed)
+	 * @param bEnsureSingleRootBone - Whether to ensure the resulting ConvertedData has a *single* root bone
+	 *                                (adding a new one and reparenting the previous root bones to it, if needed)
+	 * @return Whether the conversion was successful or not.
+	 */
+	USDUTILITIES_API bool ConvertSkeleton(
+		const pxr::UsdSkelSkeletonQuery& UsdSkeletonQuery,
+		FUsdSkeletonData& ConvertedData,
+		bool bEnsureAtLeastOneBone = true,
+		bool bEnsureSingleRootBone = true
+	);
 	/**
 	 * Extracts skeleton data from UsdSkeletonQuery and places the results in SkelMeshImportData.
 	 * @param UsdSkeletonQuery - SkeletonQuery with the data to convert
@@ -305,7 +363,7 @@ namespace UsdToUnreal
 		UsdUtils::FBlendShapeMap& InBlendShapesByPath,
 		USkeletalMesh* InOutSkeletalMesh
 	);
-}
+}	 // namespace UsdToUnreal
 
 namespace UnrealToUsd
 {
diff --git a/Engine/Plugins/Runtime/USDCore/Source/UnrealUSDWrapper/Private/UsdWrappers/UsdGeomXformable.cpp b/Engine/Plugins/Runtime/USDCore/Source/UnrealUSDWrapper/Private/UsdWrappers/UsdGeomXformable.cpp
--- a/Engine/Plugins/Runtime/USDCore/Source/UnrealUSDWrapper/Private/UsdWrappers/UsdGeomXformable.cpp
+++ b/Engine/Plugins/Runtime/USDCore/Source/UnrealUSDWrapper/Private/UsdWrappers/UsdGeomXformable.cpp
@@ -112,18 +112,21 @@ namespace UE
 
 #if USE_USD_SDK
 	FUsdGeomXformable::FUsdGeomXformable(const pxr::UsdGeomXformable& InUsdGeomXformable)
+		: FUsdTyped(InUsdGeomXformable)
 	{
 		FScopedUnrealAllocs UnrealAllocs;
 		Impl = MakeUnique<Internal::FUsdGeomXformableImpl>(InUsdGeomXformable);
 	}
 
 	FUsdGeomXformable::FUsdGeomXformable(pxr::UsdGeomXformable&& InUsdGeomXformable)
+		: FUsdTyped(MoveTemp(InUsdGeomXformable))
 	{
 		FScopedUnrealAllocs UnrealAllocs;
 		Impl = MakeUnique<Internal::FUsdGeomXformableImpl>(MoveTemp(InUsdGeomXformable));
 	}
 
 	FUsdGeomXformable::FUsdGeomXformable(const pxr::UsdPrim& Prim)
+		: FUsdTyped(Prim)
 	{
 		FScopedUnrealAllocs UnrealAllocs;
 		Impl = MakeUnique<Internal::FUsdGeomXformableImpl>(Prim);
diff --git a/Engine/Plugins/Runtime/USDCore/Source/UnrealUSDWrapper/Private/UsdWrappers/UsdSkelAnimQuery.cpp b/Engine/Plugins/Runtime/USDCore/Source/UnrealUSDWrapper/Private/UsdWrappers/UsdSkelAnimQuery.cpp
new file mode 100644
--- /dev/null
+++ b/Engine/Plugins/Runtime/USDCore/Source/UnrealUSDWrapper/Private/UsdWrappers/UsdSkelAnimQuery.cpp
@@ -0,0 +1,145 @@
+// Copyright Epic Games, Inc. All Rights Reserved.
+
+#include "UsdWrappers/UsdSkelAnimQuery.h"
+
+#include "USDMemory.h"
+#include "UsdWrappers/UsdPrim.h"
+
+#if USE_USD_SDK
+#include "USDIncludesStart.h"
+#include "pxr/usd/usdSkel/animQuery.h"
+#include "USDIncludesEnd.h"
+#endif	  //  #if USE_USD_SDK
+
+namespace UE
+{
+	namespace Internal
+	{
+		class FUsdSkelAnimQueryImpl
+		{
+		public:
+			FUsdSkelAnimQueryImpl()
+#if USE_USD_SDK
+				: PxrUsdSkelAnimQuery()
+#endif	  // #if USE_USD_SDK
+			{
+			}
+
+#if USE_USD_SDK
+			explicit FUsdSkelAnimQueryImpl(const pxr::UsdSkelAnimQuery& InUsdSkelAnimQuery)
+				: PxrUsdSkelAnimQuery(InUsdSkelAnimQuery)
+			{
+			}
+
+			explicit FUsdSkelAnimQueryImpl(pxr::UsdSkelAnimQuery&& InUsdSkelAnimQuery)
+				: PxrUsdSkelAnimQuery(MoveTemp(InUsdSkelAnimQuery))
+			{
+			}
+
+			TUsdStore<pxr::UsdSkelAnimQuery> PxrUsdSkelAnimQuery;
+#endif	  // #if USE_USD_SDK
+		};
+	}	  // namespace Internal
+
+	FUsdSkelAnimQuery::FUsdSkelAnimQuery()
+	{
+		FScopedUnrealAllocs UnrealAllocs;
+		Impl = MakeUnique<Internal::FUsdSkelAnimQueryImpl>();
+	}
+
+	FUsdSkelAnimQuery::FUsdSkelAnimQuery(const FUsdSkelAnimQuery& Other)
+	{
+#if USE_USD_SDK
+		FScopedUnrealAllocs UnrealAllocs;
+		Impl = MakeUnique<Internal::FUsdSkelAnimQueryImpl>(Other.Impl->PxrUsdSkelAnimQuery.Get());
+#endif	  // #if USE_USD_SDK
+	}
+
+	FUsdSkelAnimQuery::FUsdSkelAnimQuery(FUsdSkelAnimQuery&& Other)
+		: Impl(MakeUnique<Internal::FUsdSkelAnimQueryImpl>())
+	{
+		if (this != &Other)
+		{
+			Impl = MoveTemp(Other.Impl);
+		}
+	}
+
+	FUsdSkelAnimQuery::~FUsdSkelAnimQuery()
+	{
+		FScopedUnrealAllocs UnrealAllocs;
+		Impl.Reset();
+	}
+
+	FUsdSkelAnimQuery& FUsdSkelAnimQuery::operator=(const FUsdSkelAnimQuery& Other)
+	{
+#if USE_USD_SDK
+		FScopedUnrealAllocs UnrealAllocs;
+		Impl = MakeUnique<Internal::FUsdSkelAnimQueryImpl>(Other.Impl->PxrUsdSkelAnimQuery.Get());
+#endif	  // #if USE_USD_SDK
+
+		return *this;
+	}
+
+	FUsdSkelAnimQuery& FUsdSkelAnimQuery::operator=(FUsdSkelAnimQuery&& Other)
+	{
+		if (this != &Other)
+		{
+			FScopedUnrealAllocs UnrealAllocs;
+			Impl = MoveTemp(Other.Impl);
+		}
+
+		return *this;
+	}
+
+	FUsdSkelAnimQuery::operator bool() const
+	{
+#if USE_USD_SDK
+		return static_cast<bool>(Impl->PxrUsdSkelAnimQuery.Get());
+#else
+		return false;
+#endif
+	}
+
+#if USE_USD_SDK
+	FUsdSkelAnimQuery::FUsdSkelAnimQuery(const pxr::UsdSkelAnimQuery& InUsdSkelAnimQuery)
+	{
+		FScopedUnrealAllocs UnrealAllocs;
+		Impl = MakeUnique<Internal::FUsdSkelAnimQueryImpl>(InUsdSkelAnimQuery);
+	}
+
+	FUsdSkelAnimQuery::FUsdSkelAnimQuery(pxr::UsdSkelAnimQuery&& InUsdSkelAnimQuery)
+	{
+		FScopedUnrealAllocs UnrealAllocs;
+		Impl = MakeUnique<Internal::FUsdSkelAnimQueryImpl>(MoveTemp(InUsdSkelAnimQuery));
+	}
+
+	FUsdSkelAnimQuery::operator pxr::UsdSkelAnimQuery&()
+	{
+		return Impl->PxrUsdSkelAnimQuery.Get();
+	}
+
+	FUsdSkelAnimQuery::operator const pxr::UsdSkelAnimQuery&() const
+	{
+		return Impl->PxrUsdSkelAnimQuery.Get();
+	}
+#endif	  // #if USE_USD_SDK
+
+	bool FUsdSkelAnimQuery::IsValid() const
+	{
+#if USE_USD_SDK
+		return Impl->PxrUsdSkelAnimQuery.Get().IsValid();
+#else
+		return false;
+#endif	  // #if USE_USD_SDK
+	}
+
+	FUsdPrim FUsdSkelAnimQuery::GetPrim() const
+	{
+#if USE_USD_SDK
+		FScopedUsdAllocs Allocs;
+		return FUsdPrim{Impl->PxrUsdSkelAnimQuery.Get().GetPrim()};
+#else
+		return {};
+#endif	  // #if USE_USD_SDK
+	}
+}	 // namespace UE
diff --git a/Engine/Plugins/Runtime/USDCore/Source/UnrealUSDWrapper/Private/UsdWrappers/UsdSkelBinding.cpp b/Engine/Plugins/Runtime/USDCore/Source/UnrealUSDWrapper/Private/UsdWrappers/UsdSkelBinding.cpp
new file mode 100644
--- /dev/null
+++ b/Engine/Plugins/Runtime/USDCore/Source/UnrealUSDWrapper/Private/UsdWrappers/UsdSkelBinding.cpp
@@ -0,0 +1,149 @@
+// Copyright Epic Games, Inc. All Rights Reserved.
+
+#include "UsdWrappers/UsdSkelBinding.h"
+
+#include "USDMemory.h"
+#include "UsdWrappers/UsdPrim.h"
+#include "UsdWrappers/UsdSkelSkinningQuery.h"
+
+#if USE_USD_SDK
+#include "USDIncludesStart.h"
+#include "pxr/usd/usdSkel/binding.h"
+#include "pxr/usd/usdSkel/skeleton.h"
+#include "USDIncludesEnd.h"
+#endif	  //  #if USE_USD_SDK
+
+namespace UE
+{
+	namespace Internal
+	{
+		class FUsdSkelBindingImpl
+		{
+		public:
+			FUsdSkelBindingImpl()
+#if USE_USD_SDK
+				: PxrUsdSkelBinding()
+#endif	  // #if USE_USD_SDK
+			{
+			}
+
+#if USE_USD_SDK
+			explicit FUsdSkelBindingImpl(const pxr::UsdSkelBinding& InUsdSkelBinding)
+				: PxrUsdSkelBinding(InUsdSkelBinding)
+			{
+			}
+
+			explicit FUsdSkelBindingImpl(pxr::UsdSkelBinding&& InUsdSkelBinding)
+				: PxrUsdSkelBinding(MoveTemp(InUsdSkelBinding))
+			{
+			}
+
+			TUsdStore<pxr::UsdSkelBinding> PxrUsdSkelBinding;
+#endif	  // #if USE_USD_SDK
+		};
+	}	  // namespace Internal
+
+	FUsdSkelBinding::FUsdSkelBinding()
+	{
+		FScopedUnrealAllocs UnrealAllocs;
+		Impl = MakeUnique<Internal::FUsdSkelBindingImpl>();
+	}
+
+	FUsdSkelBinding::FUsdSkelBinding(const FUsdSkelBinding& Other)
+	{
+#if USE_USD_SDK
+		FScopedUnrealAllocs UnrealAllocs;
+		Impl = MakeUnique<Internal::FUsdSkelBindingImpl>(Other.Impl->PxrUsdSkelBinding.Get());
+#endif	  // #if USE_USD_SDK
+	}
+
+	FUsdSkelBinding::FUsdSkelBinding(FUsdSkelBinding&& Other)
+		: Impl(MakeUnique<Internal::FUsdSkelBindingImpl>())
+	{
+		if (this != &Other)
+		{
+			Impl = MoveTemp(Other.Impl);
+		}
+	}
+
+	FUsdSkelBinding::~FUsdSkelBinding()
+	{
+		FScopedUnrealAllocs UnrealAllocs;
+		Impl.Reset();
+	}
+
+	FUsdSkelBinding& FUsdSkelBinding::operator=(const FUsdSkelBinding& Other)
+	{
+#if USE_USD_SDK
+		FScopedUnrealAllocs UnrealAllocs;
+		Impl = MakeUnique<Internal::FUsdSkelBindingImpl>(Other.Impl->PxrUsdSkelBinding.Get());
+#endif	  // #if USE_USD_SDK
+
+		return *this;
+	}
+
+	FUsdSkelBinding& FUsdSkelBinding::operator=(FUsdSkelBinding&& Other)
+	{
+		if (this != &Other)
+		{
+			FScopedUnrealAllocs UnrealAllocs;
+			Impl = MoveTemp(Other.Impl);
+		}
+
+		return *this;
+	}
+
+#if USE_USD_SDK
+	FUsdSkelBinding::FUsdSkelBinding(const pxr::UsdSkelBinding& InUsdSkelBinding)
+	{
+		FScopedUnrealAllocs UnrealAllocs;
+		Impl = MakeUnique<Internal::FUsdSkelBindingImpl>(InUsdSkelBinding);
+	}
+
+	FUsdSkelBinding::FUsdSkelBinding(pxr::UsdSkelBinding&& InUsdSkelBinding)
+	{
+		FScopedUnrealAllocs UnrealAllocs;
+		Impl = MakeUnique<Internal::FUsdSkelBindingImpl>(MoveTemp(InUsdSkelBinding));
+	}
+
+	FUsdSkelBinding::operator pxr::UsdSkelBinding&()
+	{
+		return Impl->PxrUsdSkelBinding.Get();
+	}
+
+	FUsdSkelBinding::operator const pxr::UsdSkelBinding&() const
+	{
+		return Impl->PxrUsdSkelBinding.Get();
+	}
+#endif	  // #if USE_USD_SDK
+
+	FUsdPrim FUsdSkelBinding::GetSkeleton() const
+	{
+#if USE_USD_SDK
+		FScopedUsdAllocs Allocs;
+		return FUsdPrim{Impl->PxrUsdSkelBinding.Get().GetSkeleton().GetPrim()};
+#else
+		return {};
+#endif	  // #if USE_USD_SDK
+	}
+
+	TArray<FUsdSkelSkinningQuery> FUsdSkelBinding::GetSkinningTargets() const
+	{
+#if USE_USD_SDK
+		TArray<FUsdSkelSkinningQuery> Result;
+
+		FScopedUsdAllocs Allocs;
+		pxr::VtArray<pxr::UsdSkelSkinningQuery> UsdQueries = Impl->PxrUsdSkelBinding.Get().GetSkinningTargets();
+
+		Result.Reserve(UsdQueries.size());
+		for (const pxr::UsdSkelSkinningQuery& UsdQuery : UsdQueries)
+		{
+			Result.Add(FUsdSkelSkinningQuery{UsdQuery});
+		}
+
+		return Result;
+#else
+		return {};
+#endif	  // #if USE_USD_SDK
+	}
+}	 // namespace UE
diff --git a/Engine/Plugins/Runtime/USDCore/Source/UnrealUSDWrapper/Private/UsdWrappers/UsdSkelBlendShape.cpp b/Engine/Plugins/Runtime/USDCore/Source/UnrealUSDWrapper/Private/UsdWrappers/UsdSkelBlendShape.cpp
new file mode 100644
--- /dev/null
+++ b/Engine/Plugins/Runtime/USDCore/Source/UnrealUSDWrapper/Private/UsdWrappers/UsdSkelBlendShape.cpp
@@ -0,0 +1,185 @@
+// Copyright Epic Games, Inc. All Rights Reserved.
+
+#include "UsdWrappers/UsdSkelBlendShape.h"
+
+#include "USDMemory.h"
+#include "UsdWrappers/UsdPrim.h"
+#include "UsdWrappers/UsdSkelBlendShape.h"
+#include "UsdWrappers/UsdSkelInbetweenShape.h"
+
+#if USE_USD_SDK
+#include "USDIncludesStart.h"
+#include "pxr/usd/usdSkel/blendShape.h"
+#include "USDIncludesEnd.h"
+#endif	  //  #if USE_USD_SDK
+
+namespace UE
+{
+	namespace Internal
+	{
+		class FUsdSkelBlendShapeImpl
+		{
+		public:
+			FUsdSkelBlendShapeImpl()
+#if USE_USD_SDK
+				: PxrUsdSkelBlendShape()
+#endif	  // #if USE_USD_SDK
+			{
+			}
+
+			FUsdSkelBlendShapeImpl(const UE::FUsdPrim& Prim)
+#if USE_USD_SDK
+				: PxrUsdSkelBlendShape(pxr::UsdSkelBlendShape{Prim})
+#endif	  // #if USE_USD_SDK
+			{
+			}
+
+#if USE_USD_SDK
+			explicit FUsdSkelBlendShapeImpl(const pxr::UsdSkelBlendShape& InUsdSkelBlendShape)
+				: PxrUsdSkelBlendShape(InUsdSkelBlendShape)
+			{
+			}
+
+			explicit FUsdSkelBlendShapeImpl(pxr::UsdSkelBlendShape&& InUsdSkelBlendShape)
+				: PxrUsdSkelBlendShape(MoveTemp(InUsdSkelBlendShape))
+			{
+			}
+
+			TUsdStore<pxr::UsdSkelBlendShape> PxrUsdSkelBlendShape;
+#endif	  // #if USE_USD_SDK
+		};
+	}	  // namespace Internal
+
+	FUsdSkelBlendShape::FUsdSkelBlendShape()
+		: FUsdTyped()
+	{
+		FScopedUnrealAllocs UnrealAllocs;
+		Impl = MakeUnique<Internal::FUsdSkelBlendShapeImpl>();
+	}
+
+	FUsdSkelBlendShape::FUsdSkelBlendShape(const UE::FUsdPrim& Prim)
+		: FUsdTyped(Prim)
+	{
+		FScopedUnrealAllocs UnrealAllocs;
+		Impl = MakeUnique<Internal::FUsdSkelBlendShapeImpl>(Prim);
+	}
+
+	FUsdSkelBlendShape::FUsdSkelBlendShape(const FUsdSkelBlendShape& Other)
+		: FUsdTyped(Other)
+	{
+#if USE_USD_SDK
+		FScopedUnrealAllocs UnrealAllocs;
+		Impl = MakeUnique<Internal::FUsdSkelBlendShapeImpl>(Other.Impl->PxrUsdSkelBlendShape.Get());
+#endif	  // #if USE_USD_SDK
+	}
+
+	FUsdSkelBlendShape::FUsdSkelBlendShape(FUsdSkelBlendShape&& Other)
+		: FUsdTyped(Other)
+		, Impl(MakeUnique<Internal::FUsdSkelBlendShapeImpl>())
+	{
+		if (this != &Other)
+		{
+			Impl = MoveTemp(Other.Impl);
+		}
+	}
+
+	FUsdSkelBlendShape::~FUsdSkelBlendShape()
+	{
+		FScopedUnrealAllocs UnrealAllocs;
+		Impl.Reset();
+	}
+
+	FUsdSkelBlendShape& FUsdSkelBlendShape::operator=(const FUsdSkelBlendShape& Other)
+	{
+		FUsdTyped::operator=(Other);
+
+#if USE_USD_SDK
+		FScopedUnrealAllocs UnrealAllocs;
+		Impl = MakeUnique<Internal::FUsdSkelBlendShapeImpl>(Other.Impl->PxrUsdSkelBlendShape.Get());
+#endif	  // #if USE_USD_SDK
+
+		return *this;
+	}
+
+	FUsdSkelBlendShape& FUsdSkelBlendShape::operator=(FUsdSkelBlendShape&& Other)
+	{
+		if (this != &Other)
+		{
+			// It's moving twice from Other but this should be fine, as FUsdTyped::operator=
+			// will only really move FUsdTyped::Impl, and shouldn't have any effect on
+			// Other's FUsdSkelBlendShape::Impl
+			FUsdTyped::operator=(MoveTemp(Other));
+
+			FScopedUnrealAllocs UnrealAllocs;
+			Impl = MoveTemp(Other.Impl);
+		}
+
+		return *this;
+	}
+
+	FUsdSkelBlendShape::operator bool() const
+	{
+#if USE_USD_SDK
+		return static_cast<bool>(Impl->PxrUsdSkelBlendShape.Get());
+#else
+		return false;
+#endif
+	}
+
+#if USE_USD_SDK
+	FUsdSkelBlendShape::FUsdSkelBlendShape(const pxr::UsdSkelBlendShape& InUsdSkelBlendShape)
+		: FUsdTyped(InUsdSkelBlendShape)
+	{
+		FScopedUnrealAllocs UnrealAllocs;
+		Impl = MakeUnique<Internal::FUsdSkelBlendShapeImpl>(InUsdSkelBlendShape);
+	}
+
+	FUsdSkelBlendShape::FUsdSkelBlendShape(pxr::UsdSkelBlendShape&& InUsdSkelBlendShape)
+		: FUsdTyped(MoveTemp(InUsdSkelBlendShape))
+	{
+		FScopedUnrealAllocs UnrealAllocs;
+		Impl = MakeUnique<Internal::FUsdSkelBlendShapeImpl>(MoveTemp(InUsdSkelBlendShape));
+	}
+
+	FUsdSkelBlendShape::operator pxr::UsdSkelBlendShape&()
+	{
+		return Impl->PxrUsdSkelBlendShape.Get();
+	}
+
+	FUsdSkelBlendShape::operator const pxr::UsdSkelBlendShape&() const
+	{
+		return Impl->PxrUsdSkelBlendShape.Get();
+	}
+#endif	  // #if USE_USD_SDK
+
+	FUsdSkelInbetweenShape FUsdSkelBlendShape::GetInbetween(const FString& Name) const
+	{
+#if USE_USD_SDK
+		FScopedUsdAllocs Allocs;
+
+		pxr::TfToken TokenName{TCHAR_TO_UTF8(*Name)};
+		return UE::FUsdSkelInbetweenShape{Impl->PxrUsdSkelBlendShape.Get().GetInbetween(TokenName)};
+#else
+		return {};
+#endif	  // #if USE_USD_SDK
+	}
+
+	TArray<FUsdSkelInbetweenShape> FUsdSkelBlendShape::GetInbetweens() const
+	{
+		TArray<FUsdSkelInbetweenShape> Result;
+
+#if USE_USD_SDK
+		FScopedUsdAllocs Allocs;
+
+		std::vector<pxr::UsdSkelInbetweenShape> UsdInbetweens = Impl->PxrUsdSkelBlendShape.Get().GetInbetweens();
+		Result.Reserve(UsdInbetweens.size());
+		for (const pxr::UsdSkelInbetweenShape& UsdInbetween : UsdInbetweens)
+		{
+			Result.Add(UE::FUsdSkelInbetweenShape{UsdInbetween});
+		}
+#endif	  // #if USE_USD_SDK
+
+		return Result;
+	}
+
+}	 // namespace UE
diff --git a/Engine/Plugins/Runtime/USDCore/Source/UnrealUSDWrapper/Private/UsdWrappers/UsdSkelBlendShapeQuery.cpp b/Engine/Plugins/Runtime/USDCore/Source/UnrealUSDWrapper/Private/UsdWrappers/UsdSkelBlendShapeQuery.cpp
new file mode 100644
--- /dev/null
+++ b/Engine/Plugins/Runtime/USDCore/Source/UnrealUSDWrapper/Private/UsdWrappers/UsdSkelBlendShapeQuery.cpp
@@ -0,0 +1,187 @@
+// Copyright Epic Games, Inc. All Rights Reserved.
+
+#include "UsdWrappers/UsdSkelBlendShapeQuery.h"
+
+#include "USDMemory.h"
+#include "UsdWrappers/UsdPrim.h"
+#include "UsdWrappers/UsdSkelBlendShape.h"
+#include "UsdWrappers/UsdSkelBlendShapeQuery.h"
+
+#if USE_USD_SDK
+#include "USDIncludesStart.h"
+#include "pxr/usd/usdSkel/bindingAPI.h"
+#include "pxr/usd/usdSkel/blendShapeQuery.h"
+#include "USDIncludesEnd.h"
+#endif	  //  #if USE_USD_SDK
+
+namespace UE
+{
+	namespace Internal
+	{
+		class FUsdSkelBlendShapeQueryImpl
+		{
+		public:
+			FUsdSkelBlendShapeQueryImpl()
+#if USE_USD_SDK
+				: PxrUsdSkelBlendShapeQuery()
+#endif	  // #if USE_USD_SDK
+			{
+			}
+
+			FUsdSkelBlendShapeQueryImpl(const UE::FUsdPrim& SkelBindingAPIPrim)
+#if USE_USD_SDK
+				: PxrUsdSkelBlendShapeQuery(pxr::UsdSkelBindingAPI{SkelBindingAPIPrim})
+#endif	  // #if USE_USD_SDK
+			{
+			}
+
+#if USE_USD_SDK
+			explicit FUsdSkelBlendShapeQueryImpl(const pxr::UsdSkelBlendShapeQuery& InUsdSkelBlendShapeQuery)
+				: PxrUsdSkelBlendShapeQuery(InUsdSkelBlendShapeQuery)
+			{
+			}
+
+			explicit FUsdSkelBlendShapeQueryImpl(pxr::UsdSkelBlendShapeQuery&& InUsdSkelBlendShapeQuery)
+				: PxrUsdSkelBlendShapeQuery(MoveTemp(InUsdSkelBlendShapeQuery))
+			{
+			}
+
+			TUsdStore<pxr::UsdSkelBlendShapeQuery> PxrUsdSkelBlendShapeQuery;
+#endif	  // #if USE_USD_SDK
+		};
+	}	  // namespace Internal
+
+	FUsdSkelBlendShapeQuery::FUsdSkelBlendShapeQuery()
+	{
+		FScopedUnrealAllocs UnrealAllocs;
+		Impl = MakeUnique<Internal::FUsdSkelBlendShapeQueryImpl>();
+	}
+
+	FUsdSkelBlendShapeQuery::FUsdSkelBlendShapeQuery(const UE::FUsdPrim& SkelBindingAPIPrim)
+	{
+		FScopedUnrealAllocs UnrealAllocs;
+		Impl = MakeUnique<Internal::FUsdSkelBlendShapeQueryImpl>(SkelBindingAPIPrim);
+	}
+
+	FUsdSkelBlendShapeQuery::FUsdSkelBlendShapeQuery(const FUsdSkelBlendShapeQuery& Other)
+	{
+#if USE_USD_SDK
+		FScopedUnrealAllocs UnrealAllocs;
+		Impl = MakeUnique<Internal::FUsdSkelBlendShapeQueryImpl>(Other.Impl->PxrUsdSkelBlendShapeQuery.Get());
+#endif	  // #if USE_USD_SDK
+	}
+
+	FUsdSkelBlendShapeQuery::FUsdSkelBlendShapeQuery(FUsdSkelBlendShapeQuery&& Other)
+		: Impl(MakeUnique<Internal::FUsdSkelBlendShapeQueryImpl>())
+	{
+		if (this != &Other)
+		{
+			Impl = MoveTemp(Other.Impl);
+		}
+	}
+
+	FUsdSkelBlendShapeQuery::~FUsdSkelBlendShapeQuery()
+	{
+		FScopedUnrealAllocs UnrealAllocs;
+		Impl.Reset();
+	}
+
+	FUsdSkelBlendShapeQuery& FUsdSkelBlendShapeQuery::operator=(const FUsdSkelBlendShapeQuery& Other)
+	{
+#if USE_USD_SDK
+		FScopedUnrealAllocs UnrealAllocs;
+		Impl = MakeUnique<Internal::FUsdSkelBlendShapeQueryImpl>(Other.Impl->PxrUsdSkelBlendShapeQuery.Get());
+#endif	  // #if USE_USD_SDK
+
+		return *this;
+	}
+
+	FUsdSkelBlendShapeQuery& FUsdSkelBlendShapeQuery::operator=(FUsdSkelBlendShapeQuery&& Other)
+	{
+		if (this != &Other)
+		{
+			FScopedUnrealAllocs UnrealAllocs;
+			Impl = MoveTemp(Other.Impl);
+		}
+
+		return *this;
+	}
+
+	FUsdSkelBlendShapeQuery::operator bool() const
+	{
+#if USE_USD_SDK
+		return static_cast<bool>(Impl->PxrUsdSkelBlendShapeQuery.Get());
+#else
+		return false;
+#endif
+	}
+
+#if USE_USD_SDK
+	FUsdSkelBlendShapeQuery::FUsdSkelBlendShapeQuery(const pxr::UsdSkelBlendShapeQuery& InUsdSkelBlendShapeQuery)
+	{
+		FScopedUnrealAllocs UnrealAllocs;
+		Impl = MakeUnique<Internal::FUsdSkelBlendShapeQueryImpl>(InUsdSkelBlendShapeQuery);
+	}
+
+	FUsdSkelBlendShapeQuery::FUsdSkelBlendShapeQuery(pxr::UsdSkelBlendShapeQuery&& InUsdSkelBlendShapeQuery)
+	{
+		FScopedUnrealAllocs UnrealAllocs;
+		Impl = MakeUnique<Internal::FUsdSkelBlendShapeQueryImpl>(MoveTemp(InUsdSkelBlendShapeQuery));
+	}
+
+	FUsdSkelBlendShapeQuery::operator pxr::UsdSkelBlendShapeQuery&()
+	{
+		return Impl->PxrUsdSkelBlendShapeQuery.Get();
+	}
+
+	FUsdSkelBlendShapeQuery::operator const pxr::UsdSkelBlendShapeQuery&() const
+	{
+		return Impl->PxrUsdSkelBlendShapeQuery.Get();
+	}
+#endif	  // #if USE_USD_SDK
+
+	bool FUsdSkelBlendShapeQuery::IsValid() const
+	{
+#if USE_USD_SDK
+		return Impl->PxrUsdSkelBlendShapeQuery.Get().IsValid();
+#else
+		return false;
+#endif	  // #if USE_USD_SDK
+	}
+
+	FUsdPrim FUsdSkelBlendShapeQuery::GetPrim() const
+	{
+#if USE_USD_SDK
+		FScopedUsdAllocs Allocs;
+		return FUsdPrim{Impl->PxrUsdSkelBlendShapeQuery.Get().GetPrim()};
+#else
+		return {};
+#endif	  // #if USE_USD_SDK
+	}
+
+	size_t FUsdSkelBlendShapeQuery::GetNumBlendShapes() const
+	{
+#if USE_USD_SDK
+		FScopedUsdAllocs Allocs;
+		return Impl->PxrUsdSkelBlendShapeQuery.Get().GetNumBlendShapes();
+#else
+		return 0;
+#endif	  // #if USE_USD_SDK
+	}
+
+	FUsdSkelBlendShape FUsdSkelBlendShapeQuery::GetBlendShape(size_t BlendShapeIndex) const
+	{
+#if USE_USD_SDK
+		FScopedUsdAllocs Allocs;
+
+		pxr::UsdSkelBlendShape BlendShape = Impl->PxrUsdSkelBlendShapeQuery.Get().GetBlendShape(BlendShapeIndex);
+		if (BlendShape)
+		{
+			return UE::FUsdSkelBlendShape{BlendShape};
+		}
+#endif	  // #if USE_USD_SDK
+
+		return {};
+	}
+
+}	 // namespace UE
diff --git a/Engine/Plugins/Runtime/USDCore/Source/UnrealUSDWrapper/Private/UsdWrappers/UsdSkelCache.cpp b/Engine/Plugins/Runtime/USDCore/Source/UnrealUSDWrapper/Private/UsdWrappers/UsdSkelCache.cpp
new file mode 100644
--- /dev/null
+++ b/Engine/Plugins/Runtime/USDCore/Source/UnrealUSDWrapper/Private/UsdWrappers/UsdSkelCache.cpp
@@ -0,0 +1,242 @@
+// Copyright Epic Games, Inc. All Rights Reserved.
+
+#include "UsdWrappers/UsdSkelCache.h"
+
+#include "USDMemory.h"
+#include "UsdWrappers/UsdPrim.h"
+#include "UsdWrappers/UsdSkelAnimQuery.h"
+#include "UsdWrappers/UsdSkelBinding.h"
+#include "UsdWrappers/UsdSkelSkeletonQuery.h"
+#include "UsdWrappers/UsdSkelSkinningQuery.h"
+
+#if USE_USD_SDK
+#include "USDIncludesStart.h"
+#include "pxr/usd/usdSkel/animation.h"
+#include "pxr/usd/usdSkel/cache.h"
+#include "pxr/usd/usdSkel/root.h"
+#include "pxr/usd/usdSkel/skeletonQuery.h"
+#include "USDIncludesEnd.h"
+#endif	  //  #if USE_USD_SDK
+
+namespace UE
+{
+	namespace Internal
+	{
+		class FUsdSkelCacheImpl
+		{
+		public:
+			FUsdSkelCacheImpl()
+#if USE_USD_SDK
+				: PxrUsdSkelCache()
+#endif	  // #if USE_USD_SDK
+			{
+			}
+
+#if USE_USD_SDK
+			explicit FUsdSkelCacheImpl(const pxr::UsdSkelCache& InUsdSkelCache)
+				: PxrUsdSkelCache(InUsdSkelCache)
+			{
+			}
+
+			explicit FUsdSkelCacheImpl(pxr::UsdSkelCache&& InUsdSkelCache)
+				: PxrUsdSkelCache(MoveTemp(InUsdSkelCache))
+			{
+			}
+
+			TUsdStore<pxr::UsdSkelCache> PxrUsdSkelCache;
+#endif	  // #if USE_USD_SDK
+		};
+	}	  // namespace Internal
+
+	FUsdSkelCache::FUsdSkelCache()
+	{
+		FScopedUnrealAllocs UnrealAllocs;
+		Impl = MakeUnique<Internal::FUsdSkelCacheImpl>();
+	}
+
+	FUsdSkelCache::FUsdSkelCache(const FUsdSkelCache& Other)
+	{
+#if USE_USD_SDK
+		FScopedUnrealAllocs UnrealAllocs;
+		Impl = MakeUnique<Internal::FUsdSkelCacheImpl>(Other.Impl->PxrUsdSkelCache.Get());
+#endif	  // #if USE_USD_SDK
+	}
+
+	FUsdSkelCache::FUsdSkelCache(FUsdSkelCache&& Other)
+		: Impl(MakeUnique<Internal::FUsdSkelCacheImpl>())
+	{
+		if (this != &Other)
+		{
+			Impl = MoveTemp(Other.Impl);
+		}
+	}
+
+	FUsdSkelCache::~FUsdSkelCache()
+	{
+		FScopedUnrealAllocs UnrealAllocs;
+		Impl.Reset();
+	}
+
+	FUsdSkelCache& FUsdSkelCache::operator=(const FUsdSkelCache& Other)
+	{
+#if USE_USD_SDK
+		FScopedUnrealAllocs UnrealAllocs;
+		Impl = MakeUnique<Internal::FUsdSkelCacheImpl>(Other.Impl->PxrUsdSkelCache.Get());
+#endif	  // #if USE_USD_SDK
+
+		return *this;
+	}
+
+	FUsdSkelCache& FUsdSkelCache::operator=(FUsdSkelCache&& Other)
+	{
+		if (this != &Other)
+		{
+			FScopedUnrealAllocs UnrealAllocs;
+			Impl = MoveTemp(Other.Impl);
+		}
+
+		return *this;
+	}
+
+#if USE_USD_SDK
+	FUsdSkelCache::FUsdSkelCache(const pxr::UsdSkelCache& InUsdSkelCache)
+	{
+		FScopedUnrealAllocs UnrealAllocs;
+		Impl = MakeUnique<Internal::FUsdSkelCacheImpl>(InUsdSkelCache);
+	}
+
+	FUsdSkelCache::FUsdSkelCache(pxr::UsdSkelCache&& InUsdSkelCache)
+	{
+		FScopedUnrealAllocs UnrealAllocs;
+		Impl = MakeUnique<Internal::FUsdSkelCacheImpl>(MoveTemp(InUsdSkelCache));
+	}
+
+	FUsdSkelCache::operator pxr::UsdSkelCache&()
+	{
+		return Impl->PxrUsdSkelCache.Get();
+	}
+
+	FUsdSkelCache::operator const pxr::UsdSkelCache&() const
+	{
+		return Impl->PxrUsdSkelCache.Get();
+	}
+#endif	  // #if USE_USD_SDK
+
+	void FUsdSkelCache::Clear()
+	{
+#if USE_USD_SDK
+		FScopedUsdAllocs Allocs;
+		Impl->PxrUsdSkelCache.Get().Clear();
+#endif	  // #if USE_USD_SDK
+	}
+
+	bool FUsdSkelCache::Populate(const UE::FUsdPrim& SkelRootPrim, bool bTraverseInstanceProxies)
+	{
+#if USE_USD_SDK
+		FScopedUsdAllocs Allocs;
+
+		pxr::Usd_PrimFlagsPredicate Predicate = pxr::UsdPrimDefaultPredicate;
+		if (bTraverseInstanceProxies)
+		{
+			Predicate = pxr::UsdTraverseInstanceProxies();
+		}
+
+		return Impl->PxrUsdSkelCache.Get().Populate(pxr::UsdSkelRoot{SkelRootPrim}, Predicate);
+#else
+		return false;
+#endif	  // #if USE_USD_SDK
+	}
+
+	UE::FUsdSkelSkeletonQuery FUsdSkelCache::GetSkelQuery(const UE::FUsdPrim& SkeletonPrim) const
+	{
+#if USE_USD_SDK
+		FScopedUsdAllocs Allocs;
+		return UE::FUsdSkelSkeletonQuery{Impl->PxrUsdSkelCache.Get().GetSkelQuery(pxr::UsdSkelSkeleton{SkeletonPrim})};
+#else
+		return UE::FUsdSkelSkeletonQuery{};
+#endif	  // #if USE_USD_SDK
+	}
+
+	UE::FUsdSkelAnimQuery FUsdSkelCache::GetAnimQuery(const UE::FUsdPrim& SkelAnimationPrim) const
+	{
+#if USE_USD_SDK
+		FScopedUsdAllocs Allocs;
+		return UE::FUsdSkelAnimQuery{Impl->PxrUsdSkelCache.Get().GetAnimQuery(pxr::UsdSkelAnimation{SkelAnimationPrim})};
+#else
+		return UE::FUsdSkelAnimQuery{};
+#endif	  // #if USE_USD_SDK
+	}
+
+	UE::FUsdSkelSkinningQuery FUsdSkelCache::GetSkinningQuery(const UE::FUsdPrim& SkinnedPrim) const
+	{
+#if USE_USD_SDK
+		FScopedUsdAllocs Allocs;
+		return UE::FUsdSkelSkinningQuery{Impl->PxrUsdSkelCache.Get().GetSkinningQuery(SkinnedPrim)};
+#else
+		return UE::FUsdSkelSkinningQuery{};
+#endif	  // #if USE_USD_SDK
+	}
+
+	bool FUsdSkelCache::ComputeSkelBindings(
+		const UE::FUsdPrim& InSkelRootPrim,
+		TArray<UE::FUsdSkelBinding>& OutBindings,
+		bool bTraverseInstanceProxies
+	) const
+	{
+#if USE_USD_SDK
+		FScopedUsdAllocs Allocs;
+
+		std::vector<pxr::UsdSkelBinding> UsdBindings;
+
+		pxr::Usd_PrimFlagsPredicate Predicate = pxr::UsdPrimDefaultPredicate;
+		if (bTraverseInstanceProxies)
+		{
+			Predicate = pxr::UsdTraverseInstanceProxies();
+		}
+
+		bool bSuccess = Impl->PxrUsdSkelCache.Get().ComputeSkelBindings(pxr::UsdSkelRoot{InSkelRootPrim}, &UsdBindings, Predicate);
+		if (!bSuccess)
+		{
+			return false;
+		}
+
+		OutBindings.SetNum(UsdBindings.size());
+		for (int32 Index = 0; Index < OutBindings.Num(); ++Index)
+		{
+			OutBindings[Index] = UE::FUsdSkelBinding{UsdBindings[Index]};
+		}
+
+		return true;
+#else
+		return false;
+#endif	  // #if USE_USD_SDK
+	}
+
+	bool FUsdSkelCache::ComputeSkelBinding(
+		const UE::FUsdPrim& InSkelRootPrim,
+		const UE::FUsdPrim& InSkeletonPrim,
+		UE::FUsdSkelBinding& OutBinding,
+		bool bTraverseInstanceProxies
+	) const
+	{
+#if USE_USD_SDK
+		FScopedUsdAllocs Allocs;
+
+		pxr::Usd_PrimFlagsPredicate Predicate = pxr::UsdPrimDefaultPredicate;
+		if (bTraverseInstanceProxies)
+		{
+			Predicate = pxr::UsdTraverseInstanceProxies();
+		}
+
+		return Impl->PxrUsdSkelCache.Get().ComputeSkelBinding(
+			pxr::UsdSkelRoot{InSkelRootPrim},
+			pxr::UsdSkelSkeleton{InSkeletonPrim},
+			&static_cast<pxr::UsdSkelBinding&>(OutBinding),
+			Predicate
+		);
+#else
+		return false;
+#endif	  // #if USE_USD_SDK
+	}
+
+}	 // namespace UE
diff --git a/Engine/Plugins/Runtime/USDCore/Source/UnrealUSDWrapper/Private/UsdWrappers/UsdSkelInbetweenShape.cpp b/Engine/Plugins/Runtime/USDCore/Source/UnrealUSDWrapper/Private/UsdWrappers/UsdSkelInbetweenShape.cpp
new file mode 100644
--- /dev/null
+++ b/Engine/Plugins/Runtime/USDCore/Source/UnrealUSDWrapper/Private/UsdWrappers/UsdSkelInbetweenShape.cpp
@@ -0,0 +1,155 @@
+// Copyright Epic Games, Inc. All Rights Reserved.
+
+#include "UsdWrappers/UsdSkelInbetweenShape.h"
+
+#include "USDMemory.h"
+#include "UsdWrappers/UsdAttribute.h"
+#include "UsdWrappers/UsdPrim.h"
+#include "UsdWrappers/UsdSkelInbetweenShape.h"
+
+#if USE_USD_SDK
+#include "USDIncludesStart.h"
+#include "pxr/usd/usdSkel/inbetweenShape.h"
+#include "USDIncludesEnd.h"
+#endif	  //  #if USE_USD_SDK
+
+namespace UE
+{
+	namespace Internal
+	{
+		class FUsdSkelInbetweenShapeImpl
+		{
+		public:
+			FUsdSkelInbetweenShapeImpl()
+#if USE_USD_SDK
+				: PxrUsdSkelInbetweenShape()
+#endif	  // #if USE_USD_SDK
+			{
+			}
+
+#if USE_USD_SDK
+			explicit FUsdSkelInbetweenShapeImpl(const pxr::UsdSkelInbetweenShape& InUsdSkelInbetweenShape)
+				: PxrUsdSkelInbetweenShape(InUsdSkelInbetweenShape)
+			{
+			}
+
+			explicit FUsdSkelInbetweenShapeImpl(pxr::UsdSkelInbetweenShape&& InUsdSkelInbetweenShape)
+				: PxrUsdSkelInbetweenShape(MoveTemp(InUsdSkelInbetweenShape))
+			{
+			}
+
+			TUsdStore<pxr::UsdSkelInbetweenShape> PxrUsdSkelInbetweenShape;
+#endif	  // #if USE_USD_SDK
+		};
+	}	  // namespace Internal
+
+	FUsdSkelInbetweenShape::FUsdSkelInbetweenShape()
+	{
+		FScopedUnrealAllocs UnrealAllocs;
+		Impl = MakeUnique<Internal::FUsdSkelInbetweenShapeImpl>();
+	}
+
+	FUsdSkelInbetweenShape::FUsdSkelInbetweenShape(const FUsdSkelInbetweenShape& Other)
+	{
+#if USE_USD_SDK
+		FScopedUnrealAllocs UnrealAllocs;
+		Impl = MakeUnique<Internal::FUsdSkelInbetweenShapeImpl>(Other.Impl->PxrUsdSkelInbetweenShape.Get());
+#endif	  // #if USE_USD_SDK
+	}
+
+	FUsdSkelInbetweenShape::FUsdSkelInbetweenShape(FUsdSkelInbetweenShape&& Other)
+		: Impl(MakeUnique<Internal::FUsdSkelInbetweenShapeImpl>())
+	{
+		if (this != &Other)
+		{
+			Impl = MoveTemp(Other.Impl);
+		}
+	}
+
+	FUsdSkelInbetweenShape::~FUsdSkelInbetweenShape()
+	{
+		FScopedUnrealAllocs UnrealAllocs;
+		Impl.Reset();
+	}
+
+	FUsdSkelInbetweenShape& FUsdSkelInbetweenShape::operator=(const FUsdSkelInbetweenShape& Other)
+	{
+#if USE_USD_SDK
+		FScopedUnrealAllocs UnrealAllocs;
+		Impl = MakeUnique<Internal::FUsdSkelInbetweenShapeImpl>(Other.Impl->PxrUsdSkelInbetweenShape.Get());
+#endif	  // #if USE_USD_SDK
+
+		return *this;
+	}
+
+	FUsdSkelInbetweenShape& FUsdSkelInbetweenShape::operator=(FUsdSkelInbetweenShape&& Other)
+	{
+		if (this != &Other)
+		{
+			FScopedUnrealAllocs UnrealAllocs;
+			Impl = MoveTemp(Other.Impl);
+		}
+
+		return *this;
+	}
+
+	FUsdSkelInbetweenShape::operator bool() const
+	{
+#if USE_USD_SDK
+		return static_cast<bool>(Impl->PxrUsdSkelInbetweenShape.Get());
+#else
+		return false;
+#endif
+	}
+
+#if USE_USD_SDK
+	FUsdSkelInbetweenShape::FUsdSkelInbetweenShape(const pxr::UsdSkelInbetweenShape& InUsdSkelInbetweenShape)
+	{
+		FScopedUnrealAllocs UnrealAllocs;
+		Impl = MakeUnique<Internal::FUsdSkelInbetweenShapeImpl>(InUsdSkelInbetweenShape);
+	}
+
+	FUsdSkelInbetweenShape::FUsdSkelInbetweenShape(pxr::UsdSkelInbetweenShape&& InUsdSkelInbetweenShape)
+	{
+		FScopedUnrealAllocs UnrealAllocs;
+		Impl = MakeUnique<Internal::FUsdSkelInbetweenShapeImpl>(MoveTemp(InUsdSkelInbetweenShape));
+	}
+
+	FUsdSkelInbetweenShape::operator pxr::UsdSkelInbetweenShape&()
+	{
+		return Impl->PxrUsdSkelInbetweenShape.Get();
+	}
+
+	FUsdSkelInbetweenShape::operator const pxr::UsdSkelInbetweenShape&() const
+	{
+		return Impl->PxrUsdSkelInbetweenShape.Get();
+	}
+#endif	  // #if USE_USD_SDK
+
+	bool FUsdSkelInbetweenShape::GetWeight(float* OutWeight) const
+	{
+#if USE_USD_SDK
+		return Impl->PxrUsdSkelInbetweenShape.Get().GetWeight(OutWeight);
+#else
+		return false;
+#endif	  // #if USE_USD_SDK
+	}
+
+	UE::FUsdAttribute FUsdSkelInbetweenShape::GetAttr() const
+	{
+#if USE_USD_SDK
+		return UE::FUsdAttribute{Impl->PxrUsdSkelInbetweenShape.Get().GetAttr()};
+#else
+		return {};
+#endif	  // #if USE_USD_SDK
+	}
+
+	bool FUsdSkelInbetweenShape::IsDefined() const
+	{
+#if USE_USD_SDK
+		return Impl->PxrUsdSkelInbetweenShape.Get().IsDefined();
+#else
+		return false;
+#endif	  // #if USE_USD_SDK
+	}
+}	 // namespace UE
diff --git a/Engine/Plugins/Runtime/USDCore/Source/UnrealUSDWrapper/Private/UsdWrappers/UsdSkelSkeletonQuery.cpp b/Engine/Plugins/Runtime/USDCore/Source/UnrealUSDWrapper/Private/UsdWrappers/UsdSkelSkeletonQuery.cpp
new file mode 100644
--- /dev/null
+++ b/Engine/Plugins/Runtime/USDCore/Source/UnrealUSDWrapper/Private/UsdWrappers/UsdSkelSkeletonQuery.cpp
@@ -0,0 +1,168 @@
+// Copyright Epic Games, Inc. All Rights Reserved.
+
+#include "UsdWrappers/UsdSkelSkeletonQuery.h"
+
+#include "USDMemory.h"
+#include "UsdWrappers/UsdPrim.h"
+#include "UsdWrappers/UsdSkelAnimQuery.h"
+#include "UsdWrappers/UsdSkelSkeletonQuery.h"
+
+#if USE_USD_SDK
+#include "USDIncludesStart.h"
+#include "pxr/usd/usdSkel/skeleton.h"
+#include "pxr/usd/usdSkel/skeletonQuery.h"
+#include "USDIncludesEnd.h"
+#endif	  //  #if USE_USD_SDK
+
+namespace UE
+{
+	namespace Internal
+	{
+		class FUsdSkelSkeletonQueryImpl
+		{
+		public:
+			FUsdSkelSkeletonQueryImpl()
+#if USE_USD_SDK
+				: PxrUsdSkelSkeletonQuery()
+#endif	  // #if USE_USD_SDK
+			{
+			}
+
+#if USE_USD_SDK
+			explicit FUsdSkelSkeletonQueryImpl(const pxr::UsdSkelSkeletonQuery& InUsdSkelSkeletonQuery)
+				: PxrUsdSkelSkeletonQuery(InUsdSkelSkeletonQuery)
+			{
+			}
+
+			explicit FUsdSkelSkeletonQueryImpl(pxr::UsdSkelSkeletonQuery&& InUsdSkelSkeletonQuery)
+				: PxrUsdSkelSkeletonQuery(MoveTemp(InUsdSkelSkeletonQuery))
+			{
+			}
+
+			TUsdStore<pxr::UsdSkelSkeletonQuery> PxrUsdSkelSkeletonQuery;
+#endif	  // #if USE_USD_SDK
+		};
+	}	  // namespace Internal
+
+	FUsdSkelSkeletonQuery::FUsdSkelSkeletonQuery()
+	{
+		FScopedUnrealAllocs UnrealAllocs;
+		Impl = MakeUnique<Internal::FUsdSkelSkeletonQueryImpl>();
+	}
+
+	FUsdSkelSkeletonQuery::FUsdSkelSkeletonQuery(const FUsdSkelSkeletonQuery& Other)
+	{
+#if USE_USD_SDK
+		FScopedUnrealAllocs UnrealAllocs;
+		Impl = MakeUnique<Internal::FUsdSkelSkeletonQueryImpl>(Other.Impl->PxrUsdSkelSkeletonQuery.Get());
+#endif	  // #if USE_USD_SDK
+	}
+
+	FUsdSkelSkeletonQuery::FUsdSkelSkeletonQuery(FUsdSkelSkeletonQuery&& Other)
+		: Impl(MakeUnique<Internal::FUsdSkelSkeletonQueryImpl>())
+	{
+		if (this != &Other)
+		{
+			Impl = MoveTemp(Other.Impl);
+		}
+	}
+
+	FUsdSkelSkeletonQuery::~FUsdSkelSkeletonQuery()
+	{
+		FScopedUnrealAllocs UnrealAllocs;
+		Impl.Reset();
+	}
+
+	FUsdSkelSkeletonQuery& FUsdSkelSkeletonQuery::operator=(const FUsdSkelSkeletonQuery& Other)
+	{
+#if USE_USD_SDK
+		FScopedUnrealAllocs UnrealAllocs;
+		Impl = MakeUnique<Internal::FUsdSkelSkeletonQueryImpl>(Other.Impl->PxrUsdSkelSkeletonQuery.Get());
+#endif	  // #if USE_USD_SDK
+
+		return *this;
+	}
+
+	FUsdSkelSkeletonQuery& FUsdSkelSkeletonQuery::operator=(FUsdSkelSkeletonQuery&& Other)
+	{
+		if (this != &Other)
+		{
+			FScopedUnrealAllocs UnrealAllocs;
+			Impl = MoveTemp(Other.Impl);
+		}
+
+		return *this;
+	}
+
+	FUsdSkelSkeletonQuery::operator bool() const
+	{
+#if USE_USD_SDK
+		return static_cast<bool>(Impl->PxrUsdSkelSkeletonQuery.Get());
+#else
+		return false;
+#endif
+	}
+
+#if USE_USD_SDK
+	FUsdSkelSkeletonQuery::FUsdSkelSkeletonQuery(const pxr::UsdSkelSkeletonQuery& InUsdSkelSkeletonQuery)
+	{
+		FScopedUnrealAllocs UnrealAllocs;
+		Impl = MakeUnique<Internal::FUsdSkelSkeletonQueryImpl>(InUsdSkelSkeletonQuery);
+	}
+
+	FUsdSkelSkeletonQuery::FUsdSkelSkeletonQuery(pxr::UsdSkelSkeletonQuery&& InUsdSkelSkeletonQuery)
+	{
+		FScopedUnrealAllocs UnrealAllocs;
+		Impl = MakeUnique<Internal::FUsdSkelSkeletonQueryImpl>(MoveTemp(InUsdSkelSkeletonQuery));
+	}
+
+	FUsdSkelSkeletonQuery::operator pxr::UsdSkelSkeletonQuery&()
+	{
+		return Impl->PxrUsdSkelSkeletonQuery.Get();
+	}
+
+	FUsdSkelSkeletonQuery::operator const pxr::UsdSkelSkeletonQuery&() const
+	{
+		return Impl->PxrUsdSkelSkeletonQuery.Get();
+	}
+#endif	  // #if USE_USD_SDK
+
+	bool FUsdSkelSkeletonQuery::IsValid() const
+	{
+#if USE_USD_SDK
+		return Impl->PxrUsdSkelSkeletonQuery.Get().IsValid();
+#else
+		return false;
+#endif	  // #if USE_USD_SDK
+	}
+
+	FUsdPrim FUsdSkelSkeletonQuery::GetPrim() const
+	{
+#if USE_USD_SDK
+		FScopedUsdAllocs Allocs;
+		return FUsdPrim{Impl->PxrUsdSkelSkeletonQuery.Get().GetPrim()};
+#else
+		return {};
+#endif	  // #if USE_USD_SDK
+	}
+
+	FUsdPrim FUsdSkelSkeletonQuery::GetSkeleton() const
+	{
+#if USE_USD_SDK
+		FScopedUsdAllocs Allocs;
+		return FUsdPrim{Impl->PxrUsdSkelSkeletonQuery.Get().GetSkeleton().GetPrim()};
+#else
+		return {};
+#endif	  // #if USE_USD_SDK
+	}
+
+	FUsdSkelAnimQuery FUsdSkelSkeletonQuery::GetAnimQuery() const
+	{
+#if USE_USD_SDK
+		FScopedUsdAllocs Allocs;
+		return FUsdSkelAnimQuery{Impl->PxrUsdSkelSkeletonQuery.Get().GetAnimQuery()};
+#else
+		return {};
+#endif	  // #if USE_USD_SDK
+	}
+}	 // namespace UE
diff --git a/Engine/Plugins/Runtime/USDCore/Source/UnrealUSDWrapper/Private/UsdWrappers/UsdSkelSkinningQuery.cpp b/Engine/Plugins/Runtime/USDCore/Source/UnrealUSDWrapper/Private/UsdWrappers/UsdSkelSkinningQuery.cpp
new file mode 100644
--- /dev/null
+++ b/Engine/Plugins/Runtime/USDCore/Source/UnrealUSDWrapper/Private/UsdWrappers/UsdSkelSkinningQuery.cpp
@@ -0,0 +1,146 @@
+// Copyright Epic Games, Inc. All Rights Reserved.
+
+#include "UsdWrappers/UsdSkelSkinningQuery.h"
+
+#include "USDMemory.h"
+#include "UsdWrappers/UsdPrim.h"
+#include "UsdWrappers/UsdSkelSkinningQuery.h"
+
+#if USE_USD_SDK
+#include "USDIncludesStart.h"
+#include "pxr/usd/usdSkel/skinningQuery.h"
+#include "USDIncludesEnd.h"
+#endif	  //  #if USE_USD_SDK
+
+namespace UE
+{
+	namespace Internal
+	{
+		class FUsdSkelSkinningQueryImpl
+		{
+		public:
+			FUsdSkelSkinningQueryImpl()
+#if USE_USD_SDK
+				: PxrUsdSkelSkinningQuery()
+#endif	  // #if USE_USD_SDK
+			{
+			}
+
+#if USE_USD_SDK
+			explicit FUsdSkelSkinningQueryImpl(const pxr::UsdSkelSkinningQuery& InUsdSkelSkinningQuery)
+				: PxrUsdSkelSkinningQuery(InUsdSkelSkinningQuery)
+			{
+			}
+
+			explicit FUsdSkelSkinningQueryImpl(pxr::UsdSkelSkinningQuery&& InUsdSkelSkinningQuery)
+				: PxrUsdSkelSkinningQuery(MoveTemp(InUsdSkelSkinningQuery))
+			{
+			}
+
+			TUsdStore<pxr::UsdSkelSkinningQuery> PxrUsdSkelSkinningQuery;
+#endif	  // #if USE_USD_SDK
+		};
+	}	  // namespace Internal
+
+	FUsdSkelSkinningQuery::FUsdSkelSkinningQuery()
+	{
+		FScopedUnrealAllocs UnrealAllocs;
+		Impl = MakeUnique<Internal::FUsdSkelSkinningQueryImpl>();
+	}
+
+	FUsdSkelSkinningQuery::FUsdSkelSkinningQuery(const FUsdSkelSkinningQuery& Other)
+	{
+#if USE_USD_SDK
+		FScopedUnrealAllocs UnrealAllocs;
+		Impl = MakeUnique<Internal::FUsdSkelSkinningQueryImpl>(Other.Impl->PxrUsdSkelSkinningQuery.Get());
+#endif	  // #if USE_USD_SDK
+	}
+
+	FUsdSkelSkinningQuery::FUsdSkelSkinningQuery(FUsdSkelSkinningQuery&& Other)
+		: Impl(MakeUnique<Internal::FUsdSkelSkinningQueryImpl>())
+	{
+		if (this != &Other)
+		{
+			Impl = MoveTemp(Other.Impl);
+		}
+	}
+
+	FUsdSkelSkinningQuery::~FUsdSkelSkinningQuery()
+	{
+		FScopedUnrealAllocs UnrealAllocs;
+		Impl.Reset();
+	}
+
+	FUsdSkelSkinningQuery& FUsdSkelSkinningQuery::operator=(const FUsdSkelSkinningQuery& Other)
+	{
+#if USE_USD_SDK
+		FScopedUnrealAllocs UnrealAllocs;
+		Impl = MakeUnique<Internal::FUsdSkelSkinningQueryImpl>(Other.Impl->PxrUsdSkelSkinningQuery.Get());
+#endif	  // #if USE_USD_SDK
+
+		return *this;
+	}
+
+	FUsdSkelSkinningQuery& FUsdSkelSkinningQuery::operator=(FUsdSkelSkinningQuery&& Other)
+	{
+		if (this != &Other)
+		{
+			FScopedUnrealAllocs UnrealAllocs;
+			Impl = MoveTemp(Other.Impl);
+		}
+
+		return *this;
+	}
+
+	FUsdSkelSkinningQuery::operator bool() const
+	{
+#if USE_USD_SDK
+		return static_cast<bool>(Impl->PxrUsdSkelSkinningQuery.Get());
+#else
+		return false;
+#endif
+	}
+
+#if USE_USD_SDK
+	FUsdSkelSkinningQuery::FUsdSkelSkinningQuery(const pxr::UsdSkelSkinningQuery& InUsdSkelSkinningQuery)
+	{
+		FScopedUnrealAllocs UnrealAllocs;
+		Impl = MakeUnique<Internal::FUsdSkelSkinningQueryImpl>(InUsdSkelSkinningQuery);
+	}
+
+	FUsdSkelSkinningQuery::FUsdSkelSkinningQuery(pxr::UsdSkelSkinningQuery&& InUsdSkelSkinningQuery)
+	{
+		FScopedUnrealAllocs UnrealAllocs;
+		Impl = MakeUnique<Internal::FUsdSkelSkinningQueryImpl>(MoveTemp(InUsdSkelSkinningQuery));
+	}
+
+	FUsdSkelSkinningQuery::operator pxr::UsdSkelSkinningQuery&()
+	{
+		return Impl->PxrUsdSkelSkinningQuery.Get();
+	}
+
+	FUsdSkelSkinningQuery::operator const pxr::UsdSkelSkinningQuery&() const
+	{
+		return Impl->PxrUsdSkelSkinningQuery.Get();
+	}
+#endif	  // #if USE_USD_SDK
+
+	bool FUsdSkelSkinningQuery::IsValid() const
+	{
+#if USE_USD_SDK
+		return Impl->PxrUsdSkelSkinningQuery.Get().IsValid();
+#else
+		return false;
+#endif	  // #if USE_USD_SDK
+	}
+
+	FUsdPrim FUsdSkelSkinningQuery::GetPrim() const
+	{
+#if USE_USD_SDK
+		FScopedUsdAllocs Allocs;
+		return FUsdPrim{Impl->PxrUsdSkelSkinningQuery.Get().GetPrim()};
+#else
+		return {};
+#endif	  // #if USE_USD_SDK
+	}
+}	 // namespace UE
diff --git a/Engine/Plugins/Runtime/USDCore/Source/UnrealUSDWrapper/Public/UsdWrappers/UsdSkelAnimQuery.h b/Engine/Plugins/Runtime/USDCore/Source/UnrealUSDWrapper/Public/UsdWrappers/UsdSkelAnimQuery.h
new file mode 100644
--- /dev/null
+++ b/Engine/Plugins/Runtime/USDCore/Source/UnrealUSDWrapper/Public/UsdWrappers/UsdSkelAnimQuery.h
@@ -0,0 +1,65 @@
+// Copyright Epic Games, Inc. All Rights Reserved.
+
+#pragma once
+
+#include "CoreTypes.h"
+#include "Templates/UniquePtr.h"
+
+#if USE_USD_SDK
+
+#include "USDIncludesStart.h"
+#include "pxr/pxr.h"
+#include "USDIncludesEnd.h"
+
+PXR_NAMESPACE_OPEN_SCOPE
+	class UsdPrim;
+	class UsdSkelAnimQuery;
+PXR_NAMESPACE_CLOSE_SCOPE
+
+#endif	  // #if USE_USD_SDK
+
+namespace UE
+{
+	class FUsdPrim;
+
+	namespace Internal
+	{
+		class FUsdSkelAnimQueryImpl;
+	}
+
+	/**
+	 * Minimal pxr::UsdSkelAnimQuery wrapper for Unreal that can be used from no-rtti modules.
+	 */
+	class UNREALUSDWRAPPER_API FUsdSkelAnimQuery
+	{
+	public:
+		FUsdSkelAnimQuery();
+		FUsdSkelAnimQuery(const FUsdSkelAnimQuery& Other);
+		FUsdSkelAnimQuery(FUsdSkelAnimQuery&& Other);
+		~FUsdSkelAnimQuery();
+
+		FUsdSkelAnimQuery& operator=(const FUsdSkelAnimQuery& Other);
+		FUsdSkelAnimQuery& operator=(FUsdSkelAnimQuery&& Other);
+
+		explicit operator bool() const;
+
+		// Auto conversion from/to pxr::UsdSkelAnimQuery
+	public:
+#if USE_USD_SDK
+		explicit FUsdSkelAnimQuery(const pxr::UsdSkelAnimQuery& InUsdSkelAnimQuery);
+		explicit FUsdSkelAnimQuery(pxr::UsdSkelAnimQuery&& InUsdSkelAnimQuery);
+
+		operator pxr::UsdSkelAnimQuery&();
+		operator const pxr::UsdSkelAnimQuery&() const;
+#endif
+
+		// Wrapped pxr::UsdSkelAnimQuery functions, refer to the USD SDK documentation
+	public:
+		bool IsValid() const;
+
+		UE::FUsdPrim GetPrim() const;
+
+	private:
+		TUniquePtr<Internal::FUsdSkelAnimQueryImpl> Impl;
+	};
+}	 // namespace UE
diff --git a/Engine/Plugins/Runtime/USDCore/Source/UnrealUSDWrapper/Public/UsdWrappers/UsdSkelBinding.h b/Engine/Plugins/Runtime/USDCore/Source/UnrealUSDWrapper/Public/UsdWrappers/UsdSkelBinding.h
new file mode 100644
--- /dev/null
+++ b/Engine/Plugins/Runtime/USDCore/Source/UnrealUSDWrapper/Public/UsdWrappers/UsdSkelBinding.h
@@ -0,0 +1,64 @@
+// Copyright Epic Games, Inc. All Rights Reserved.
+
+#pragma once
+
+#include "CoreTypes.h"
+#include "Templates/UniquePtr.h"
+
+#if USE_USD_SDK
+
+#include "USDIncludesStart.h"
+#include "pxr/pxr.h"
+#include "USDIncludesEnd.h"
+
+PXR_NAMESPACE_OPEN_SCOPE
+	class UsdPrim;
+	class UsdSkelBinding;
+PXR_NAMESPACE_CLOSE_SCOPE
+
+#endif	  // #if USE_USD_SDK
+
+namespace UE
+{
+	class FUsdPrim;
+	class FUsdSkelSkinningQuery;
+
+	namespace Internal
+	{
+		class FUsdSkelBindingImpl;
+	}
+
+	/**
+	 * Minimal pxr::UsdSkelBinding wrapper for Unreal that can be used from no-rtti modules.
+	 */
+	class UNREALUSDWRAPPER_API FUsdSkelBinding
+	{
+	public:
+		FUsdSkelBinding();
+		FUsdSkelBinding(const FUsdSkelBinding& Other);
+		FUsdSkelBinding(FUsdSkelBinding&& Other);
+		~FUsdSkelBinding();
+
+		FUsdSkelBinding& operator=(const FUsdSkelBinding& Other);
+		FUsdSkelBinding& operator=(FUsdSkelBinding&& Other);
+
+		// Auto conversion from/to pxr::UsdSkelBinding
+	public:
+#if USE_USD_SDK
+		explicit FUsdSkelBinding(const pxr::UsdSkelBinding& InUsdSkelBinding);
+		explicit FUsdSkelBinding(pxr::UsdSkelBinding&& InUsdSkelBinding);
+
+		operator pxr::UsdSkelBinding&();
+		operator const pxr::UsdSkelBinding&() const;
+#endif
+
+		// Wrapped pxr::UsdSkelBinding functions, refer to the USD SDK documentation
+	public:
+		UE::FUsdPrim GetSkeleton() const;
+
+		TArray<UE::FUsdSkelSkinningQuery> GetSkinningTargets() const;
+
+	private:
+		TUniquePtr<Internal::FUsdSkelBindingImpl> Impl;
+	};
+}	 // namespace UE
diff --git a/Engine/Plugins/Runtime/USDCore/Source/UnrealUSDWrapper/Public/UsdWrappers/UsdSkelBlendShape.h b/Engine/Plugins/Runtime/USDCore/Source/UnrealUSDWrapper/Public/UsdWrappers/UsdSkelBlendShape.h
new file mode 100644
--- /dev/null
+++ b/Engine/Plugins/Runtime/USDCore/Source/UnrealUSDWrapper/Public/UsdWrappers/UsdSkelBlendShape.h
@@ -0,0 +1,65 @@
+// Copyright Epic Games, Inc. All Rights Reserved.
+
+#pragma once
+
+#include "CoreMinimal.h"
+
+#include "UsdTyped.h"
+
+#if USE_USD_SDK
+#include "USDIncludesStart.h"
+#include "pxr/pxr.h"
+#include "USDIncludesEnd.h"
+
+PXR_NAMESPACE_OPEN_SCOPE
+	class UsdPrim;
+	class UsdSkelBlendShape;
+PXR_NAMESPACE_CLOSE_SCOPE
+#endif	  // #if USE_USD_SDK
+
+namespace UE
+{
+	class FUsdSkelInbetweenShape;
+
+	namespace Internal
+	{
+		class FUsdSkelBlendShapeImpl;
+	}
+
+	/**
+	 * Minimal pxr::UsdSkelBlendShape wrapper for Unreal that can be used from no-rtti modules.
+	 */
+	class UNREALUSDWRAPPER_API FUsdSkelBlendShape : public FUsdTyped
+	{
+	public:
+		FUsdSkelBlendShape();
+		FUsdSkelBlendShape(const FUsdPrim& Prim);
+		FUsdSkelBlendShape(const FUsdSkelBlendShape& Other);
+		FUsdSkelBlendShape(FUsdSkelBlendShape&& Other);
+		~FUsdSkelBlendShape();
+
+		FUsdSkelBlendShape& operator=(const FUsdSkelBlendShape& Other);
+		FUsdSkelBlendShape& operator=(FUsdSkelBlendShape&& Other);
+
+		explicit operator bool() const;
+
+		// Auto conversion from/to pxr::UsdSkelBlendShape
+	public:
+#if USE_USD_SDK
+		explicit FUsdSkelBlendShape(const pxr::UsdSkelBlendShape& InUsdSkelBlendShape);
+		explicit FUsdSkelBlendShape(pxr::UsdSkelBlendShape&& InUsdSkelBlendShape);
+		explicit FUsdSkelBlendShape(const pxr::UsdPrim& Prim);
+
+		operator pxr::UsdSkelBlendShape&();
+		operator const pxr::UsdSkelBlendShape&() const;
+#endif
+
+		// Wrapped pxr::UsdSkelBlendShape functions, refer to the USD SDK documentation
+	public:
+		FUsdSkelInbetweenShape GetInbetween(const FString& Name) const;
+		TArray<FUsdSkelInbetweenShape> GetInbetweens() const;
+
+	private:
+		TUniquePtr<Internal::FUsdSkelBlendShapeImpl> Impl;
+	};
+}	 // namespace UE
diff --git a/Engine/Plugins/Runtime/USDCore/Source/UnrealUSDWrapper/Public/UsdWrappers/UsdSkelBlendShapeQuery.h b/Engine/Plugins/Runtime/USDCore/Source/UnrealUSDWrapper/Public/UsdWrappers/UsdSkelBlendShapeQuery.h
new file mode 100644
--- /dev/null
+++ b/Engine/Plugins/Runtime/USDCore/Source/UnrealUSDWrapper/Public/UsdWrappers/UsdSkelBlendShapeQuery.h
@@ -0,0 +1,71 @@
+// Copyright Epic Games, Inc. All Rights Reserved.
+
+#pragma once
+
+#include "CoreTypes.h"
+#include "Templates/UniquePtr.h"
+
+#if USE_USD_SDK
+
+#include "USDIncludesStart.h"
+#include "pxr/pxr.h"
+#include "USDIncludesEnd.h"
+
+PXR_NAMESPACE_OPEN_SCOPE
+	class UsdPrim;
+	class UsdSkelBlendShapeQuery;
+PXR_NAMESPACE_CLOSE_SCOPE
+
+#endif	  // #if USE_USD_SDK
+
+namespace UE
+{
+	class FUsdPrim;
+	class FUsdSkelBlendShape;
+
+	namespace Internal
+	{
+		class FUsdSkelBlendShapeQueryImpl;
+	}
+
+	/**
+	 * Minimal pxr::UsdSkelBlendShapeQuery wrapper for Unreal that can be used from no-rtti modules.
+	 */
+	class UNREALUSDWRAPPER_API FUsdSkelBlendShapeQuery
+	{
+	public:
+		FUsdSkelBlendShapeQuery();
+		FUsdSkelBlendShapeQuery(const UE::FUsdPrim& SkelBindingAPIPrim);
+		FUsdSkelBlendShapeQuery(const FUsdSkelBlendShapeQuery& Other);
+		FUsdSkelBlendShapeQuery(FUsdSkelBlendShapeQuery&& Other);
+		~FUsdSkelBlendShapeQuery();
+
+		FUsdSkelBlendShapeQuery& operator=(const FUsdSkelBlendShapeQuery& Other);
+		FUsdSkelBlendShapeQuery& operator=(FUsdSkelBlendShapeQuery&& Other);
+
+		explicit operator bool() const;
+
+		// Auto conversion from/to pxr::UsdSkelBlendShapeQuery
+	public:
+#if USE_USD_SDK
+		explicit FUsdSkelBlendShapeQuery(const pxr::UsdSkelBlendShapeQuery& InUsdSkelBlendShapeQuery);
+		explicit FUsdSkelBlendShapeQuery(pxr::UsdSkelBlendShapeQuery&& InUsdSkelBlendShapeQuery);
+
+		operator pxr::UsdSkelBlendShapeQuery&();
+		operator const pxr::UsdSkelBlendShapeQuery&() const;
+#endif
+
+		// Wrapped pxr::UsdSkelBlendShapeQuery functions, refer to the USD SDK documentation
+	public:
+		bool IsValid() const;
+
+		UE::FUsdPrim GetPrim() const;
+
+		size_t GetNumBlendShapes() const;
+
+		UE::FUsdSkelBlendShape GetBlendShape(size_t BlendShapeIndex) const;
+
+	private:
+		TUniquePtr<Internal::FUsdSkelBlendShapeQueryImpl> Impl;
+	};
+}	 // namespace UE
diff --git a/Engine/Plugins/Runtime/USDCore/Source/UnrealUSDWrapper/Public/UsdWrappers/UsdSkelCache.h b/Engine/Plugins/Runtime/USDCore/Source/UnrealUSDWrapper/Public/UsdWrappers/UsdSkelCache.h
new file mode 100644
--- /dev/null
+++ b/Engine/Plugins/Runtime/USDCore/Source/UnrealUSDWrapper/Public/UsdWrappers/UsdSkelCache.h
@@ -0,0 +1,82 @@
+// Copyright Epic Games, Inc. All Rights Reserved.
+
+#pragma once
+
+#include "CoreTypes.h"
+#include "Templates/UniquePtr.h"
+
+#if USE_USD_SDK
+
+#include "USDIncludesStart.h"
+#include "pxr/pxr.h"
+#include "USDIncludesEnd.h"
+
+PXR_NAMESPACE_OPEN_SCOPE
+	class UsdPrim;
+	class UsdSkelCache;
+PXR_NAMESPACE_CLOSE_SCOPE
+
+#endif	  // #if USE_USD_SDK
+
+namespace UE
+{
+	class FUsdPrim;
+	class FUsdSkelSkeletonQuery;
+	class FUsdSkelAnimQuery;
+	class FUsdSkelSkinningQuery;
+	class FUsdSkelBinding;
+
+	namespace Internal
+	{
+		class FUsdSkelCacheImpl;
+	}
+
+	/**
+	 * Minimal pxr::UsdSkelCache wrapper for Unreal that can be used from no-rtti modules.
+	 */
+	class UNREALUSDWRAPPER_API FUsdSkelCache
+	{
+	public:
+		FUsdSkelCache();
+		FUsdSkelCache(const FUsdSkelCache& Other);
+		FUsdSkelCache(FUsdSkelCache&& Other);
+		~FUsdSkelCache();
+
+		FUsdSkelCache& operator=(const FUsdSkelCache& Other);
+		FUsdSkelCache& operator=(FUsdSkelCache&& Other);
+
+		// Auto conversion from/to pxr::UsdSkelCache
+	public:
+#if USE_USD_SDK
+		explicit FUsdSkelCache(const pxr::UsdSkelCache& InUsdSkelCache);
+		explicit FUsdSkelCache(pxr::UsdSkelCache&& InUsdSkelCache);
+
+		operator pxr::UsdSkelCache&();
+		operator const pxr::UsdSkelCache&() const;
+#endif
+
+		// Wrapped pxr::UsdSkelCache functions, refer to the USD SDK documentation
+	public:
+		void Clear();
+
+		bool Populate(const UE::FUsdPrim& SkelRootPrim, bool bTraverseInstanceProxies);
+
+		UE::FUsdSkelSkeletonQuery GetSkelQuery(const UE::FUsdPrim& SkeletonPrim) const;
+
+		UE::FUsdSkelAnimQuery GetAnimQuery(const UE::FUsdPrim& SkelAnimationPrim) const;
+
+		UE::FUsdSkelSkinningQuery GetSkinningQuery(const UE::FUsdPrim& SkinnedPrim) const;
+
+		bool ComputeSkelBindings(const UE::FUsdPrim& InSkelRootPrim, TArray<UE::FUsdSkelBinding>& OutBindings, bool bTraverseInstanceProxies) const;
+
+		bool ComputeSkelBinding(
+			const UE::FUsdPrim& InSkelRootPrim,
+			const UE::FUsdPrim& InSkeletonPrim,
+			UE::FUsdSkelBinding& OutBinding,
+			bool bTraverseInstanceProxies
+		) const;
+
+	private:
+		TUniquePtr<Internal::FUsdSkelCacheImpl> Impl;
+	};
+}	 // namespace UE
diff --git a/Engine/Plugins/Runtime/USDCore/Source/UnrealUSDWrapper/Public/UsdWrappers/UsdSkelInbetweenShape.h b/Engine/Plugins/Runtime/USDCore/Source/UnrealUSDWrapper/Public/UsdWrappers/UsdSkelInbetweenShape.h
new file mode 100644
--- /dev/null
+++ b/Engine/Plugins/Runtime/USDCore/Source/UnrealUSDWrapper/Public/UsdWrappers/UsdSkelInbetweenShape.h
@@ -0,0 +1,66 @@
+// Copyright Epic Games, Inc. All Rights Reserved.
+
+#pragma once
+
+#include "CoreTypes.h"
+#include "Templates/UniquePtr.h"
+
+#if USE_USD_SDK
+
+#include "USDIncludesStart.h"
+#include "pxr/pxr.h"
+#include "USDIncludesEnd.h"
+
+PXR_NAMESPACE_OPEN_SCOPE
+	class UsdPrim;
+	class UsdSkelInbetweenShape;
+PXR_NAMESPACE_CLOSE_SCOPE
+
+#endif	  // #if USE_USD_SDK
+
+namespace UE
+{
+	class FUsdPrim;
+	class FUsdAttribute;
+
+	namespace Internal
+	{
+		class FUsdSkelInbetweenShapeImpl;
+	}
+
+	/**
+	 * Minimal pxr::UsdSkelInbetweenShape wrapper for Unreal that can be used from no-rtti modules.
+	 */
+	class UNREALUSDWRAPPER_API FUsdSkelInbetweenShape
+	{
+	public:
+		FUsdSkelInbetweenShape();
+		FUsdSkelInbetweenShape(const FUsdSkelInbetweenShape& Other);
+		FUsdSkelInbetweenShape(FUsdSkelInbetweenShape&& Other);
+		~FUsdSkelInbetweenShape();
+
+		FUsdSkelInbetweenShape& operator=(const FUsdSkelInbetweenShape& Other);
+		FUsdSkelInbetweenShape& operator=(FUsdSkelInbetweenShape&& Other);
+
+		explicit operator bool() const;
+
+		// Auto conversion from/to pxr::UsdSkelInbetweenShape
+	public:
+#if USE_USD_SDK
+		explicit FUsdSkelInbetweenShape(const pxr::UsdSkelInbetweenShape& InUsdSkelInbetweenShape);
+		explicit FUsdSkelInbetweenShape(pxr::UsdSkelInbetweenShape&& InUsdSkelInbetweenShape);
+
+		operator pxr::UsdSkelInbetweenShape&();
+		operator const pxr::UsdSkelInbetweenShape&() const;
+#endif
+
+		// Wrapped pxr::UsdSkelInbetweenShape functions, refer to the USD SDK documentation
+	public:
+		bool GetWeight(float* OutWeight) const;
+		UE::FUsdAttribute GetAttr() const;
+		bool IsDefined() const;
+
+	private:
+		TUniquePtr<Internal::FUsdSkelInbetweenShapeImpl> Impl;
+	};
+}	 // namespace UE
diff --git a/Engine/Plugins/Runtime/USDCore/Source/UnrealUSDWrapper/Public/UsdWrappers/UsdSkelSkeletonQuery.h b/Engine/Plugins/Runtime/USDCore/Source/UnrealUSDWrapper/Public/UsdWrappers/UsdSkelSkeletonQuery.h
new file mode 100644
--- /dev/null
+++ b/Engine/Plugins/Runtime/USDCore/Source/UnrealUSDWrapper/Public/UsdWrappers/UsdSkelSkeletonQuery.h
@@ -0,0 +1,71 @@
+// Copyright Epic Games, Inc. All Rights Reserved.
+
+#pragma once
+
+#include "CoreTypes.h"
+#include "Templates/UniquePtr.h"
+
+#if USE_USD_SDK
+
+#include "USDIncludesStart.h"
+#include "pxr/pxr.h"
+#include "USDIncludesEnd.h"
+
+PXR_NAMESPACE_OPEN_SCOPE
+	class UsdPrim;
+	class UsdSkelSkeletonQuery;
+PXR_NAMESPACE_CLOSE_SCOPE
+
+#endif	  // #if USE_USD_SDK
+
+namespace UE
+{
+	class FUsdPrim;
+	class FUsdSkelAnimQuery;
+
+	namespace Internal
+	{
+		class FUsdSkelSkeletonQueryImpl;
+	}
+
+	/**
+	 * Minimal pxr::UsdSkelSkeletonQuery wrapper for Unreal that can be used from no-rtti modules.
+	 */
+	class UNREALUSDWRAPPER_API FUsdSkelSkeletonQuery
+	{
+	public:
+		// WARNING: Default-constructed objects are permanently invalid!
+		// The default constructor is available for technical reasons, but you should create these objects
+		// by calling FUsdSkelCache::GetSkelQuery(SkeletonPrim)
+		FUsdSkelSkeletonQuery();
+		FUsdSkelSkeletonQuery(const FUsdSkelSkeletonQuery& Other);
+		FUsdSkelSkeletonQuery(FUsdSkelSkeletonQuery&& Other);
+		~FUsdSkelSkeletonQuery();
+
+		FUsdSkelSkeletonQuery& operator=(const FUsdSkelSkeletonQuery& Other);
+		FUsdSkelSkeletonQuery& operator=(FUsdSkelSkeletonQuery&& Other);
+
+		explicit operator bool() const;
+
+		// Auto conversion from/to pxr::UsdSkelSkeletonQuery
+	public:
+#if USE_USD_SDK
+		explicit FUsdSkelSkeletonQuery(const pxr::UsdSkelSkeletonQuery& InUsdSkelSkeletonQuery);
+		explicit FUsdSkelSkeletonQuery(pxr::UsdSkelSkeletonQuery&& InUsdSkelSkeletonQuery);
+
+		operator pxr::UsdSkelSkeletonQuery&();
+		operator const pxr::UsdSkelSkeletonQuery&() const;
+#endif
+
+		// Wrapped pxr::UsdSkelSkeletonQuery functions, refer to the USD SDK documentation
+	public:
+		bool IsValid() const;
+
+		UE::FUsdPrim GetPrim() const;
+		UE::FUsdPrim GetSkeleton() const;
+		UE::FUsdSkelAnimQuery GetAnimQuery() const;
+
+	private:
+		TUniquePtr<Internal::FUsdSkelSkeletonQueryImpl> Impl;
+	};
+}	 // namespace UE
diff --git a/Engine/Plugins/Runtime/USDCore/Source/UnrealUSDWrapper/Public/UsdWrappers/UsdSkelSkinningQuery.h b/Engine/Plugins/Runtime/USDCore/Source/UnrealUSDWrapper/Public/UsdWrappers/UsdSkelSkinningQuery.h
new file mode 100644
--- /dev/null
+++ b/Engine/Plugins/Runtime/USDCore/Source/UnrealUSDWrapper/Public/UsdWrappers/UsdSkelSkinningQuery.h
@@ -0,0 +1,65 @@
+// Copyright Epic Games, Inc. All Rights Reserved.
+
+#pragma once
+
+#include "CoreTypes.h"
+#include "Templates/UniquePtr.h"
+
+#if USE_USD_SDK
+
+#include "USDIncludesStart.h"
+#include "pxr/pxr.h"
+#include "USDIncludesEnd.h"
+
+PXR_NAMESPACE_OPEN_SCOPE
+	class UsdPrim;
+	class UsdSkelSkinningQuery;
+PXR_NAMESPACE_CLOSE_SCOPE
+
+#endif	  // #if USE_USD_SDK
+
+namespace UE
+{
+	class FUsdPrim;
+
+	namespace Internal
+	{
+		class FUsdSkelSkinningQueryImpl;
+	}
+
+	/**
+	 * Minimal pxr::UsdSkelSkinningQuery wrapper for Unreal that can be used from no-rtti modules.
+	 */
+	class UNREALUSDWRAPPER_API FUsdSkelSkinningQuery
+	{
+	public:
+		FUsdSkelSkinningQuery();
+		FUsdSkelSkinningQuery(const FUsdSkelSkinningQuery& Other);
+		FUsdSkelSkinningQuery(FUsdSkelSkinningQuery&& Other);
+		~FUsdSkelSkinningQuery();
+
+		FUsdSkelSkinningQuery& operator=(const FUsdSkelSkinningQuery& Other);
+		FUsdSkelSkinningQuery& operator=(FUsdSkelSkinningQuery&& Other);
+
+		explicit operator bool() const;
+
+		// Auto conversion from/to pxr::UsdSkelSkinningQuery
+	public:
+#if USE_USD_SDK
+		explicit FUsdSkelSkinningQuery(const pxr::UsdSkelSkinningQuery& InUsdSkelSkinningQuery);
+		explicit FUsdSkelSkinningQuery(pxr::UsdSkelSkinningQuery&& InUsdSkelSkinningQuery);
+
+		operator pxr::UsdSkelSkinningQuery&();
+		operator const pxr::UsdSkelSkinningQuery&() const;
+#endif
+
+		// Wrapped pxr::UsdSkelSkinningQuery functions, refer to the USD SDK documentation
+	public:
+		bool IsValid() const;
+
+		UE::FUsdPrim GetPrim() const;
+
+	private:
+		TUniquePtr<Internal::FUsdSkelSkinningQueryImpl> Impl;
+	};
+}	 // namespace UE
