Interchange: Handle Cine/CameraActors by putting transform and visibility on the root component.

diff --git a/Engine/Plugins/Interchange/Runtime/Source/Import/Private/Scene/InterchangeActorFactory.cpp b/Engine/Plugins/Interchange/Runtime/Source/Import/Private/Scene/InterchangeActorFactory.cpp
--- a/Engine/Plugins/Interchange/Runtime/Source/Import/Private/Scene/InterchangeActorFactory.cpp
+++ b/Engine/Plugins/Interchange/Runtime/Source/Import/Private/Scene/InterchangeActorFactory.cpp
@@ -2,6 +2,8 @@
 #include "Scene/InterchangeActorFactory.h"
 
 #include "InterchangeActorFactoryNode.h"
+#include "InterchangeCameraFactoryNode.h"
+#include "Nodes/InterchangeBaseNodeContainer.h"
 #include "Nodes/InterchangeFactoryBaseNode.h"
 #include "Scene/InterchangeActorHelper.h"
 
@@ -10,6 +12,56 @@
 
 #include UE_INLINE_GENERATED_CPP_BY_NAME(InterchangeActorFactory)
 
+namespace UE::InterchangeActorFactory::Private
+{
+	// For the camera actor types we get two components each: The root component is a default scene component, and the actual camera
+	// component is a child of the scene component. We want to place all scene component stuff (mostly transform) on the root
+	// component, and all the camera stuff on the camera component. This agrees with how the actor/root component is bound on
+	// LevelSequences, and is likely what users expect because when you place a camera actor on the level and move it around, you always
+	// affect the root component transform
+	void ApplyAllCameraCustomAttributes(
+		const UInterchangeFactoryBase::FImportSceneObjectsParams& CreateSceneObjectsParams,
+		UInterchangeActorFactoryNode* CameraFactoryNode,
+		USceneComponent* RootSceneComponent,
+		USceneComponent* ChildCameraComponent
+	)
+	{
+		using namespace UE::Interchange;
+
+		if (!RootSceneComponent || !ChildCameraComponent)
+		{
+			return;
+		}
+
+		UInterchangeBaseNodeContainer* const NodeContainer = const_cast<UInterchangeBaseNodeContainer* const>(CreateSceneObjectsParams.NodeContainer);
+
+		// Create a temp factory node so we don't modify our existing nodes with our changes
+		UInterchangePhysicalCameraFactoryNode* FactoryNodeCopy = NewObject<UInterchangePhysicalCameraFactoryNode>(NodeContainer, NAME_None);
+		FactoryNodeCopy->InitializeNode(
+			CameraFactoryNode->GetUniqueID(),
+			CameraFactoryNode->GetDisplayLabel(),
+			CameraFactoryNode->GetNodeContainerType()
+		);
+
+		NodeContainer->ReplaceNode(CameraFactoryNode->GetUniqueID(), FactoryNodeCopy);
+		{
+			UInterchangeFactoryBase::FImportSceneObjectsParams ParamsCopy{CreateSceneObjectsParams};
+			ParamsCopy.FactoryNode = FactoryNodeCopy;
+
+			// Apply exclusively camera stuff to the CineCameraComponent
+			FactoryNodeCopy->CopyWithObject(CameraFactoryNode, ChildCameraComponent);
+			FactoryNodeCopy->RemoveCustomAttributesForClass(USceneComponent::StaticClass());
+			ActorHelper::ApplyAllCustomAttributes(ParamsCopy, *ChildCameraComponent);
+
+			// Apply exclusively scene component stuff to the root SceneComponent
+			FactoryNodeCopy->CopyWithObject(CameraFactoryNode, RootSceneComponent);
+			FactoryNodeCopy->RemoveCustomAttributesForClass(UCineCameraComponent::StaticClass());
+			ActorHelper::ApplyAllCustomAttributes(ParamsCopy, *RootSceneComponent);
+		}
+		NodeContainer->ReplaceNode(FactoryNodeCopy->GetUniqueID(), CameraFactoryNode);
+	}
+}
+
 UClass* UInterchangeActorFactory::GetFactoryClass() const
 {
 	return AActor::StaticClass();
@@ -18,6 +70,7 @@ UClass* UInterchangeActorFactory::GetFactoryClass() const
 UObject* UInterchangeActorFactory::ImportSceneObject_GameThread(const UInterchangeFactoryBase::FImportSceneObjectsParams& CreateSceneObjectsParams)
 {
 	using namespace UE::Interchange;
+	using namespace UE::InterchangeActorFactory::Private;
 
 	UInterchangeActorFactoryNode* FactoryNode = Cast<UInterchangeActorFactoryNode>(CreateSceneObjectsParams.FactoryNode);
 	if (!ensure(FactoryNode) || !CreateSceneObjectsParams.NodeContainer)
@@ -37,7 +90,21 @@ UObject* UInterchangeActorFactory::ImportSceneObject_GameThread(const UInterchan
 				EComponentMobility::Type CachedMobility = RootComponent->Mobility;
 				RootComponent->SetMobility(EComponentMobility::Type::Movable);
 
-				ActorHelper::ApplyAllCustomAttributes(CreateSceneObjectsParams, *ObjectToUpdate);
+				// Apply factory node to object(s)
+				if (FactoryNode->IsA<UInterchangePhysicalCameraFactoryNode>())
+				{
+					UCineCameraComponent* CameraComponent = Cast<UCineCameraComponent>(ObjectToUpdate);
+					ApplyAllCameraCustomAttributes(CreateSceneObjectsParams, FactoryNode, RootComponent, CameraComponent);
+				}
+				else if (FactoryNode->IsA<UInterchangeStandardCameraFactoryNode>())
+				{
+					UCameraComponent* CameraComponent = Cast<UCameraComponent>(ObjectToUpdate);
+					ApplyAllCameraCustomAttributes(CreateSceneObjectsParams, FactoryNode, RootComponent, CameraComponent);
+				}
+				else
+				{
+					ActorHelper::ApplyAllCustomAttributes(CreateSceneObjectsParams, *ObjectToUpdate);
+				}
 
 				// Restore mobility value
 				if (CachedMobility != EComponentMobility::Type::Movable)
diff --git a/Engine/Source/Runtime/Interchange/Core/Private/Nodes/InterchangeFactoryBaseNode.cpp b/Engine/Source/Runtime/Interchange/Core/Private/Nodes/InterchangeFactoryBaseNode.cpp
--- a/Engine/Source/Runtime/Interchange/Core/Private/Nodes/InterchangeFactoryBaseNode.cpp
+++ b/Engine/Source/Runtime/Interchange/Core/Private/Nodes/InterchangeFactoryBaseNode.cpp
@@ -288,6 +288,12 @@ void UInterchangeFactoryBaseNode::FillAllCustomAttributeFromObject(UObject* Obje
 	}
 }
 
+void UInterchangeFactoryBaseNode::RemoveCustomAttributesForClass(UClass* Class)
+{
+	ApplyCustomAttributeDelegates.Remove(Class);
+	FillCustomAttributeDelegates.Remove(Class);
+}
+
 UInterchangeFactoryBaseNode* UInterchangeFactoryBaseNode::DuplicateWithObject(const UInterchangeFactoryBaseNode* SourceNode, UObject* Object)
 {
 	UInterchangeFactoryBaseNode* CurrentNode = NewObject<UInterchangeFactoryBaseNode>(GetTransientPackage(), SourceNode->GetClass());
diff --git a/Engine/Source/Runtime/Interchange/Core/Public/Nodes/InterchangeFactoryBaseNode.h b/Engine/Source/Runtime/Interchange/Core/Public/Nodes/InterchangeFactoryBaseNode.h
--- a/Engine/Source/Runtime/Interchange/Core/Public/Nodes/InterchangeFactoryBaseNode.h
+++ b/Engine/Source/Runtime/Interchange/Core/Public/Nodes/InterchangeFactoryBaseNode.h
@@ -316,6 +316,12 @@ public:
 
 	INTERCHANGECORE_API void FillAllCustomAttributeFromObject(UObject* Object) const;
 
+	/**
+	 * Removes all custom attribute delegates registered for the provided class.
+	 * @param Class				The class used to pick which delegates to remove.
+	 */
+	INTERCHANGECORE_API void RemoveCustomAttributesForClass(UClass* Class);
+
 	/**
 	 * Copies all the custom attributes from SourceNode to this node, and
 	 * gets the appropriate values from Object.
