USD: Translate transform and property animations via Interchange.

diff --git a/Engine/Plugins/Importers/USDImporter/Source/USDStage/Private/USDLevelSequenceHelper.cpp b/Engine/Plugins/Importers/USDImporter/Source/USDStage/Private/USDLevelSequenceHelper.cpp
--- a/Engine/Plugins/Importers/USDImporter/Source/USDStage/Private/USDLevelSequenceHelper.cpp
+++ b/Engine/Plugins/Importers/USDImporter/Source/USDStage/Private/USDLevelSequenceHelper.cpp
@@ -1405,7 +1405,7 @@ void FUsdLevelSequenceHelperImpl::CreateSubSequenceSection(ULevelSequence& Seque
 	{
 		if (UE::FUsdPrim SequencePrim = UsdStage.GetPrimAtPath(UE::FSdfPath(*PrimPathsForSequence[0])))
 		{
-			TArray<UE::FUsdAttribute> Attrs = UnrealToUsd::GetAttributesForProperty(SequencePrim, UnrealIdentifiers::TransformPropertyName);
+			TArray<UE::FUsdAttribute> Attrs = UsdUtils::GetAttributesForProperty(SequencePrim, UnrealIdentifiers::TransformPropertyName);
 			if (Attrs.Num() > 0)
 			{
 				SubLayerOffset = UsdUtils::GetLayerToStageOffset(Attrs[0]);
@@ -1717,7 +1717,7 @@ void FUsdLevelSequenceHelperImpl::AddCommonTracks(const UUsdPrimTwin& PrimTwin,
 
 		// Get all *animated* attributes that may contribute to the transform
 		bool bAreAllMuted = true;
-		TArray<UE::FUsdAttribute> Attrs = UnrealToUsd::GetAttributesForProperty(Prim, UnrealIdentifiers::TransformPropertyName);
+		TArray<UE::FUsdAttribute> Attrs = UsdUtils::GetAttributesForProperty(Prim, UnrealIdentifiers::TransformPropertyName);
 		for (int32 Index = Attrs.Num() - 1; Index >= 0; --Index)
 		{
 			const UE::FUsdAttribute& Attr = Attrs[Index];
@@ -1796,7 +1796,7 @@ void FUsdLevelSequenceHelperImpl::AddCommonTracks(const UUsdPrimTwin& PrimTwin,
 		}
 	}
 
-	TArray<UE::FUsdAttribute> Attrs = UnrealToUsd::GetAttributesForProperty(Prim, UnrealIdentifiers::HiddenInGamePropertyName);
+	TArray<UE::FUsdAttribute> Attrs = UsdUtils::GetAttributesForProperty(Prim, UnrealIdentifiers::HiddenInGamePropertyName);
 	if (Attrs.Num() > 0)
 	{
 		UE::FUsdAttribute VisibilityAttribute = Attrs[0];
@@ -1821,7 +1821,7 @@ void FUsdLevelSequenceHelperImpl::AddCommonTracks(const UUsdPrimTwin& PrimTwin,
 				{
 					if (UsdUtils::HasAnimatedVisibility(ParentPrim))
 					{
-						TArray<UE::FUsdAttribute> ParentAttrs = UnrealToUsd::GetAttributesForProperty(
+						TArray<UE::FUsdAttribute> ParentAttrs = UsdUtils::GetAttributesForProperty(
 							ParentPrim,
 							UnrealIdentifiers::HiddenInGamePropertyName
 						);
@@ -2067,7 +2067,7 @@ void FUsdLevelSequenceHelperImpl::AddCameraTracks(const UUsdPrimTwin& PrimTwin,
 
 	for (const FName& PropertyName : TrackedProperties)
 	{
-		TArray<UE::FUsdAttribute> Attrs = UnrealToUsd::GetAttributesForProperty(Prim, PropertyName);
+		TArray<UE::FUsdAttribute> Attrs = UsdUtils::GetAttributesForProperty(Prim, PropertyName);
 		if (Attrs.Num() < 1)
 		{
 			continue;
@@ -2191,7 +2191,7 @@ void FUsdLevelSequenceHelperImpl::AddLightTracks(const UUsdPrimTwin& PrimTwin, c
 		const FName& PropertyPath = Pair.Key;
 		ETrackType TrackType = Pair.Value;
 
-		TArray<UE::FUsdAttribute> Attrs = UnrealToUsd::GetAttributesForProperty(Prim, PropertyPath);
+		TArray<UE::FUsdAttribute> Attrs = UsdUtils::GetAttributesForProperty(Prim, PropertyPath);
 		if (Attrs.Num() < 1)
 		{
 			continue;
@@ -2537,7 +2537,7 @@ void FUsdLevelSequenceHelperImpl::AddVolumeTracks(const UUsdPrimTwin& PrimTwin,
 	// Here we'll just get *any* of the filePath attrs from this Volume prim to check for the
 	// muted bool. We won't use the attribute itself for the baking though, as our timeSamples are already
 	// on the SparseVolumeTexture AssetUserData, and the keyframe values are just their indices
-	TArray<UE::FUsdAttribute> Attrs = UnrealToUsd::GetAttributesForProperty(Prim, PropertyPath);
+	TArray<UE::FUsdAttribute> Attrs = UsdUtils::GetAttributesForProperty(Prim, PropertyPath);
 	if (Attrs.Num() < 1)
 	{
 		return;
@@ -4004,7 +4004,7 @@ void FUsdLevelSequenceHelperImpl::HandleMovieSceneChange(UMovieScene& MovieScene
 	{
 		if (!UsdLevelSequenceHelperImpl::FindTrackTypeOrDerived<UMovieScenePropertyTrack>(&MovieScene, Guid, PropertyPath))
 		{
-			for (UE::FUsdAttribute& Attr : UnrealToUsd::GetAttributesForProperty(Prim, PropertyPath))
+			for (UE::FUsdAttribute& Attr : UsdUtils::GetAttributesForProperty(Prim, PropertyPath))
 			{
 				RemoveTimeSamplesForAttr(Attr);
 			}
@@ -4421,7 +4421,7 @@ void FUsdLevelSequenceHelperImpl::HandleTrackChange(const UMovieSceneTrack& Trac
 			{
 				const FName& PropertyPath = PropertyTrack->GetPropertyPath();
 
-				TArray<UE::FUsdAttribute> Attrs = UnrealToUsd::GetAttributesForProperty(UsdPrim, PropertyPath);
+				TArray<UE::FUsdAttribute> Attrs = UsdUtils::GetAttributesForProperty(UsdPrim, PropertyPath);
 				if (Attrs.Num() > 0)
 				{
 					// Only mute/unmute the first (i.e. main) attribute: If we mute the intensity track we don't want to also mute the
diff --git a/Engine/Plugins/Interchange/Runtime/Source/Import/Private/Usd/InterchangeUsdTranslator.cpp b/Engine/Plugins/Interchange/Runtime/Source/Import/Private/Usd/InterchangeUsdTranslator.cpp
--- a/Engine/Plugins/Interchange/Runtime/Source/Import/Private/Usd/InterchangeUsdTranslator.cpp
+++ b/Engine/Plugins/Interchange/Runtime/Source/Import/Private/Usd/InterchangeUsdTranslator.cpp
@@ -40,6 +40,7 @@
 #include "InterchangeTexture2DNode.h"
 #include "InterchangeTranslatorHelper.h"
 #include "Internationalization/Regex.h"
+#include "MovieSceneSection.h"
 #include "Rendering/SkeletalMeshLODImporterData.h"
 #include "StaticMeshAttributes.h"
 #include "StaticMeshOperations.h"
@@ -74,6 +75,8 @@ static FAutoConsoleVariableRef CVarInterchangeEnableUSDLevelImport(
 
 namespace UE::InterchangeUsdTranslator::Private
 {
+	const static FString AnimationPrefix = TEXT("\\Animation\\");
+	const static FString AnimationTrackPrefix = TEXT("\\AnimationTrack\\");
 	const static FString CameraPrefix = TEXT("\\Camera\\");
 	const static FString LightPrefix = TEXT("\\Light\\");
 	const static FString MaterialPrefix = TEXT("\\Material\\");
@@ -92,6 +95,32 @@ namespace UE::InterchangeUsdTranslator::Private
 		TSharedPtr<TArray<FString>> SkelJointNames;	   // Needed for skel mesh payloads
 	};
 
+	// clang-format off
+	const static TMap<FName, EInterchangePropertyTracks> PropertyNameToTrackType = {
+		// Common properties
+		{UnrealIdentifiers::HiddenInGamePropertyName, 			EInterchangePropertyTracks::Visibility}, // Binding visibility to the actor works better for cameras
+
+		// Camera properties
+		{UnrealIdentifiers::CurrentFocalLengthPropertyName, 	EInterchangePropertyTracks::CameraCurrentFocalLength},
+		{UnrealIdentifiers::ManualFocusDistancePropertyName, 	EInterchangePropertyTracks::CameraFocusSettingsManualFocusDistance},
+		{UnrealIdentifiers::CurrentAperturePropertyName, 		EInterchangePropertyTracks::CameraCurrentAperture},
+		{UnrealIdentifiers::SensorWidthPropertyName, 			EInterchangePropertyTracks::CameraFilmbackSensorWidth},
+		{UnrealIdentifiers::SensorHeightPropertyName, 			EInterchangePropertyTracks::CameraFilmbackSensorHeight},
+
+		// Light properties
+		{UnrealIdentifiers::LightColorPropertyName, 			EInterchangePropertyTracks::LightColor},
+		{UnrealIdentifiers::TemperaturePropertyName, 			EInterchangePropertyTracks::LightTemperature},
+		{UnrealIdentifiers::UseTemperaturePropertyName, 		EInterchangePropertyTracks::LightUseTemperature},
+		{UnrealIdentifiers::SourceHeightPropertyName, 			EInterchangePropertyTracks::LightSourceHeight},
+		{UnrealIdentifiers::SourceWidthPropertyName, 			EInterchangePropertyTracks::LightSourceWidth},
+		{UnrealIdentifiers::SourceRadiusPropertyName, 			EInterchangePropertyTracks::LightSourceRadius},
+		{UnrealIdentifiers::OuterConeAnglePropertyName, 		EInterchangePropertyTracks::LightOuterConeAngle},
+		{UnrealIdentifiers::InnerConeAnglePropertyName, 		EInterchangePropertyTracks::LightInnerConeAngle},
+		{UnrealIdentifiers::LightSourceAnglePropertyName, 		EInterchangePropertyTracks::LightSourceAngle},
+		{UnrealIdentifiers::IntensityPropertyName, 				EInterchangePropertyTracks::LightIntensity},
+	};
+	// clang-format on
+
 	// Small container that we can use Pimpl with so we don't have to include too many USD includes on the header file.
 	//
 	// It also skirts around a small complication where UInterchangeUSDTranslator::Translate is const, and yet we must
@@ -118,6 +147,11 @@ namespace UE::InterchangeUsdTranslator::Private
 		// Note: We only do this when needed: This shouldn't have data for every prim in the stage.
 		TMap<FString, FTraversalInfo> NodeUidToCachedTraversalInfo;
 		mutable FRWLock CachedTraversalInfoLock;
+
+		// This node eventually becomes a LevelSequence, and all track nodes are connected to it.
+		// For now we only generate a single LevelSequence per stage though, so we'll keep track of this
+		// here for easy access when parsing the tracks
+		UInterchangeAnimationTrackSetNode* CurrentTrackSet = nullptr;
 	};
 }	 // namespace UE::InterchangeUsdTranslator::Private
 
@@ -125,6 +159,7 @@ UInterchangeUsdTranslatorSettings::UInterchangeUsdTranslatorSettings()
 	: GeometryPurpose((int32)(EUsdPurpose::Default | EUsdPurpose::Proxy | EUsdPurpose::Render | EUsdPurpose::Guide))
 	, RenderContext(TEXT("unreal"))	   // The proper definition of this is on the USDSchemas module, which we can't depend on
 	, MaterialPurpose(*UnrealIdentifiers::MaterialPreviewPurpose)
+	, InterpolationType(EUsdInterpolationType::Linear)
 	, bOverrideStageOptions(false)
 	, StageOptions{
 		  0.01f,			   // MetersPerUnit
@@ -369,7 +404,7 @@ namespace UE::InterchangeUsdTranslator::Private
 
 	void AddMaterialInstanceNode(
 		const UE::FUsdPrim& Prim,
-		UInterchangeUSDTranslatorImpl* TranslatorImpl,
+		UInterchangeUSDTranslatorImpl& TranslatorImpl,
 		UInterchangeBaseNodeContainer& NodeContainer
 	)
 	{
@@ -389,8 +424,8 @@ namespace UE::InterchangeUsdTranslator::Private
 		NodeContainer.AddNode(MaterialNode);
 
 		UsdToUnreal::FUsdPreviewSurfaceMaterialData MaterialData;
-		FString RenderContext = TranslatorImpl->TranslatorSettings ? TranslatorImpl->TranslatorSettings->RenderContext.ToString() : FString();
-		const bool bSuccess = UsdToUnreal::ConvertMaterial(Prim, MaterialData, TranslatorImpl->TranslatorSettings ? *RenderContext : nullptr);
+		FString RenderContext = TranslatorImpl.TranslatorSettings ? TranslatorImpl.TranslatorSettings->RenderContext.ToString() : FString();
+		const bool bSuccess = UsdToUnreal::ConvertMaterial(Prim, MaterialData, TranslatorImpl.TranslatorSettings ? *RenderContext : nullptr);
 
 		// Set all the parameter values to the interchange node
 		bool bHasUDIMTexture = false;
@@ -642,7 +677,7 @@ namespace UE::InterchangeUsdTranslator::Private
 
 	void AddMorphTargetNodes(
 		const UE::FUsdPrim& Prim,
-		UInterchangeUSDTranslatorImpl* TranslatorImpl,
+		UInterchangeUSDTranslatorImpl& TranslatorImpl,
 		UInterchangeMeshNode& MeshNode,
 		UInterchangeBaseNodeContainer& NodeContainer
 	)
@@ -704,7 +739,7 @@ namespace UE::InterchangeUsdTranslator::Private
 
 	void AddMeshNode(
 		const UE::FUsdPrim& Prim,
-		UInterchangeUSDTranslatorImpl* TranslatorImpl,
+		UInterchangeUSDTranslatorImpl& TranslatorImpl,
 		UInterchangeBaseNodeContainer& NodeContainer,
 		const FTraversalInfo& Info
 	)
@@ -738,8 +773,8 @@ namespace UE::InterchangeUsdTranslator::Private
 			// When returning the payload data later, we'll need at the very least our SkeletonQuery, so
 			// here we store the Info object into the Impl
 			{
-				FWriteScopeLock ScopedInfoWriteLock{TranslatorImpl->CachedTraversalInfoLock};
-				TranslatorImpl->NodeUidToCachedTraversalInfo.Add(NodeUid, Info);
+				FWriteScopeLock ScopedInfoWriteLock{TranslatorImpl.CachedTraversalInfoLock};
+				TranslatorImpl.NodeUidToCachedTraversalInfo.Add(NodeUid, Info);
 			}
 		}
 		else
@@ -755,8 +790,8 @@ namespace UE::InterchangeUsdTranslator::Private
 				Prim,
 				TimeCode,
 				bProvideMaterialIndices,
-				TranslatorImpl->CachedMeshConversionOptions.RenderContext,
-				TranslatorImpl->CachedMeshConversionOptions.MaterialPurpose
+				TranslatorImpl.CachedMeshConversionOptions.RenderContext,
+				TranslatorImpl.CachedMeshConversionOptions.MaterialPurpose
 			);
 
 			for (const UsdUtils::FUsdPrimMaterialSlot& Slot : Assignments.Slots)
@@ -806,7 +841,7 @@ namespace UE::InterchangeUsdTranslator::Private
 
 	void AddSkeletonNodes(
 		const UE::FUsdPrim& Prim,
-		UInterchangeUSDTranslatorImpl* TranslatorImpl,
+		UInterchangeUSDTranslatorImpl& TranslatorImpl,
 		UInterchangeSceneNode& SkeletonPrimNode,
 		UInterchangeBaseNodeContainer& NodeContainer,
 		FTraversalInfo Info
@@ -888,8 +923,8 @@ namespace UE::InterchangeUsdTranslator::Private
 			Info.SkelJointNames->Add(Bone.Name);
 		}
 		{
-			FWriteScopeLock ScopedInfoWriteLock{TranslatorImpl->CachedTraversalInfoLock};
-			TranslatorImpl->NodeUidToCachedTraversalInfo.Add(SkeletonPrimNodeUid, Info);
+			FWriteScopeLock ScopedInfoWriteLock{TranslatorImpl.CachedTraversalInfoLock};
+			TranslatorImpl.NodeUidToCachedTraversalInfo.Add(SkeletonPrimNodeUid, Info);
 		}
 	}
 
@@ -928,9 +963,126 @@ namespace UE::InterchangeUsdTranslator::Private
 		}
 	}
 
+	void AddTrackSetNode(UInterchangeUSDTranslatorImpl& Impl, UInterchangeBaseNodeContainer& NodeContainer)
+	{
+		// For now we only want a single track set (i.e. LevelSequence) per stage.
+		// TODO: One track set per layer, and add the tracks to the tracksets that correspond to layers where the opinions came from
+		// (similar to LevelSequenceHelper). Then we can use UInterchangeAnimationTrackSetInstanceNode to create "subsequences"
+		if (Impl.CurrentTrackSet)
+		{
+			return;
+		}
+
+		UE::FSdfLayer Layer = Impl.UsdStage.GetRootLayer();
+		const FString AnimTrackSetNodeUid = AnimationPrefix + Layer.GetIdentifier();
+		const FString AnimTrackSetNodeDisplayName = FPaths::GetBaseFilename(Layer.GetDisplayName());	// Strip extension
+
+		// We should only have one track set node per scene for now
+		const UInterchangeAnimationTrackSetNode* ExistingNode = Cast<UInterchangeAnimationTrackSetNode>(NodeContainer.GetNode(AnimTrackSetNodeUid));
+		if (!ensure(ExistingNode == nullptr))
+		{
+			return;
+		};
+
+		UInterchangeAnimationTrackSetNode* TrackSetNode = NewObject<UInterchangeAnimationTrackSetNode>(&NodeContainer);
+		TrackSetNode->InitializeNode(AnimTrackSetNodeUid, AnimTrackSetNodeDisplayName, EInterchangeNodeContainerType::TranslatedAsset);
+		TrackSetNode->SetCustomFrameRate(Layer.GetFramesPerSecond());	 // Key values in Interchange seem to be in seconds, so timeCodesPerSecond is
+																		 // not relevant here
+
+		NodeContainer.AddNode(TrackSetNode);
+		Impl.CurrentTrackSet = TrackSetNode;
+	}
+
+	void AddTransformAnimationNode(const UE::FUsdPrim& Prim, UInterchangeUSDTranslatorImpl& Impl, UInterchangeBaseNodeContainer& NodeContainer)
+	{
+		const FString PrimPath = Prim.GetPrimPath().GetString();
+		const FString UniquePath = PrimPath + TEXT("\\") + UnrealIdentifiers::TransformPropertyName.ToString();
+		const FString AnimTrackNodeUid = AnimationTrackPrefix + UniquePath;
+
+		const UInterchangeTransformAnimationTrackNode* ExistingNode = Cast<UInterchangeTransformAnimationTrackNode>(
+			NodeContainer.GetNode(AnimTrackNodeUid)
+		);
+		if (ExistingNode)
+		{
+			return;
+		}
+
+		UInterchangeTransformAnimationTrackNode* TransformAnimTrackNode = NewObject<UInterchangeTransformAnimationTrackNode>(&NodeContainer);
+		TransformAnimTrackNode->InitializeNode(AnimTrackNodeUid, UniquePath, EInterchangeNodeContainerType::TranslatedAsset);
+		TransformAnimTrackNode->SetCustomActorDependencyUid(*PrimPath);
+		TransformAnimTrackNode->SetCustomAnimationPayloadKey(UniquePath, EInterchangeAnimationPayLoadType::CURVE);
+		TransformAnimTrackNode->SetCustomUsedChannels((int32)EMovieSceneTransformChannel::AllTransform);
+
+		NodeContainer.AddNode(TransformAnimTrackNode);
+
+		AddTrackSetNode(Impl, NodeContainer);
+		Impl.CurrentTrackSet->AddCustomAnimationTrackUid(AnimTrackNodeUid);
+	}
+
+	void AddPropertyAnimationNodes(const UE::FUsdPrim& Prim, UInterchangeUSDTranslatorImpl& Impl, UInterchangeBaseNodeContainer& NodeContainer)
+	{
+		using namespace UE::InterchangeUsdTranslator::Private;
+
+		if (!Prim)
+		{
+			return;
+		}
+		const FString PrimPath = Prim.GetPrimPath().GetString();
+
+		for (UE::FUsdAttribute Attr : Prim.GetAttributes())
+		{
+			if (!Attr || !Attr.ValueMightBeTimeVarying() || Attr.GetNumTimeSamples() == 0)
+			{
+				continue;
+			}
+
+			// Emit a STEPCURVE in case of a bool track: CURVE is only for floats/doubles (c.f. FLevelSequenceHelper::PopulateAnimationTrack).
+			// For now we're lucky in that all possible results from GetPropertiesForAttribute() are either all not bool, or either all bool,
+			// so we can reuse this for all the different UEAttrNames we get from the same attribute
+			const FName AttrTypeName = Attr.GetTypeName();
+			const bool bIsBoolTrack = AttrTypeName == TEXT("bool") || AttrTypeName == TEXT("token");	// Visibility is a token track
+
+			TArray<FName> UEAttrNames = UsdUtils::GetPropertiesForAttribute(Prim, Attr.GetName().ToString());
+			for (const FName& UEAttrName : UEAttrNames)
+			{
+				const EInterchangePropertyTracks* FoundTrackType = PropertyNameToTrackType.Find(UEAttrName);
+				if (!FoundTrackType)
+				{
+					continue;
+				}
+
+				// We don't use the USD attribute path here because we want one unique node per UE track name,
+				// so that if e.g. both "intensity" and "exposure" are animated we make a single track for
+				// the Intensity UE property
+				const FString UniquePath = PrimPath + TEXT("\\") + UEAttrName.ToString();
+				const FString AnimTrackNodeUid = AnimationTrackPrefix + UniquePath;
+
+				const UInterchangeAnimationTrackNode* ExistingNode = Cast<UInterchangeAnimationTrackNode>(NodeContainer.GetNode(AnimTrackNodeUid));
+				if (ExistingNode)
+				{
+					continue;
+				}
+
+				UInterchangeAnimationTrackNode* AnimTrackNode = NewObject<UInterchangeAnimationTrackNode>(&NodeContainer);
+				AnimTrackNode->InitializeNode(AnimTrackNodeUid, UniquePath, EInterchangeNodeContainerType::TranslatedAsset);
+				AnimTrackNode->SetCustomActorDependencyUid(*PrimPath);
+				AnimTrackNode->SetCustomPropertyTrack(*FoundTrackType);
+				AnimTrackNode->SetCustomAnimationPayloadKey(
+					UniquePath,
+					bIsBoolTrack ? EInterchangeAnimationPayLoadType::STEPCURVE : EInterchangeAnimationPayLoadType::CURVE
+				);
+
+				NodeContainer.AddNode(AnimTrackNode);
+
+				AddTrackSetNode(Impl, NodeContainer);
+				Impl.CurrentTrackSet->AddCustomAnimationTrackUid(AnimTrackNodeUid);
+			}
+		}
+	}
+
 	void Traverse(
 		const UE::FUsdPrim& Prim,
-		UInterchangeUSDTranslatorImpl* TranslatorImpl,
+		UInterchangeUSDTranslatorImpl& TranslatorImpl,
 		UInterchangeBaseNodeContainer& NodeContainer,
 		FTraversalInfo Info
 	)
@@ -938,7 +1090,7 @@ namespace UE::InterchangeUsdTranslator::Private
 		// Ignore prim subtrees from disabled purposes
 		// TODO: Move this to the pipeline and filter only the factory nodes
 		EUsdPurpose PrimPurpose = IUsdPrim::GetPurpose(Prim);
-		if (!EnumHasAllFlags(TranslatorImpl->CachedMeshConversionOptions.PurposesToLoad, PrimPurpose))
+		if (!EnumHasAllFlags(TranslatorImpl.CachedMeshConversionOptions.PurposesToLoad, PrimPurpose))
 		{
 			return;
 		}
@@ -1013,6 +1165,13 @@ namespace UE::InterchangeUsdTranslator::Private
 			{
 				NodeContainer.SetNodeParentUid(SceneNode->GetUniqueID(), Info.ParentNode->GetUniqueID());
 			}
+
+			// Add animation tracks
+			AddPropertyAnimationNodes(Prim, TranslatorImpl, NodeContainer);
+			if (UsdUtils::HasAnimatedTransform(Prim))
+			{
+				AddTransformAnimationNode(Prim, TranslatorImpl, NodeContainer);
+			}
 		}
 
 		// Note: This has the effect of effectively shutting down the generation of scene nodes
@@ -1241,6 +1400,254 @@ namespace UE::InterchangeUsdTranslator::Private
 		const float Weight = 1.0f;
 		return UsdUtils::ApplyBlendShape(OutMeshDescription, BlendShape.GetPrim(), Weight, InbetweenName);
 	}
+
+	bool ReadBools(
+		const UE::FUsdStage& UsdStage,
+		const TArray<double>& UsdTimeSamples,
+		const TFunction<bool(double)>& ReaderFunc,
+		Interchange::FAnimationPayloadData& OutPayloadData
+	)
+	{
+		OutPayloadData.StepCurves.SetNum(1);
+		FInterchangeStepCurve& Curve = OutPayloadData.StepCurves[0];
+		TArray<float>& KeyTimes = Curve.KeyTimes;
+		TArray<bool>& BooleanKeyValues = Curve.BooleanKeyValues.Emplace();
+
+		KeyTimes.Reserve(UsdTimeSamples.Num());
+		BooleanKeyValues.Reserve(UsdTimeSamples.Num());
+
+		const FFrameRate StageFrameRate{static_cast<uint32>(UsdStage.GetTimeCodesPerSecond()), 1};
+
+		double LastTimeSample = TNumericLimits<double>::Lowest();
+		for (const double UsdTimeSample : UsdTimeSamples)
+		{
+			// We never want to evaluate the same time twice
+			if (FMath::IsNearlyEqual(UsdTimeSample, LastTimeSample))
+			{
+				continue;
+			}
+			LastTimeSample = UsdTimeSample;
+
+			int32 FrameNumber = FMath::FloorToInt(UsdTimeSample);
+			float SubFrameNumber = UsdTimeSample - FrameNumber;
+
+			FFrameTime FrameTime{FrameNumber, SubFrameNumber};
+			double FrameTimeSeconds = static_cast<float>(StageFrameRate.AsSeconds(FrameTime));
+
+			bool UEValue = ReaderFunc(UsdTimeSample);
+
+			KeyTimes.Add(FrameTimeSeconds);
+			BooleanKeyValues.Add(UEValue);
+		}
+
+		return true;
+	}
+
+	bool ReadFloats(
+		const UE::FUsdStage& UsdStage,
+		const TArray<double>& UsdTimeSamples,
+		const TFunction<float(double)>& ReaderFunc,
+		Interchange::FAnimationPayloadData& OutPayloadData
+	)
+	{
+		OutPayloadData.Curves.SetNum(1);
+		FRichCurve& Curve = OutPayloadData.Curves[0];
+
+		const FFrameRate StageFrameRate{static_cast<uint32>(UsdStage.GetTimeCodesPerSecond()), 1};
+		const ERichCurveInterpMode InterpMode = (UsdStage.GetInterpolationType() == EUsdInterpolationType::Linear)
+													? ERichCurveInterpMode::RCIM_Linear
+													: ERichCurveInterpMode::RCIM_Constant;
+
+		double LastTimeSample = TNumericLimits<double>::Lowest();
+		for (const double UsdTimeSample : UsdTimeSamples)
+		{
+			// We never want to evaluate the same time twice
+			if (FMath::IsNearlyEqual(UsdTimeSample, LastTimeSample))
+			{
+				continue;
+			}
+			LastTimeSample = UsdTimeSample;
+
+			int32 FrameNumber = FMath::FloorToInt(UsdTimeSample);
+			float SubFrameNumber = UsdTimeSample - FrameNumber;
+
+			FFrameTime FrameTime{FrameNumber, SubFrameNumber};
+			double FrameTimeSeconds = static_cast<float>(StageFrameRate.AsSeconds(FrameTime));
+
+			float UEValue = ReaderFunc(UsdTimeSample);
+
+			FKeyHandle Handle = Curve.AddKey(FrameTimeSeconds, UEValue);
+			Curve.SetKeyInterpMode(Handle, InterpMode);
+		}
+
+		return true;
+	}
+
+	bool ReadColors(
+		const UE::FUsdStage& UsdStage,
+		const TArray<double>& UsdTimeSamples,
+		const TFunction<FLinearColor(double)>& ReaderFunc,
+		Interchange::FAnimationPayloadData& OutPayloadData
+	)
+	{
+		OutPayloadData.Curves.SetNum(4);
+		FRichCurve& RCurve = OutPayloadData.Curves[0];
+		FRichCurve& GCurve = OutPayloadData.Curves[1];
+		FRichCurve& BCurve = OutPayloadData.Curves[2];
+		FRichCurve& ACurve = OutPayloadData.Curves[3];
+
+		const FFrameRate StageFrameRate{static_cast<uint32>(UsdStage.GetTimeCodesPerSecond()), 1};
+		const ERichCurveInterpMode InterpMode = (UsdStage.GetInterpolationType() == EUsdInterpolationType::Linear)
+													? ERichCurveInterpMode::RCIM_Linear
+													: ERichCurveInterpMode::RCIM_Constant;
+
+		double LastTimeSample = TNumericLimits<double>::Lowest();
+		for (const double UsdTimeSample : UsdTimeSamples)
+		{
+			// We never want to evaluate the same time twice
+			if (FMath::IsNearlyEqual(UsdTimeSample, LastTimeSample))
+			{
+				continue;
+			}
+			LastTimeSample = UsdTimeSample;
+
+			int32 FrameNumber = FMath::FloorToInt(UsdTimeSample);
+			float SubFrameNumber = UsdTimeSample - FrameNumber;
+
+			FFrameTime FrameTime{FrameNumber, SubFrameNumber};
+			double FrameTimeSeconds = static_cast<float>(StageFrameRate.AsSeconds(FrameTime));
+
+			FLinearColor UEValue = ReaderFunc(UsdTimeSample);
+
+			FKeyHandle RHandle = RCurve.AddKey(FrameTimeSeconds, UEValue.R);
+			FKeyHandle GHandle = GCurve.AddKey(FrameTimeSeconds, UEValue.G);
+			FKeyHandle BHandle = BCurve.AddKey(FrameTimeSeconds, UEValue.B);
+			FKeyHandle AHandle = ACurve.AddKey(FrameTimeSeconds, UEValue.A);
+
+			RCurve.SetKeyInterpMode(RHandle, InterpMode);
+			GCurve.SetKeyInterpMode(GHandle, InterpMode);
+			BCurve.SetKeyInterpMode(BHandle, InterpMode);
+			ACurve.SetKeyInterpMode(AHandle, InterpMode);
+		}
+
+		return true;
+	}
+
+	bool ReadTransforms(
+		const UE::FUsdStage& UsdStage,
+		const TArray<double>& UsdTimeSamples,
+		const TFunction<FTransform(double)>& ReaderFunc,
+		Interchange::FAnimationPayloadData& OutPayloadData
+	)
+	{
+		OutPayloadData.Curves.SetNum(9);
+		FRichCurve& TransXCurve = OutPayloadData.Curves[0];
+		FRichCurve& TransYCurve = OutPayloadData.Curves[1];
+		FRichCurve& TransZCurve = OutPayloadData.Curves[2];
+		FRichCurve& RotXCurve = OutPayloadData.Curves[3];
+		FRichCurve& RotYCurve = OutPayloadData.Curves[4];
+		FRichCurve& RotZCurve = OutPayloadData.Curves[5];
+		FRichCurve& ScaleXCurve = OutPayloadData.Curves[6];
+		FRichCurve& ScaleYCurve = OutPayloadData.Curves[7];
+		FRichCurve& ScaleZCurve = OutPayloadData.Curves[8];
+
+		const FFrameRate StageFrameRate{static_cast<uint32>(UsdStage.GetTimeCodesPerSecond()), 1};
+		const ERichCurveInterpMode InterpMode = (UsdStage.GetInterpolationType() == EUsdInterpolationType::Linear)
+													? ERichCurveInterpMode::RCIM_Linear
+													: ERichCurveInterpMode::RCIM_Constant;
+
+		double LastTimeSample = TNumericLimits<double>::Lowest();
+		for (const double UsdTimeSample : UsdTimeSamples)
+		{
+			// We never want to evaluate the same time twice
+			if (FMath::IsNearlyEqual(UsdTimeSample, LastTimeSample))
+			{
+				continue;
+			}
+			LastTimeSample = UsdTimeSample;
+
+			int32 FrameNumber = FMath::FloorToInt(UsdTimeSample);
+			float SubFrameNumber = UsdTimeSample - FrameNumber;
+
+			FFrameTime FrameTime{FrameNumber, SubFrameNumber};
+			double FrameTimeSeconds = static_cast<float>(StageFrameRate.AsSeconds(FrameTime));
+
+			FTransform UEValue = ReaderFunc(UsdTimeSample);
+			FVector Location = UEValue.GetLocation();
+			FRotator Rotator = UEValue.Rotator();
+			FVector Scale = UEValue.GetScale3D();
+
+			FKeyHandle HandleTransX = TransXCurve.AddKey(FrameTimeSeconds, Location.X);
+			FKeyHandle HandleTransY = TransYCurve.AddKey(FrameTimeSeconds, Location.Y);
+			FKeyHandle HandleTransZ = TransZCurve.AddKey(FrameTimeSeconds, Location.Z);
+			FKeyHandle HandleRotX = RotXCurve.AddKey(FrameTimeSeconds, Rotator.Roll);
+			FKeyHandle HandleRotY = RotYCurve.AddKey(FrameTimeSeconds, Rotator.Pitch);
+			FKeyHandle HandleRotZ = RotZCurve.AddKey(FrameTimeSeconds, Rotator.Yaw);
+			FKeyHandle HandleScaleX = ScaleXCurve.AddKey(FrameTimeSeconds, Scale.X);
+			FKeyHandle HandleScaleY = ScaleYCurve.AddKey(FrameTimeSeconds, Scale.Y);
+			FKeyHandle HandleScaleZ = ScaleZCurve.AddKey(FrameTimeSeconds, Scale.Z);
+
+			TransXCurve.SetKeyInterpMode(HandleTransX, InterpMode);
+			TransYCurve.SetKeyInterpMode(HandleTransY, InterpMode);
+			TransZCurve.SetKeyInterpMode(HandleTransZ, InterpMode);
+			RotXCurve.SetKeyInterpMode(HandleRotX, InterpMode);
+			RotYCurve.SetKeyInterpMode(HandleRotY, InterpMode);
+			RotZCurve.SetKeyInterpMode(HandleRotZ, InterpMode);
+			ScaleXCurve.SetKeyInterpMode(HandleScaleX, InterpMode);
+			ScaleYCurve.SetKeyInterpMode(HandleScaleY, InterpMode);
+			ScaleZCurve.SetKeyInterpMode(HandleScaleZ, InterpMode);
+		}
+
+		return true;
+	}
+
+	bool GetAnimationCurvePayloadData(const UE::FUsdStage& UsdStage, const FString& PayloadKey, Interchange::FAnimationPayloadData& OutPayloadData)
+	{
+		FString PrimPath;
+		FString UEPropertyNameStr;
+		bool bSplit = PayloadKey.Split(TEXT("\\"), &PrimPath, &UEPropertyNameStr, ESearchCase::CaseSensitive, ESearchDir::FromEnd);
+		if (!bSplit)
+		{
+			return false;
+		}
+
+		UE::FUsdPrim Prim = UsdStage.GetPrimAtPath(UE::FSdfPath{*PrimPath});
+		FName UEPropertyName = *UEPropertyNameStr;
+		if (!Prim || UEPropertyName == NAME_None)
+		{
+			return false;
+		}
+
+		TArray<double> TimeSampleUnion;
+		TArray<UE::FUsdAttribute> Attrs = UsdUtils::GetAttributesForProperty(Prim, UEPropertyName);
+		bool bSuccess = UE::FUsdAttribute::GetUnionedTimeSamples(Attrs, TimeSampleUnion);
+		if (!bSuccess)
+		{
+			return false;
+		}
+
+		const bool bIgnorePrimLocalTransform = false;
+		UsdToUnreal::FPropertyTrackReader Reader = UsdToUnreal::CreatePropertyTrackReader(Prim, UEPropertyName, bIgnorePrimLocalTransform);
+		if (Reader.BoolReader)
+		{
+			return ReadBools(UsdStage, TimeSampleUnion, Reader.BoolReader, OutPayloadData);
+		}
+		else if (Reader.ColorReader)
+		{
+			return ReadColors(UsdStage, TimeSampleUnion, Reader.ColorReader, OutPayloadData);
+		}
+		else if (Reader.FloatReader)
+		{
+			return ReadFloats(UsdStage, TimeSampleUnion, Reader.FloatReader, OutPayloadData);
+		}
+		else if (Reader.TransformReader)
+		{
+			return ReadTransforms(UsdStage, TimeSampleUnion, Reader.TransformReader, OutPayloadData);
+		}
+
+		return false;
+	}
+
 }	 // namespace UE::InterchangeUsdTranslator::Private
 #endif	  // USE_USD_SDK
 
@@ -1256,6 +1663,7 @@ bool UInterchangeUSDTranslator::Translate(UInterchangeBaseNodeContainer& NodeCon
 	{
 		return false;
 	}
+	ImplPtr->CurrentTrackSet = nullptr;
 
 	UInterchangeUsdTranslatorSettings* Settings = Cast<UInterchangeUsdTranslatorSettings>(GetSettings());
 	if (!Settings)
@@ -1278,11 +1686,17 @@ bool UInterchangeUSDTranslator::Translate(UInterchangeBaseNodeContainer& NodeCon
 		return false;
 	}
 
-	// Apply coordinate system conversion to the stage if we have one
-	if (Settings && Settings->bOverrideStageOptions)
+	// Apply stage settings
+	if (Settings)
 	{
-		UsdUtils::SetUsdStageMetersPerUnit(ImplPtr->UsdStage, Settings->StageOptions.MetersPerUnit);
-		UsdUtils::SetUsdStageUpAxis(ImplPtr->UsdStage, Settings->StageOptions.UpAxis);
+		// Apply coordinate system conversion to the stage if we have one
+		if (Settings->bOverrideStageOptions)
+		{
+			UsdUtils::SetUsdStageMetersPerUnit(ImplPtr->UsdStage, Settings->StageOptions.MetersPerUnit);
+			UsdUtils::SetUsdStageUpAxis(ImplPtr->UsdStage, Settings->StageOptions.UpAxis);
+		}
+
+		ImplPtr->UsdStage.SetInterpolationType(Settings->InterpolationType);
 	}
 
 	// Cache these so we don't have to keep converting these tokens over and over during translation
@@ -1297,7 +1711,7 @@ bool UInterchangeUSDTranslator::Translate(UInterchangeBaseNodeContainer& NodeCon
 	FTraversalInfo Info;
 	for (const FUsdPrim& Prim : ImplPtr->UsdStage.GetPseudoRoot().GetChildren())
 	{
-		Traverse(Prim, ImplPtr, NodeContainer, Info);
+		Traverse(Prim, *ImplPtr, NodeContainer, Info);
 	}
 
 	return true;
@@ -1315,6 +1729,7 @@ void UInterchangeUSDTranslator::ReleaseSource()
 	}
 
 	ImplPtr->UsdStage = UE::FUsdStage{};
+	ImplPtr->CurrentTrackSet = nullptr;
 
 	if (ImplPtr->TranslatorSettings)
 	{
@@ -1581,11 +1996,81 @@ TOptional<UE::Interchange::FImportBlockedImage> UInterchangeUSDTranslator::GetBl
 	return BlockData;
 }
 
+TFuture<TOptional<UE::Interchange::FAnimationPayloadData>> UInterchangeUSDTranslator::ResolveAnimationPayloadQuery(
+	const UE::Interchange::FAnimationPayloadQuery& PayloadQuery
+) const
+{
+	using namespace UE::Interchange;
+	using namespace UE::InterchangeUsdTranslator::Private;
+
+	return Async(
+		EAsyncExecution::TaskGraph,
+		[this, PayloadQuery]
+		{
+			TOptional<FAnimationPayloadData> Result;
+
+#if USE_USD_SDK
+			FAnimationPayloadData AnimationPayLoadData{PayloadQuery.SceneNodeUniqueID, PayloadQuery.PayloadKey};
+
+			UInterchangeUSDTranslatorImpl* ImplPtr = Impl.Get();
+			if (!ImplPtr)
+			{
+				return Result;
+			}
+
+			switch (PayloadQuery.PayloadKey.Type)
+			{
+				case EInterchangeAnimationPayLoadType::CURVE:	 // Fallthrough
+				case EInterchangeAnimationPayLoadType::STEPCURVE:
+				{
+					if (GetAnimationCurvePayloadData(ImplPtr->UsdStage, PayloadQuery.PayloadKey.UniqueId, AnimationPayLoadData))
+					{
+						Result.Emplace(AnimationPayLoadData);
+					}
+					break;
+				}
+				case EInterchangeAnimationPayLoadType::MORPHTARGETCURVE:
+				case EInterchangeAnimationPayLoadType::BAKED:
+				case EInterchangeAnimationPayLoadType::NONE:
+				default:
+				{
+					break;
+				}
+			}
+#endif	  // USE_USD_SDK
+
+			return Result;
+		}
+	);
+}
+
 TArray<UE::Interchange::FAnimationPayloadData> UInterchangeUSDTranslator::GetAnimationPayloadData(
-	const TArray<UE::Interchange::FAnimationPayloadQuery>& PayloadQuery
+	const TArray<UE::Interchange::FAnimationPayloadQuery>& PayloadQueries
 ) const
 {
-	return {};
+	// Note: This TFuture/Async approach is a bit overkill when it comes to the property track case, as each
+	// individual PayloadQueries has a single item in it (check FLevelSequenceHelper::PopulateAnimationTrack).
+	// This will actually help a lot with AnimSequences though, where RetrieveAnimationPayloads shows how all
+	// the skeletal joint tracks are queried in a single call (and the analogous for the morph target curves).
+
+	TArray<TFuture<TOptional<UE::Interchange::FAnimationPayloadData>>> AnimationPayloadFutures;
+	for (const UE::Interchange::FAnimationPayloadQuery& PayloadQuery : PayloadQueries)
+	{
+		AnimationPayloadFutures.Add(ResolveAnimationPayloadQuery(PayloadQuery));
+	}
+
+	TArray<UE::Interchange::FAnimationPayloadData> AnimationPayloads;
+	for (TFuture<TOptional<UE::Interchange::FAnimationPayloadData>>& AnimationPayloadFuture : AnimationPayloadFutures)
+	{
+		TOptional<UE::Interchange::FAnimationPayloadData> OptionalPayloadData = AnimationPayloadFuture.Get();
+		if (!OptionalPayloadData.IsSet())
+		{
+			continue;
+		}
+		AnimationPayloads.Add(OptionalPayloadData.GetValue());
+	}
+
+	return AnimationPayloads;
 }
 
 #undef LOCTEXT_NAMESPACE
diff --git a/Engine/Plugins/Interchange/Runtime/Source/Import/Public/Usd/InterchangeUsdTranslator.h b/Engine/Plugins/Interchange/Runtime/Source/Import/Public/Usd/InterchangeUsdTranslator.h
--- a/Engine/Plugins/Interchange/Runtime/Source/Import/Public/Usd/InterchangeUsdTranslator.h
+++ b/Engine/Plugins/Interchange/Runtime/Source/Import/Public/Usd/InterchangeUsdTranslator.h
@@ -14,6 +14,7 @@
 
 #include "InterchangeUsdTranslator.generated.h"
 
+enum class EUsdInterpolationType : uint8;
 namespace UE::InterchangeUsdTranslator::Private
 {
 	class UInterchangeUSDTranslatorImpl;
@@ -51,6 +52,10 @@ public:
 	UPROPERTY(EditAnywhere, Category = "USD Translator")
 	FName MaterialPurpose;
 
+	/** Describes how to interpolate between a timeSample value and the next */
+	UPROPERTY(EditAnywhere, Category = "USD Translator")
+	EUsdInterpolationType InterpolationType;
+
 	/** Whether to use the specified StageOptions instead of the stage's own settings */
 	UPROPERTY(EditAnywhere, Category = "USD Translator")
 	bool bOverrideStageOptions;
@@ -87,6 +92,10 @@ public:
 	virtual void SetSettings(const UInterchangeTranslatorSettings* InterchangeTranslatorSettings) override;
 	/** End UInterchangeTranslatorBase API*/
 
+	TFuture<TOptional<UE::Interchange::FAnimationPayloadData>> ResolveAnimationPayloadQuery(
+		const UE::Interchange::FAnimationPayloadQuery& PayloadQuery
+	) const;
+
 	/** Begin Interchange payload interfaces */
 	virtual TFuture<TOptional<UE::Interchange::FMeshPayloadData>> GetMeshPayloadData(
 		const FInterchangeMeshPayLoadKey& PayLoadKey,
@@ -101,7 +110,8 @@ public:
 		TOptional<FString>& AlternateTexturePath
 	) const override;
 
-	virtual TArray<UE::Interchange::FAnimationPayloadData> GetAnimationPayloadData(const TArray<UE::Interchange::FAnimationPayloadQuery>& PayloadQuery
+	virtual TArray<UE::Interchange::FAnimationPayloadData> GetAnimationPayloadData(
+		const TArray<UE::Interchange::FAnimationPayloadQuery>& PayloadQueries
 	) const override;
 	/** End Interchange payload interfaces */
 
diff --git a/Engine/Plugins/Runtime/USDCore/Source/USDUtilities/Private/USDConversionUtils.cpp b/Engine/Plugins/Runtime/USDCore/Source/USDUtilities/Private/USDConversionUtils.cpp
--- a/Engine/Plugins/Runtime/USDCore/Source/USDUtilities/Private/USDConversionUtils.cpp
+++ b/Engine/Plugins/Runtime/USDCore/Source/USDUtilities/Private/USDConversionUtils.cpp
@@ -915,57 +915,14 @@
 
 	FScopedUsdAllocs UsdAllocs;
 
-	pxr::UsdGeomXformable Xformable(Prim);
-	if (Xformable)
+	if (HasAnimatedTransform(Prim))
 	{
-		std::vector<double> TimeSamples;
-		Xformable.GetTimeSamples(&TimeSamples);
-
-		if (TimeSamples.size() > 0)
-		{
-			return true;
-		}
-
-		// If this xformable has an op to reset the xform stack and one of its ancestors is animated, then we need to pretend
-		// its transform is also animated. This because that op effectively means "discard the parent transform and treat this
-		// as a direct world transform", but when reading we'll manually recompute the relative transform to its parent anyway
-		// (for simplicity's sake). If that parent (or any of its ancestors) is being animated, we'll need to recompute this
-		// for every animation keyframe, which basically means we're animated too
-		if (Xformable.GetResetXformStack())
-		{
-			pxr::UsdPrim AncestorPrim = Prim.GetParent();
-			while (AncestorPrim && !AncestorPrim.IsPseudoRoot())
-			{
-				if (pxr::UsdGeomXformable AncestorXformable{AncestorPrim})
-				{
-					std::vector<double> AncestorTimeSamples;
-					if (AncestorXformable.GetTimeSamples(&AncestorTimeSamples) && AncestorTimeSamples.size() > 0)
-					{
-						return true;
-					}
-
-					// The exception is if our ancestor also wants to reset its xform stack (i.e. its transform is meant to be
-					// used as the world transform). In this case we don't need to care about higher up ancestors anymore, as
-					// their transforms wouldn't affect below this prim anyway
-					if (AncestorXformable.GetResetXformStack())
-					{
-						break;
-					}
-				}
-
-				AncestorPrim = AncestorPrim.GetParent();
-			}
-		}
+		return true;
 	}
 
-	const std::vector<pxr::UsdAttribute>& Attributes = Prim.GetAttributes();
-	for (const pxr::UsdAttribute& Attribute : Attributes)
+	if (HasAnimatedAttributes(Prim))
 	{
-		std::vector<double> TimeSamples;
-		if (Attribute.GetTimeSamples(&TimeSamples) && TimeSamples.size() > 0)
-		{
-			return true;
-		}
+		return true;
 	}
 
 	if (pxr::UsdSkelSkeleton Skeleton{Prim})
@@ -1012,6 +969,83 @@
 	return false;
 }
 
+bool UsdUtils::HasAnimatedAttributes(const pxr::UsdPrim& Prim)
+{
+	if (!Prim || !Prim.IsActive())
+	{
+		return false;
+	}
+
+	FScopedUsdAllocs UsdAllocs;
+
+	const std::vector<pxr::UsdAttribute>& Attributes = Prim.GetAttributes();
+	for (const pxr::UsdAttribute& Attribute : Attributes)
+	{
+		std::vector<double> TimeSamples;
+		if (Attribute.GetTimeSamples(&TimeSamples) && TimeSamples.size() > 0)
+		{
+			return true;
+		}
+	}
+
+	return false;
+}
+
+bool UsdUtils::HasAnimatedTransform(const pxr::UsdPrim& Prim)
+{
+	if (!Prim || !Prim.IsActive())
+	{
+		return false;
+	}
+
+	FScopedUsdAllocs UsdAllocs;
+
+	pxr::UsdGeomXformable Xformable(Prim);
+	if (Xformable)
+	{
+		std::vector<double> TimeSamples;
+		Xformable.GetTimeSamples(&TimeSamples);
+
+		if (TimeSamples.size() > 0)
+		{
+			return true;
+		}
+
+		// If this xformable has an op to reset the xform stack and one of its ancestors is animated, then we need to pretend
+		// its transform is also animated. This because that op effectively means "discard the parent transform and treat this
+		// as a direct world transform", but when reading we'll manually recompute the relative transform to its parent anyway
+		// (for simplicity's sake). If that parent (or any of its ancestors) is being animated, we'll need to recompute this
+		// for every animation keyframe, which basically means we're animated too
+		if (Xformable.GetResetXformStack())
+		{
+			pxr::UsdPrim AncestorPrim = Prim.GetParent();
+			while (AncestorPrim && !AncestorPrim.IsPseudoRoot())
+			{
+				if (pxr::UsdGeomXformable AncestorXformable{AncestorPrim})
+				{
+					std::vector<double> AncestorTimeSamples;
+					if (AncestorXformable.GetTimeSamples(&AncestorTimeSamples) && AncestorTimeSamples.size() > 0)
+					{
+						return true;
+					}
+
+					// The exception is if our ancestor also wants to reset its xform stack (i.e. its transform is meant to be
+					// used as the world transform). In this case we don't need to care about higher up ancestors anymore, as
+					// their transforms wouldn't affect below this prim anyway
+					if (AncestorXformable.GetResetXformStack())
+					{
+						break;
+					}
+				}
+
+				AncestorPrim = AncestorPrim.GetParent();
+			}
+		}
+	}
+
+	return false;
+}
+
 bool UsdUtils::HasAnimatedVisibility(const pxr::UsdPrim& Prim)
 {
 	if (!Prim || !Prim.IsActive())
diff --git a/Engine/Plugins/Runtime/USDCore/Source/USDUtilities/Private/USDPrimConversion.cpp b/Engine/Plugins/Runtime/USDCore/Source/USDUtilities/Private/USDPrimConversion.cpp
--- a/Engine/Plugins/Runtime/USDCore/Source/USDUtilities/Private/USDPrimConversion.cpp
+++ b/Engine/Plugins/Runtime/USDCore/Source/USDUtilities/Private/USDPrimConversion.cpp
@@ -5315,6 +5315,11 @@ bool UnrealToUsd::ConvertXformable(
 }
 
 TArray<UE::FUsdAttribute> UnrealToUsd::GetAttributesForProperty(const UE::FUsdPrim& Prim, const FName& PropertyPath)
+{
+	return UsdUtils::GetAttributesForProperty(Prim, PropertyPath);
+}
+
+TArray<UE::FUsdAttribute> UsdUtils::GetAttributesForProperty(const UE::FUsdPrim& Prim, const FName& PropertyPath)
 {
 	using namespace UnrealIdentifiers;
 
@@ -5549,6 +5554,93 @@ TArray<UE::FUsdAttribute> UnrealToUsd::GetAttributesForProperty(const UE::FUsdPr
 	return {};
 }
 
+TArray<FName> UsdUtils::GetPropertiesForAttribute(const UE::FUsdPrim& Prim, const FString& AttrName)
+{
+	using namespace UnrealIdentifiers;
+
+	// Note: This function may seem confusing and "backwards", but it is correct.
+	// The logic here is essentially "If this attribute is animated, which tracks do I need to generate?"
+
+	FScopedUsdAllocs Allocs;
+	pxr::UsdPrim UsdPrim{Prim};
+
+	pxr::TfToken AttrToken = UnrealToUsd::ConvertToken(*AttrName).Get();
+
+	// First check the standard cases, which should be the most common
+	// clang-format off
+	const static std::unordered_map<pxr::TfToken, TArray<FName>, pxr::TfHash> StandardMapping = {
+		{pxr::UsdGeomTokens->visibility, 					{HiddenInGamePropertyName}},
+
+		{pxr::UsdGeomTokens->focalLength, 					{CurrentFocalLengthPropertyName}},
+		{pxr::UsdGeomTokens->focusDistance, 				{ManualFocusDistancePropertyName}},
+		{pxr::UsdGeomTokens->fStop, 						{CurrentAperturePropertyName}},
+		{pxr::UsdGeomTokens->horizontalAperture, 			{SensorWidthPropertyName}},
+		{pxr::UsdGeomTokens->verticalAperture, 				{SensorHeightPropertyName}},
+
+		{pxr::UsdLuxTokens->inputsIntensity, 				{IntensityPropertyName}},
+		{pxr::UsdLuxTokens->inputsExposure, 				{IntensityPropertyName}}, // In USD true intensity is just a function of "intensity" and "exposure"
+		{pxr::UsdLuxTokens->inputsColor, 					{LightColorPropertyName}},
+		{pxr::UsdLuxTokens->inputsColorTemperature, 		{TemperaturePropertyName}},
+		{pxr::UsdLuxTokens->inputsEnableColorTemperature, 	{UseTemperaturePropertyName}},
+	};
+	// clang-format on
+	std::unordered_map<pxr::TfToken, TArray<FName>, pxr::TfHash>::const_iterator Iter = StandardMapping.find(AttrToken);
+	if (Iter != StandardMapping.end())
+	{
+		return Iter->second;
+	}
+
+	// Check to see if it's one of the edge cases that depends on the prim type.
+	if (AttrToken == pxr::UsdLuxTokens->inputsRadius)
+	{
+		if (UsdPrim.IsA<pxr::UsdLuxSphereLight>())
+		{
+			return {SourceRadiusPropertyName, IntensityPropertyName};
+		}
+		else if (UsdPrim.IsA<pxr::UsdLuxDiskLight>())
+		{
+			return {SourceWidthPropertyName, SourceHeightPropertyName, IntensityPropertyName};
+		}
+	}
+	else if (AttrToken == pxr::UsdLuxTokens->inputsWidth)
+	{
+		if (UsdPrim.IsA<pxr::UsdLuxRectLight>())
+		{
+			return {SourceWidthPropertyName, IntensityPropertyName};
+		}
+	}
+	else if (AttrToken == pxr::UsdLuxTokens->inputsHeight)
+	{
+		if (UsdPrim.IsA<pxr::UsdLuxRectLight>())
+		{
+			return {SourceHeightPropertyName, IntensityPropertyName};
+		}
+	}
+	else if (AttrToken == pxr::UsdLuxTokens->inputsAngle)
+	{
+		if (UsdPrim.IsA<pxr::UsdLuxDistantLight>())
+		{
+			return {LightSourceAnglePropertyName};
+		}
+	}
+	else if (AttrToken == pxr::UsdLuxTokens->inputsShapingConeAngle)
+	{
+		if (UsdPrim.IsA<pxr::UsdLuxSphereLight>() && UsdPrim.HasAPI<pxr::UsdLuxShapingAPI>())
+		{
+			return {OuterConeAnglePropertyName, IntensityPropertyName};
+		}
+	}
+	else if (AttrToken == pxr::UsdLuxTokens->inputsShapingConeSoftness)
+	{
+		if (UsdPrim.IsA<pxr::UsdLuxSphereLight>() && UsdPrim.HasAPI<pxr::UsdLuxShapingAPI>())
+		{
+			return {InnerConeAnglePropertyName, IntensityPropertyName};
+		}
+	}
+
+	return {};
+}
+
 bool UnrealToUsd::ConvertDrawModeComponent(
 	const UUsdDrawModeComponent& DrawModeComponent,
 	pxr::UsdPrim& UsdPrim,
diff --git a/Engine/Plugins/Runtime/USDCore/Source/USDUtilities/Public/USDConversionUtils.h b/Engine/Plugins/Runtime/USDCore/Source/USDUtilities/Public/USDConversionUtils.h
--- a/Engine/Plugins/Runtime/USDCore/Source/USDUtilities/Public/USDConversionUtils.h
+++ b/Engine/Plugins/Runtime/USDCore/Source/USDUtilities/Public/USDConversionUtils.h
@@ -252,6 +252,8 @@ namespace UsdUtils
 	}
 
 	USDUTILITIES_API bool IsAnimated(const pxr::UsdPrim& Prim);
+	USDUTILITIES_API bool HasAnimatedAttributes(const pxr::UsdPrim& Prim);
+	USDUTILITIES_API bool HasAnimatedTransform(const pxr::UsdPrim& Prim);
 	USDUTILITIES_API bool HasAnimatedVisibility(const pxr::UsdPrim& Prim);
 
 	/**
diff --git a/Engine/Plugins/Runtime/USDCore/Source/USDUtilities/Public/USDPrimConversion.h b/Engine/Plugins/Runtime/USDCore/Source/USDUtilities/Public/USDPrimConversion.h
--- a/Engine/Plugins/Runtime/USDCore/Source/USDUtilities/Public/USDPrimConversion.h
+++ b/Engine/Plugins/Runtime/USDCore/Source/USDUtilities/Public/USDPrimConversion.h
@@ -442,11 +442,7 @@ namespace UnrealToUsd
 		TSet<FName>& OutPropertyPathsToRefresh
 	);
 
-	/**
-	 * Returns the attributes that in USD correspond to a property in UE.
-	 * The first one is the "main" attribute, when appropriate. e.g. for a RectLight prim and UnrealIdentifiers::IntensityPropertyName,
-	 * we'll receive intensity, exposure, width and height, in that order.
-	 */
+	UE_DEPRECATED(5.5, "This function has moved to the UsdUtils namespace")
 	USDUTILITIES_API TArray<UE::FUsdAttribute> GetAttributesForProperty(const UE::FUsdPrim& Prim, const FName& PropertyPath);
 
 	/**
@@ -501,4 +497,22 @@ namespace UnrealToUsd
 	);
 }	 // namespace UnrealToUsd
 
+namespace UsdUtils
+{
+	/**
+	 * Returns the attributes that in USD correspond to a property in UE.
+	 * The first one is the "main" attribute, when appropriate. e.g. for a RectLight prim and UnrealIdentifiers::IntensityPropertyName,
+	 * we'll receive intensity, exposure, width and height, in that order.
+	 */
+	USDUTILITIES_API TArray<UE::FUsdAttribute> GetAttributesForProperty(const UE::FUsdPrim& Prim, const FName& PropertyPath);
+
+	/**
+	 * Returns the names of UE properties that need to be written to, when reading an USD attribute.
+	 * Example: If AttrName is "inputs:radius" and Prim is a SphereLight, this will return [UnrealIdentifiers::SourceRadiusPropertyName];
+	 * 			If AttrName is "inputs:radius" and Prim is a DiskLight, this will return
+	 *          [UnrealIdentifiers::SourceWidthPropertyName, UnrealIdentifiers::SourceHeightPropertyName].
+	 */
+	USDUTILITIES_API TArray<FName> GetPropertiesForAttribute(const UE::FUsdPrim& Prim, const FString& AttrName);
+}
+
 #endif	  // #if USE_USD_SDK
diff --git a/Engine/Plugins/Runtime/USDCore/Source/UnrealUSDWrapper/Private/UsdWrappers/SdfPath.cpp b/Engine/Plugins/Runtime/USDCore/Source/UnrealUSDWrapper/Private/UsdWrappers/SdfPath.cpp
--- a/Engine/Plugins/Runtime/USDCore/Source/UnrealUSDWrapper/Private/UsdWrappers/SdfPath.cpp
+++ b/Engine/Plugins/Runtime/USDCore/Source/UnrealUSDWrapper/Private/UsdWrappers/SdfPath.cpp
@@ -220,6 +220,15 @@ namespace UE
 #endif	  // #if USE_USD_SDK
 	}
 
+	bool FSdfPath::IsPrimPath() const
+	{
+#if USE_USD_SDK
+		return Impl->PxrSdfPath.Get().IsPrimPath();
+#else
+		return false;
+#endif	  // #if USE_USD_SDK
+	}
+
 	bool FSdfPath::IsAbsoluteRootOrPrimPath() const
 	{
 #if USE_USD_SDK
@@ -229,6 +238,15 @@ namespace UE
 #endif	  // #if USE_USD_SDK
 	}
 
+	bool FSdfPath::IsPrimPropertyPath() const
+	{
+#if USE_USD_SDK
+		return Impl->PxrSdfPath.Get().IsPrimPropertyPath();
+#else
+		return false;
+#endif	  // #if USE_USD_SDK
+	}
+
 	FString FSdfPath::GetName() const
 	{
 #if USE_USD_SDK
@@ -274,6 +292,15 @@ namespace UE
 #endif	  // #if USE_USD_SDK
 	}
 
+	FSdfPath FSdfPath::GetPrimPath() const
+	{
+#if USE_USD_SDK
+		return FSdfPath(Impl->PxrSdfPath.Get().GetPrimPath());
+#else
+		return FSdfPath();
+#endif	  // #if USE_USD_SDK
+	}
+
 	FSdfPath FSdfPath::AppendPath(const UE::FSdfPath& NewRelativeSuffix) const
 	{
 #if USE_USD_SDK
diff --git a/Engine/Plugins/Runtime/USDCore/Source/UnrealUSDWrapper/Public/UsdWrappers/SdfPath.h b/Engine/Plugins/Runtime/USDCore/Source/UnrealUSDWrapper/Public/UsdWrappers/SdfPath.h
--- a/Engine/Plugins/Runtime/USDCore/Source/UnrealUSDWrapper/Public/UsdWrappers/SdfPath.h
+++ b/Engine/Plugins/Runtime/USDCore/Source/UnrealUSDWrapper/Public/UsdWrappers/SdfPath.h
@@ -65,7 +65,9 @@ namespace UE
 		bool IsEmpty() const noexcept;
 
 		bool IsAbsoluteRootPath() const;
+		bool IsPrimPath() const;
 		bool IsAbsoluteRootOrPrimPath() const;
+		bool IsPrimPropertyPath() const;
 		FString GetName() const;
 		FString GetElementString() const;
 		FSdfPath GetAbsoluteRootOrPrimPath() const;
@@ -73,6 +75,7 @@ namespace UE
 		FSdfPath ReplaceName(const TCHAR* NewLeafName) const;
 
 		FSdfPath GetParentPath() const;
+		FSdfPath GetPrimPath() const;
 		FSdfPath AppendPath(const UE::FSdfPath& NewRelativeSuffix) const;
 		FSdfPath AppendChild(const TCHAR* ChildName) const;
 		FSdfPath AppendProperty(FName PropertyName) const;
