Interchange: Fix vertex color lod import

#rnx

diff --git a/Engine/Plugins/Interchange/Runtime/Source/Pipelines/Private/InterchangeGenericMeshPipeline.cpp b/Engine/Plugins/Interchange/Runtime/Source/Pipelines/Private/InterchangeGenericMeshPipeline.cpp
--- a/Engine/Plugins/Interchange/Runtime/Source/Pipelines/Private/InterchangeGenericMeshPipeline.cpp
+++ b/Engine/Plugins/Interchange/Runtime/Source/Pipelines/Private/InterchangeGenericMeshPipeline.cpp
@@ -114,37 +114,36 @@ void UInterchangeGenericMeshPipeline::AdjustSettingsForContext(const FInterchang
 		else if (ContextParams.ContextType == EInterchangePipelineContext::AssetCustomLODImport
 			|| ContextParams.ContextType == EInterchangePipelineContext::AssetCustomLODReimport)
 		{
-			CommonMeshesProperties->bBakeMeshes = true;
-			CommonMeshesProperties->bBakePivotMeshes = false;
-			CommonMeshesProperties->bImportLods = false;
-			CommonMeshesProperties->bKeepSectionsSeparate = false;
-			CommonMeshesProperties->VertexColorImportOption = EInterchangeVertexColorImportOption::IVCIO_Ignore;
-			bBuildNanite = false;
-			LodGroup = NAME_None;
-			bImportCollision = false;
-			bImportCollisionAccordingToMeshName = false;
-			bGenerateLightmapUVs = false;
-			bGenerateDistanceFieldAsIfTwoSided = false;
-			bSupportFaceRemap = false;
 			//We are importing custom LODs
 			if (ContextParams.ImportObjectType)
 			{
 				//If we have a provided import object type we can make sure we import the correct type
 				if (ContextParams.ImportObjectType->IsChildOf<UStaticMesh>())
 				{
-					CommonMeshesProperties->bAutoDetectMeshType = false;
+					bImportStaticMeshes = true;
 					CommonMeshesProperties->ForceAllMeshAsType = EInterchangeForceMeshType::IFMT_StaticMesh;
+					CommonMeshesProperties->bAutoDetectMeshType = false;
 					bImportSkeletalMeshes = false;
-					bImportStaticMeshes = true;
 					bCombineStaticMeshes = true;
+					LodGroup = NAME_None;
+					bSupportFaceRemap = false;
+					bImportCollision = false;
+					bImportCollisionAccordingToMeshName = false;
+					bGenerateLightmapUVs = false;
+					bGenerateDistanceFieldAsIfTwoSided = false;
 				}
 				else if (ContextParams.ImportObjectType->IsChildOf<USkeletalMesh>())
 				{
-					CommonMeshesProperties->bAutoDetectMeshType = false;
-					CommonMeshesProperties->ForceAllMeshAsType = EInterchangeForceMeshType::IFMT_SkeletalMesh;
 					bImportSkeletalMeshes = true;
+					CommonMeshesProperties->ForceAllMeshAsType = EInterchangeForceMeshType::IFMT_SkeletalMesh;
+					CommonMeshesProperties->bAutoDetectMeshType = false;
+					bCreatePhysicsAsset = false;
 					bImportStaticMeshes = false;
 				}
+				else
+				{
+					CommonMeshesProperties->bAutoDetectMeshType = true;
+				}
 			}
 		}
 	}
