USD: Fix another instance of editor-only code usage at runtime.

diff --git a/Engine/Plugins/Interchange/Runtime/Source/Import/Private/Usd/InterchangeUsdTranslator.cpp b/Engine/Plugins/Interchange/Runtime/Source/Import/Private/Usd/InterchangeUsdTranslator.cpp
--- a/Engine/Plugins/Interchange/Runtime/Source/Import/Private/Usd/InterchangeUsdTranslator.cpp
+++ b/Engine/Plugins/Interchange/Runtime/Source/Import/Private/Usd/InterchangeUsdTranslator.cpp
@@ -866,6 +866,7 @@ namespace UE::InterchangeUsdTranslator::Private
 		SkeletonPrimNode.SetCustomBindPoseLocalTransform(&NodeContainer, FTransform::Identity);
 		SkeletonPrimNode.SetCustomTimeZeroLocalTransform(&NodeContainer, FTransform::Identity);
 
+#if WITH_EDITOR
 		// Convert the skeleton bones/joints into ConvertedData
 		UE::FUsdSkelSkeletonQuery SkelQuery = Info.FurthestSkelCache->GetSkelQuery(Prim);
 		const bool bEnsureAtLeastOneBone = false;
@@ -926,6 +927,7 @@ namespace UE::InterchangeUsdTranslator::Private
 			FWriteScopeLock ScopedInfoWriteLock{TranslatorImpl.CachedTraversalInfoLock};
 			TranslatorImpl.NodeUidToCachedTraversalInfo.Add(SkeletonPrimNodeUid, Info);
 		}
+#endif
 	}
 
 	void UpdateTraversalInfo(FTraversalInfo& Info, const UE::FUsdPrim& CurrentPrim)
