MaterialX:
- Improve quality of material function UsdPreviewSurface, only for substrate
- Add support of integer for UsdPreviewSurface handled as scalars
- remove MX_UsdPreviewSurface material function and put the material function in Engine/Substrate

diff --git a/Engine/Plugins/Interchange/Runtime/Source/Import/Private/MaterialX/MaterialXUtils/MaterialXSurfaceShaderAbstract.cpp b/Engine/Plugins/Interchange/Runtime/Source/Import/Private/MaterialX/MaterialXUtils/MaterialXSurfaceShaderAbstract.cpp
--- a/Engine/Plugins/Interchange/Runtime/Source/Import/Private/MaterialX/MaterialXUtils/MaterialXSurfaceShaderAbstract.cpp
+++ b/Engine/Plugins/Interchange/Runtime/Source/Import/Private/MaterialX/MaterialXUtils/MaterialXSurfaceShaderAbstract.cpp
@@ -1599,6 +1599,10 @@ UInterchangeShaderNode* FMaterialXSurfaceShaderAbstract::Translate(EInterchangeM
 		{
 			ConnectNodeOutputToInput(Input->getName().c_str(), FunctionCallShaderNode, Input->getName().c_str(), GetValue<FVector4f>(mx::fromValueString<mx::Vector4>(DefaultValue->getValueString())), bInputInTangentSpace);
 		}
+		else if(DefaultValue->getTypeString() == mx::Type::Integer)
+		{
+			ConnectNodeOutputToInput(Input->getName().c_str(), FunctionCallShaderNode, Input->getName().c_str(), mx::fromValueString<int32>(DefaultValue->getValueString()));
+		}
 	}
 
 	return FunctionCallShaderNode;
diff --git a/Engine/Plugins/Interchange/Runtime/Source/Import/Private/MaterialX/MaterialXUtils/MaterialXUsdPreviewSurfaceShader.cpp b/Engine/Plugins/Interchange/Runtime/Source/Import/Private/MaterialX/MaterialXUtils/MaterialXUsdPreviewSurfaceShader.cpp
--- a/Engine/Plugins/Interchange/Runtime/Source/Import/Private/MaterialX/MaterialXUtils/MaterialXUsdPreviewSurfaceShader.cpp
+++ b/Engine/Plugins/Interchange/Runtime/Source/Import/Private/MaterialX/MaterialXUtils/MaterialXUsdPreviewSurfaceShader.cpp
@@ -55,7 +55,7 @@ void FMaterialXUsdPreviewSurfaceShader::Translate(MaterialX::NodePtr UsdPreviewS
 	}
 	else
 	{
-		UInterchangeShaderPortsAPI::ConnectOuputToInputByName(ShaderGraphNode, SubstrateMaterial::Parameters::FrontMaterial.ToString(), UsdPreviewSurfaceShaderNode->GetUniqueID(), SubstrateMaterial::Parameters::FrontMaterial.ToString());
+		UInterchangeShaderPortsAPI::ConnectOuputToInputByName(ShaderGraphNode, SubstrateMaterial::Parameters::FrontMaterial.ToString(), UsdPreviewSurfaceShaderNode->GetUniqueID(), UsdPreviewSurface::SubstrateMaterial::Outputs::FrontMaterial.ToString());
 		UInterchangeShaderPortsAPI::ConnectOuputToInputByName(ShaderGraphNode, SubstrateMaterial::Parameters::Occlusion.ToString(), UsdPreviewSurfaceShaderNode->GetUniqueID(), SubstrateMaterial::Parameters::Occlusion.ToString());
 
 		if(UInterchangeShaderPortsAPI::HasInput(UsdPreviewSurfaceShaderNode, UsdPreviewSurface::Parameters::Displacement))
@@ -66,7 +66,6 @@ void FMaterialXUsdPreviewSurfaceShader::Translate(MaterialX::NodePtr UsdPreviewS
 		if(UInterchangeShaderPortsAPI::HasInput(UsdPreviewSurfaceShaderNode, UsdPreviewSurface::Parameters::Opacity))
 		{
 			ShaderGraphNode->SetCustomBlendMode(EBlendMode::BLEND_TranslucentColoredTransmittance);
-			UInterchangeShaderPortsAPI::ConnectOuputToInputByName(ShaderGraphNode, PBRMR::Parameters::Opacity.ToString(), UsdPreviewSurfaceShaderNode->GetUniqueID(), PBRMR::Parameters::Opacity.ToString());
 		}
 	}
 }
diff --git a/Engine/Plugins/Interchange/Runtime/Source/Import/Public/MaterialX/MaterialXUtils/MaterialXSurfaceShaderAbstract.h b/Engine/Plugins/Interchange/Runtime/Source/Import/Public/MaterialX/MaterialXUtils/MaterialXSurfaceShaderAbstract.h
--- a/Engine/Plugins/Interchange/Runtime/Source/Import/Public/MaterialX/MaterialXUtils/MaterialXSurfaceShaderAbstract.h
+++ b/Engine/Plugins/Interchange/Runtime/Source/Import/Public/MaterialX/MaterialXUtils/MaterialXSurfaceShaderAbstract.h
@@ -135,7 +135,8 @@ protected:
 			if(!bIsConnected)
 			{
 				// only handle float, linear color and vector here, for other types, the child should handle them as it is most likely not an input but a parameter to set in Interchange
-				if constexpr(std::is_same_v<decltype(DefaultValue), float>)
+				// handle integers as scalars
+				if constexpr(std::is_same_v<decltype(DefaultValue), float> || std::is_same_v<decltype(DefaultValue), int32>)
 				{
 					bIsConnected = AddFloatAttribute(Input, InputShaderName, ShaderNode, DefaultValue);
 				}
diff --git a/Engine/Plugins/Interchange/Runtime/Source/Nodes/Public/InterchangeMaterialDefinitions.h b/Engine/Plugins/Interchange/Runtime/Source/Nodes/Public/InterchangeMaterialDefinitions.h
--- a/Engine/Plugins/Interchange/Runtime/Source/Nodes/Public/InterchangeMaterialDefinitions.h
+++ b/Engine/Plugins/Interchange/Runtime/Source/Nodes/Public/InterchangeMaterialDefinitions.h
@@ -626,6 +626,17 @@ namespace Interchange
 				const FName Displacement = TEXT("displacement");
 				const FName Occlusion = TEXT("occlusion");
 			}
+
+			// These outputs are only used for Substrate
+			namespace SubstrateMaterial
+			{
+				namespace Outputs
+				{
+					const FName FrontMaterial = TEXT("Substrate UsdPreviewSurface");
+					const FName Displacement = TEXT("Displacement");
+					const FName Occlusion = TEXT("Occlusion");
+				}
+			}
 		}
 
 		namespace Unlit
diff --git a/Engine/Plugins/Interchange/Runtime/Source/Pipelines/Private/InterchangeMaterialXPipeline.cpp b/Engine/Plugins/Interchange/Runtime/Source/Pipelines/Private/InterchangeMaterialXPipeline.cpp
--- a/Engine/Plugins/Interchange/Runtime/Source/Pipelines/Private/InterchangeMaterialXPipeline.cpp
+++ b/Engine/Plugins/Interchange/Runtime/Source/Pipelines/Private/InterchangeMaterialXPipeline.cpp
@@ -31,7 +31,8 @@ namespace
 	constexpr const TCHAR* TransmissionSurfaceSubstratePath = TEXT("/Engine/Functions/Substrate/Substrate-StandardSurface-Translucent.Substrate-StandardSurface-Translucent");
 	MATERIALX_FUNCTIONS_SUBSTRATE_PATH(SurfaceUnlit);
 	MATERIALX_FUNCTIONS_SUBSTRATE_PATH(Surface);
-	MATERIALX_FUNCTIONS_SUBSTRATE_PATH(UsdPreviewSurface);
+	constexpr const TCHAR* UsdPreviewSurfaceFunctionsPath = TEXT("/Interchange/Functions/MX_UsdPreviewSurface.MX_UsdPreviewSurface");
+	constexpr const TCHAR* UsdPreviewSurfaceSubstratePath = TEXT("/Engine/Functions/Substrate/MF_Substrate_UsdPreviewSurface.MF_Substrate_UsdPreviewSurface");
 	
 	MATERIALX_FUNCTIONS_SUBSTRATE_PATH(OrenNayarBSDF);
 	MATERIALX_FUNCTIONS_SUBSTRATE_PATH(BurleyDiffuseBSDF);
@@ -58,6 +59,7 @@ TMap<EInterchangeMaterialXSettings, TPair<TSet<FName>, TSet<FName>>> UMaterialXP
 
 UMaterialXPipelineSettings::UMaterialXPipelineSettings()
 {
+	using namespace UE::Interchange::Materials;
 #if WITH_EDITOR
 	if(HasAnyFlags(EObjectFlags::RF_ClassDefaultObject))
 	{
@@ -65,69 +67,69 @@ UMaterialXPipelineSettings::UMaterialXPipelineSettings()
 
 		SettingsInputsOutputs = {
 			//Surface Shaders
-						{
+			{
 				UMaterialXPipelineSettings::ToEnumKey(EInterchangeMaterialXShaders::OpenPBRSurface),
 				{
 				// OpenPBRSurface Inputs
 				TSet<FName>{
-					UE::Interchange::Materials::OpenPBRSurface::Parameters::BaseWeight,
-					UE::Interchange::Materials::OpenPBRSurface::Parameters::BaseColor,
-					UE::Interchange::Materials::OpenPBRSurface::Parameters::BaseRoughness,
-					UE::Interchange::Materials::OpenPBRSurface::Parameters::BaseMetalness,
-					UE::Interchange::Materials::OpenPBRSurface::Parameters::SpecularWeight,
-					UE::Interchange::Materials::OpenPBRSurface::Parameters::SpecularColor,
-					UE::Interchange::Materials::OpenPBRSurface::Parameters::SpecularRoughness,
-					UE::Interchange::Materials::OpenPBRSurface::Parameters::SpecularIOR,
-					UE::Interchange::Materials::OpenPBRSurface::Parameters::SpecularIORLevel,
-					UE::Interchange::Materials::OpenPBRSurface::Parameters::SpecularAnisotropy,
-					UE::Interchange::Materials::OpenPBRSurface::Parameters::SpecularRotation,
-					UE::Interchange::Materials::OpenPBRSurface::Parameters::SubsurfaceWeight,
-					UE::Interchange::Materials::OpenPBRSurface::Parameters::SubsurfaceColor,
-					UE::Interchange::Materials::OpenPBRSurface::Parameters::SubsurfaceRadius,
-					UE::Interchange::Materials::OpenPBRSurface::Parameters::SubsurfaceRadiusScale,
-					UE::Interchange::Materials::OpenPBRSurface::Parameters::SubsurfaceAnisotropy,
-					UE::Interchange::Materials::OpenPBRSurface::Parameters::FuzzWeight,
-					UE::Interchange::Materials::OpenPBRSurface::Parameters::FuzzColor,
-					UE::Interchange::Materials::OpenPBRSurface::Parameters::FuzzRoughness,
-					UE::Interchange::Materials::OpenPBRSurface::Parameters::CoatWeight,
-					UE::Interchange::Materials::OpenPBRSurface::Parameters::CoatColor,
-					UE::Interchange::Materials::OpenPBRSurface::Parameters::CoatRoughness,
-					UE::Interchange::Materials::OpenPBRSurface::Parameters::CoatAnisotropy,
-					UE::Interchange::Materials::OpenPBRSurface::Parameters::CoatRotation,
-					UE::Interchange::Materials::OpenPBRSurface::Parameters::CoatIOR,
-					UE::Interchange::Materials::OpenPBRSurface::Parameters::CoatIORLevel,
-					UE::Interchange::Materials::OpenPBRSurface::Parameters::GeometryCoatNormal,
-					UE::Interchange::Materials::OpenPBRSurface::Parameters::ThinFilmThickness,
-					UE::Interchange::Materials::OpenPBRSurface::Parameters::ThinFilmIOR,
-					UE::Interchange::Materials::OpenPBRSurface::Parameters::EmissionLuminance,
-					UE::Interchange::Materials::OpenPBRSurface::Parameters::EmissionColor,
-					UE::Interchange::Materials::OpenPBRSurface::Parameters::GeometryNormal,
-					UE::Interchange::Materials::OpenPBRSurface::Parameters::GeometryTangent,
-					UE::Interchange::Materials::OpenPBRSurface::Parameters::GeometryOpacity,
-					UE::Interchange::Materials::OpenPBRSurface::Parameters::GeometryThinWalled,
+					OpenPBRSurface::Parameters::BaseWeight,
+					OpenPBRSurface::Parameters::BaseColor,
+					OpenPBRSurface::Parameters::BaseRoughness,
+					OpenPBRSurface::Parameters::BaseMetalness,
+					OpenPBRSurface::Parameters::SpecularWeight,
+					OpenPBRSurface::Parameters::SpecularColor,
+					OpenPBRSurface::Parameters::SpecularRoughness,
+					OpenPBRSurface::Parameters::SpecularIOR,
+					OpenPBRSurface::Parameters::SpecularIORLevel,
+					OpenPBRSurface::Parameters::SpecularAnisotropy,
+					OpenPBRSurface::Parameters::SpecularRotation,
+					OpenPBRSurface::Parameters::SubsurfaceWeight,
+					OpenPBRSurface::Parameters::SubsurfaceColor,
+					OpenPBRSurface::Parameters::SubsurfaceRadius,
+					OpenPBRSurface::Parameters::SubsurfaceRadiusScale,
+					OpenPBRSurface::Parameters::SubsurfaceAnisotropy,
+					OpenPBRSurface::Parameters::FuzzWeight,
+					OpenPBRSurface::Parameters::FuzzColor,
+					OpenPBRSurface::Parameters::FuzzRoughness,
+					OpenPBRSurface::Parameters::CoatWeight,
+					OpenPBRSurface::Parameters::CoatColor,
+					OpenPBRSurface::Parameters::CoatRoughness,
+					OpenPBRSurface::Parameters::CoatAnisotropy,
+					OpenPBRSurface::Parameters::CoatRotation,
+					OpenPBRSurface::Parameters::CoatIOR,
+					OpenPBRSurface::Parameters::CoatIORLevel,
+					OpenPBRSurface::Parameters::GeometryCoatNormal,
+					OpenPBRSurface::Parameters::ThinFilmThickness,
+					OpenPBRSurface::Parameters::ThinFilmIOR,
+					OpenPBRSurface::Parameters::EmissionLuminance,
+					OpenPBRSurface::Parameters::EmissionColor,
+					OpenPBRSurface::Parameters::GeometryNormal,
+					OpenPBRSurface::Parameters::GeometryTangent,
+					OpenPBRSurface::Parameters::GeometryOpacity,
+					OpenPBRSurface::Parameters::GeometryThinWalled,
 				},
 				// OpenPBRSurface Outputs
 				!bIsSubstrateEnabled ?
 				TSet<FName>{
-					UE::Interchange::Materials::PBRMR::Parameters::BaseColor,
-						UE::Interchange::Materials::PBRMR::Parameters::Metallic,
-						UE::Interchange::Materials::PBRMR::Parameters::Specular,
-						UE::Interchange::Materials::PBRMR::Parameters::Roughness,
-						UE::Interchange::Materials::PBRMR::Parameters::Anisotropy,
-						UE::Interchange::Materials::PBRMR::Parameters::EmissiveColor,
-						UE::Interchange::Materials::PBRMR::Parameters::Opacity,
-						UE::Interchange::Materials::PBRMR::Parameters::Normal,
-						UE::Interchange::Materials::PBRMR::Parameters::Tangent,
-						UE::Interchange::Materials::Sheen::Parameters::SheenRoughness,
-						UE::Interchange::Materials::Sheen::Parameters::SheenColor,
-						UE::Interchange::Materials::Subsurface::Parameters::SubsurfaceColor,
-						UE::Interchange::Materials::ClearCoat::Parameters::ClearCoat,
-						UE::Interchange::Materials::ClearCoat::Parameters::ClearCoatRoughness,
-						UE::Interchange::Materials::ClearCoat::Parameters::ClearCoatNormal
+					PBRMR::Parameters::BaseColor,
+					PBRMR::Parameters::Metallic,
+					PBRMR::Parameters::Specular,
+					PBRMR::Parameters::Roughness,
+					PBRMR::Parameters::Anisotropy,
+					PBRMR::Parameters::EmissiveColor,
+					PBRMR::Parameters::Opacity,
+					PBRMR::Parameters::Normal,
+					PBRMR::Parameters::Tangent,
+					Sheen::Parameters::SheenRoughness,
+					Sheen::Parameters::SheenColor,
+					Subsurface::Parameters::SubsurfaceColor,
+					ClearCoat::Parameters::ClearCoat,
+					ClearCoat::Parameters::ClearCoatRoughness,
+					ClearCoat::Parameters::ClearCoatNormal
 					}	:
-					TSet<FName>{
-						UE::Interchange::Materials::OpenPBRSurface::SubstrateMaterial::Outputs::FrontMaterial,
-						UE::Interchange::Materials::OpenPBRSurface::SubstrateMaterial::Outputs::OpacityMask
+				TSet<FName>{
+					OpenPBRSurface::SubstrateMaterial::Outputs::FrontMaterial,
+					OpenPBRSurface::SubstrateMaterial::Outputs::OpacityMask
 					}
 				}
 			},
@@ -136,62 +138,62 @@ UMaterialXPipelineSettings::UMaterialXPipelineSettings()
 				{
 					// OpenPBRSurfaceTransmission Inputs
 					TSet<FName>{
-						UE::Interchange::Materials::OpenPBRSurface::Parameters::BaseWeight,
-						UE::Interchange::Materials::OpenPBRSurface::Parameters::BaseColor,
-						UE::Interchange::Materials::OpenPBRSurface::Parameters::BaseRoughness,
-						UE::Interchange::Materials::OpenPBRSurface::Parameters::BaseMetalness,
-						UE::Interchange::Materials::OpenPBRSurface::Parameters::SpecularWeight,
-						UE::Interchange::Materials::OpenPBRSurface::Parameters::SpecularColor,
-						UE::Interchange::Materials::OpenPBRSurface::Parameters::SpecularRoughness,
-						UE::Interchange::Materials::OpenPBRSurface::Parameters::SpecularIOR,
-						UE::Interchange::Materials::OpenPBRSurface::Parameters::SpecularIORLevel,
-						UE::Interchange::Materials::OpenPBRSurface::Parameters::SpecularAnisotropy,
-						UE::Interchange::Materials::OpenPBRSurface::Parameters::SpecularRotation,
-						UE::Interchange::Materials::OpenPBRSurface::Parameters::TransmissionWeight,
-						UE::Interchange::Materials::OpenPBRSurface::Parameters::TransmissionColor,
-						UE::Interchange::Materials::OpenPBRSurface::Parameters::TransmissionDepth,
-						UE::Interchange::Materials::OpenPBRSurface::Parameters::TransmissionDispersionScale,
-						UE::Interchange::Materials::OpenPBRSurface::Parameters::TransmissionDispersionAbbeNumber,
-						UE::Interchange::Materials::OpenPBRSurface::Parameters::TransmissionScatter,
-						UE::Interchange::Materials::OpenPBRSurface::Parameters::TransmissionScatterAnisotropy,
-						UE::Interchange::Materials::OpenPBRSurface::Parameters::FuzzWeight,
-						UE::Interchange::Materials::OpenPBRSurface::Parameters::FuzzColor,
-						UE::Interchange::Materials::OpenPBRSurface::Parameters::FuzzRoughness,
-						UE::Interchange::Materials::OpenPBRSurface::Parameters::CoatWeight,
-						UE::Interchange::Materials::OpenPBRSurface::Parameters::CoatColor,
-						UE::Interchange::Materials::OpenPBRSurface::Parameters::CoatRoughness,
-						UE::Interchange::Materials::OpenPBRSurface::Parameters::CoatAnisotropy,
-						UE::Interchange::Materials::OpenPBRSurface::Parameters::CoatRotation,
-						UE::Interchange::Materials::OpenPBRSurface::Parameters::CoatIOR,
-						UE::Interchange::Materials::OpenPBRSurface::Parameters::CoatIORLevel,
-						UE::Interchange::Materials::OpenPBRSurface::Parameters::GeometryCoatNormal,
-						UE::Interchange::Materials::OpenPBRSurface::Parameters::ThinFilmThickness,
-						UE::Interchange::Materials::OpenPBRSurface::Parameters::ThinFilmIOR,
-						UE::Interchange::Materials::OpenPBRSurface::Parameters::EmissionLuminance,
-						UE::Interchange::Materials::OpenPBRSurface::Parameters::EmissionColor,
-						UE::Interchange::Materials::OpenPBRSurface::Parameters::GeometryNormal,
-						UE::Interchange::Materials::OpenPBRSurface::Parameters::GeometryTangent,
-						UE::Interchange::Materials::OpenPBRSurface::Parameters::GeometryOpacity,
-						UE::Interchange::Materials::OpenPBRSurface::Parameters::GeometryThinWalled,
+						OpenPBRSurface::Parameters::BaseWeight,
+						OpenPBRSurface::Parameters::BaseColor,
+						OpenPBRSurface::Parameters::BaseRoughness,
+						OpenPBRSurface::Parameters::BaseMetalness,
+						OpenPBRSurface::Parameters::SpecularWeight,
+						OpenPBRSurface::Parameters::SpecularColor,
+						OpenPBRSurface::Parameters::SpecularRoughness,
+						OpenPBRSurface::Parameters::SpecularIOR,
+						OpenPBRSurface::Parameters::SpecularIORLevel,
+						OpenPBRSurface::Parameters::SpecularAnisotropy,
+						OpenPBRSurface::Parameters::SpecularRotation,
+						OpenPBRSurface::Parameters::TransmissionWeight,
+						OpenPBRSurface::Parameters::TransmissionColor,
+						OpenPBRSurface::Parameters::TransmissionDepth,
+						OpenPBRSurface::Parameters::TransmissionDispersionScale,
+						OpenPBRSurface::Parameters::TransmissionDispersionAbbeNumber,
+						OpenPBRSurface::Parameters::TransmissionScatter,
+						OpenPBRSurface::Parameters::TransmissionScatterAnisotropy,
+						OpenPBRSurface::Parameters::FuzzWeight,
+						OpenPBRSurface::Parameters::FuzzColor,
+						OpenPBRSurface::Parameters::FuzzRoughness,
+						OpenPBRSurface::Parameters::CoatWeight,
+						OpenPBRSurface::Parameters::CoatColor,
+						OpenPBRSurface::Parameters::CoatRoughness,
+						OpenPBRSurface::Parameters::CoatAnisotropy,
+						OpenPBRSurface::Parameters::CoatRotation,
+						OpenPBRSurface::Parameters::CoatIOR,
+						OpenPBRSurface::Parameters::CoatIORLevel,
+						OpenPBRSurface::Parameters::GeometryCoatNormal,
+						OpenPBRSurface::Parameters::ThinFilmThickness,
+						OpenPBRSurface::Parameters::ThinFilmIOR,
+						OpenPBRSurface::Parameters::EmissionLuminance,
+						OpenPBRSurface::Parameters::EmissionColor,
+						OpenPBRSurface::Parameters::GeometryNormal,
+						OpenPBRSurface::Parameters::GeometryTangent,
+						OpenPBRSurface::Parameters::GeometryOpacity,
+						OpenPBRSurface::Parameters::GeometryThinWalled,
 					},
 					// OpenPBRSurfaceTransmission Outputs
 					!bIsSubstrateEnabled ?
 					TSet<FName>{
-						UE::Interchange::Materials::PBRMR::Parameters::BaseColor,
-						UE::Interchange::Materials::PBRMR::Parameters::Metallic,
-						UE::Interchange::Materials::PBRMR::Parameters::Specular,
-						UE::Interchange::Materials::PBRMR::Parameters::Roughness,
-						UE::Interchange::Materials::PBRMR::Parameters::Anisotropy,
-						UE::Interchange::Materials::PBRMR::Parameters::EmissiveColor,
-						UE::Interchange::Materials::PBRMR::Parameters::Opacity,
-						UE::Interchange::Materials::PBRMR::Parameters::Normal,
-						UE::Interchange::Materials::PBRMR::Parameters::Tangent,
-						UE::Interchange::Materials::PBRMR::Parameters::Refraction,
-						UE::Interchange::Materials::ThinTranslucent::Parameters::TransmissionColor
+						PBRMR::Parameters::BaseColor,
+						PBRMR::Parameters::Metallic,
+						PBRMR::Parameters::Specular,
+						PBRMR::Parameters::Roughness,
+						PBRMR::Parameters::Anisotropy,
+						PBRMR::Parameters::EmissiveColor,
+						PBRMR::Parameters::Opacity,
+						PBRMR::Parameters::Normal,
+						PBRMR::Parameters::Tangent,
+						PBRMR::Parameters::Refraction,
+						ThinTranslucent::Parameters::TransmissionColor
 					}	:
 					TSet<FName>{
-						UE::Interchange::Materials::OpenPBRSurface::SubstrateMaterial::Outputs::FrontMaterial,
-						UE::Interchange::Materials::OpenPBRSurface::SubstrateMaterial::Outputs::OpacityMask
+						OpenPBRSurface::SubstrateMaterial::Outputs::FrontMaterial,
+						OpenPBRSurface::SubstrateMaterial::Outputs::OpacityMask
 					}
 				}
 			},
@@ -200,63 +202,63 @@ UMaterialXPipelineSettings::UMaterialXPipelineSettings()
 				{
 					// StandardSurface Inputs
 					TSet<FName>{
-						UE::Interchange::Materials::StandardSurface::Parameters::Base,
-						UE::Interchange::Materials::StandardSurface::Parameters::BaseColor,
-						UE::Interchange::Materials::StandardSurface::Parameters::DiffuseRoughness,
-						UE::Interchange::Materials::StandardSurface::Parameters::Metalness,
-						UE::Interchange::Materials::StandardSurface::Parameters::Specular,
-						UE::Interchange::Materials::StandardSurface::Parameters::SpecularRoughness,
-						UE::Interchange::Materials::StandardSurface::Parameters::SpecularIOR,
-						UE::Interchange::Materials::StandardSurface::Parameters::SpecularAnisotropy,
-						UE::Interchange::Materials::StandardSurface::Parameters::SpecularRotation,
-						UE::Interchange::Materials::StandardSurface::Parameters::Subsurface,
-						UE::Interchange::Materials::StandardSurface::Parameters::SubsurfaceColor,
-						UE::Interchange::Materials::StandardSurface::Parameters::SubsurfaceRadius,
-						UE::Interchange::Materials::StandardSurface::Parameters::SubsurfaceScale,
-						UE::Interchange::Materials::StandardSurface::Parameters::Sheen,
-						UE::Interchange::Materials::StandardSurface::Parameters::SheenColor,
-						UE::Interchange::Materials::StandardSurface::Parameters::SheenRoughness,
-						UE::Interchange::Materials::StandardSurface::Parameters::Coat,
-						UE::Interchange::Materials::StandardSurface::Parameters::CoatColor,
-						UE::Interchange::Materials::StandardSurface::Parameters::CoatRoughness,
-						UE::Interchange::Materials::StandardSurface::Parameters::CoatNormal,
-						UE::Interchange::Materials::StandardSurface::Parameters::ThinFilmThickness,
-						UE::Interchange::Materials::StandardSurface::Parameters::ThinFilmIOR,
-						UE::Interchange::Materials::StandardSurface::Parameters::Emission,
-						UE::Interchange::Materials::StandardSurface::Parameters::EmissionColor,
-						UE::Interchange::Materials::StandardSurface::Parameters::Normal,
-						UE::Interchange::Materials::StandardSurface::Parameters::Tangent,
-						UE::Interchange::Materials::StandardSurface::Parameters::ThinWalled,
-						bIsSubstrateEnabled ? UE::Interchange::Materials::StandardSurface::Parameters::SpecularColor : FName{},
-						UE::Interchange::Materials::StandardSurface::Parameters::CoatIOR,
-						bIsSubstrateEnabled ? UE::Interchange::Materials::StandardSurface::Parameters::CoatAnisotropy : FName{},
-						bIsSubstrateEnabled ? UE::Interchange::Materials::StandardSurface::Parameters::CoatRotation : FName{},
-						bIsSubstrateEnabled ? UE::Interchange::Materials::StandardSurface::Parameters::ThinFilmIOR : FName{},
-						bIsSubstrateEnabled ? UE::Interchange::Materials::StandardSurface::Parameters::Opacity : FName{},
-						bIsSubstrateEnabled ? UE::Interchange::Materials::StandardSurface::Parameters::SubsurfaceAnisotropy : FName{}
+						StandardSurface::Parameters::Base,
+						StandardSurface::Parameters::BaseColor,
+						StandardSurface::Parameters::DiffuseRoughness,
+						StandardSurface::Parameters::Metalness,
+						StandardSurface::Parameters::Specular,
+						StandardSurface::Parameters::SpecularRoughness,
+						StandardSurface::Parameters::SpecularIOR,
+						StandardSurface::Parameters::SpecularAnisotropy,
+						StandardSurface::Parameters::SpecularRotation,
+						StandardSurface::Parameters::Subsurface,
+						StandardSurface::Parameters::SubsurfaceColor,
+						StandardSurface::Parameters::SubsurfaceRadius,
+						StandardSurface::Parameters::SubsurfaceScale,
+						StandardSurface::Parameters::Sheen,
+						StandardSurface::Parameters::SheenColor,
+						StandardSurface::Parameters::SheenRoughness,
+						StandardSurface::Parameters::Coat,
+						StandardSurface::Parameters::CoatColor,
+						StandardSurface::Parameters::CoatRoughness,
+						StandardSurface::Parameters::CoatNormal,
+						StandardSurface::Parameters::ThinFilmThickness,
+						StandardSurface::Parameters::ThinFilmIOR,
+						StandardSurface::Parameters::Emission,
+						StandardSurface::Parameters::EmissionColor,
+						StandardSurface::Parameters::Normal,
+						StandardSurface::Parameters::Tangent,
+						StandardSurface::Parameters::ThinWalled,
+						bIsSubstrateEnabled ? StandardSurface::Parameters::SpecularColor : FName{},
+						StandardSurface::Parameters::CoatIOR,
+						bIsSubstrateEnabled ? StandardSurface::Parameters::CoatAnisotropy : FName{},
+						bIsSubstrateEnabled ? StandardSurface::Parameters::CoatRotation : FName{},
+						bIsSubstrateEnabled ? StandardSurface::Parameters::ThinFilmIOR : FName{},
+						bIsSubstrateEnabled ? StandardSurface::Parameters::Opacity : FName{},
+						bIsSubstrateEnabled ? StandardSurface::Parameters::SubsurfaceAnisotropy : FName{}
 					},
 					// StandardSurface Outputs
 					!bIsSubstrateEnabled ?
 					TSet<FName>{
 						TEXT("Base Color"), // MX_StandardSurface has BaseColor with a whitespace, this should be fixed in further release
-						UE::Interchange::Materials::PBRMR::Parameters::Metallic,
-						UE::Interchange::Materials::PBRMR::Parameters::Specular,
-						UE::Interchange::Materials::PBRMR::Parameters::Roughness,
-						UE::Interchange::Materials::PBRMR::Parameters::Anisotropy,
-						UE::Interchange::Materials::PBRMR::Parameters::EmissiveColor,
-						UE::Interchange::Materials::PBRMR::Parameters::Opacity,
-						UE::Interchange::Materials::PBRMR::Parameters::Normal,
-						UE::Interchange::Materials::PBRMR::Parameters::Tangent,
-						UE::Interchange::Materials::Sheen::Parameters::SheenRoughness,
-						UE::Interchange::Materials::Sheen::Parameters::SheenColor,
-						UE::Interchange::Materials::Subsurface::Parameters::SubsurfaceColor,
-						UE::Interchange::Materials::ClearCoat::Parameters::ClearCoat,
-						UE::Interchange::Materials::ClearCoat::Parameters::ClearCoatRoughness,
-						UE::Interchange::Materials::ClearCoat::Parameters::ClearCoatNormal
+						PBRMR::Parameters::Metallic,
+						PBRMR::Parameters::Specular,
+						PBRMR::Parameters::Roughness,
+						PBRMR::Parameters::Anisotropy,
+						PBRMR::Parameters::EmissiveColor,
+						PBRMR::Parameters::Opacity,
+						PBRMR::Parameters::Normal,
+						PBRMR::Parameters::Tangent,
+						Sheen::Parameters::SheenRoughness,
+						Sheen::Parameters::SheenColor,
+						Subsurface::Parameters::SubsurfaceColor,
+						ClearCoat::Parameters::ClearCoat,
+						ClearCoat::Parameters::ClearCoatRoughness,
+						ClearCoat::Parameters::ClearCoatNormal
 					}	:
 					TSet<FName>{
-						UE::Interchange::Materials::StandardSurface::SubstrateMaterial::Outputs::Opaque,
-						UE::Interchange::Materials::StandardSurface::SubstrateMaterial::Outputs::Opacity
+						StandardSurface::SubstrateMaterial::Outputs::Opaque,
+						StandardSurface::SubstrateMaterial::Outputs::Opacity
 					}
 				}
 			},
@@ -265,62 +267,62 @@ UMaterialXPipelineSettings::UMaterialXPipelineSettings()
 				{
 					// StandardSurfaceTransmission Inputs
 					TSet<FName>{
-						UE::Interchange::Materials::StandardSurface::Parameters::Base,
-						UE::Interchange::Materials::StandardSurface::Parameters::BaseColor,
-						UE::Interchange::Materials::StandardSurface::Parameters::DiffuseRoughness,
-						UE::Interchange::Materials::StandardSurface::Parameters::Metalness,
-						UE::Interchange::Materials::StandardSurface::Parameters::Specular,
-						bIsSubstrateEnabled ? UE::Interchange::Materials::StandardSurface::Parameters::SpecularColor : FName{},
-						UE::Interchange::Materials::StandardSurface::Parameters::SpecularRoughness,
-						UE::Interchange::Materials::StandardSurface::Parameters::SpecularIOR,
-						UE::Interchange::Materials::StandardSurface::Parameters::SpecularAnisotropy,
-						UE::Interchange::Materials::StandardSurface::Parameters::SpecularRotation,
-						UE::Interchange::Materials::StandardSurface::Parameters::Transmission,
-						UE::Interchange::Materials::StandardSurface::Parameters::TransmissionColor,
-						UE::Interchange::Materials::StandardSurface::Parameters::TransmissionDepth,
-						UE::Interchange::Materials::StandardSurface::Parameters::TransmissionScatter,
-						UE::Interchange::Materials::StandardSurface::Parameters::TransmissionScatterAnisotropy,
-						UE::Interchange::Materials::StandardSurface::Parameters::TransmissionDispersion,
-						UE::Interchange::Materials::StandardSurface::Parameters::TransmissionExtraRoughness,
-						!bIsSubstrateEnabled ? UE::Interchange::Materials::StandardSurface::Parameters::Subsurface : FName{},
-						!bIsSubstrateEnabled ? UE::Interchange::Materials::StandardSurface::Parameters::SubsurfaceColor : FName{},
-						!bIsSubstrateEnabled ? UE::Interchange::Materials::StandardSurface::Parameters::SubsurfaceRadius : FName{},
-						!bIsSubstrateEnabled ? UE::Interchange::Materials::StandardSurface::Parameters::SubsurfaceScale : FName{},
-						UE::Interchange::Materials::StandardSurface::Parameters::Sheen,
-						UE::Interchange::Materials::StandardSurface::Parameters::SheenColor,
-						UE::Interchange::Materials::StandardSurface::Parameters::SheenRoughness,
-						UE::Interchange::Materials::StandardSurface::Parameters::Coat,
-						UE::Interchange::Materials::StandardSurface::Parameters::CoatColor,
-						UE::Interchange::Materials::StandardSurface::Parameters::CoatRoughness,
-						bIsSubstrateEnabled ? UE::Interchange::Materials::StandardSurface::Parameters::CoatAnisotropy : FName{},
-						bIsSubstrateEnabled ? UE::Interchange::Materials::StandardSurface::Parameters::CoatRotation : FName{},
-						UE::Interchange::Materials::StandardSurface::Parameters::CoatIOR,
-						UE::Interchange::Materials::StandardSurface::Parameters::CoatNormal,
-						UE::Interchange::Materials::StandardSurface::Parameters::ThinFilmThickness,
-						UE::Interchange::Materials::StandardSurface::Parameters::ThinFilmIOR,
-						bIsSubstrateEnabled ? UE::Interchange::Materials::StandardSurface::Parameters::ThinFilmIOR : FName{},
-						UE::Interchange::Materials::StandardSurface::Parameters::Emission,
-						UE::Interchange::Materials::StandardSurface::Parameters::EmissionColor,
-						bIsSubstrateEnabled ? UE::Interchange::Materials::StandardSurface::Parameters::Opacity : FName{},
-						UE::Interchange::Materials::StandardSurface::Parameters::Normal,
-						UE::Interchange::Materials::StandardSurface::Parameters::Tangent
+						StandardSurface::Parameters::Base,
+						StandardSurface::Parameters::BaseColor,
+						StandardSurface::Parameters::DiffuseRoughness,
+						StandardSurface::Parameters::Metalness,
+						StandardSurface::Parameters::Specular,
+						bIsSubstrateEnabled ? StandardSurface::Parameters::SpecularColor : FName{},
+						StandardSurface::Parameters::SpecularRoughness,
+						StandardSurface::Parameters::SpecularIOR,
+						StandardSurface::Parameters::SpecularAnisotropy,
+						StandardSurface::Parameters::SpecularRotation,
+						StandardSurface::Parameters::Transmission,
+						StandardSurface::Parameters::TransmissionColor,
+						StandardSurface::Parameters::TransmissionDepth,
+						StandardSurface::Parameters::TransmissionScatter,
+						StandardSurface::Parameters::TransmissionScatterAnisotropy,
+						StandardSurface::Parameters::TransmissionDispersion,
+						StandardSurface::Parameters::TransmissionExtraRoughness,
+						!bIsSubstrateEnabled ? StandardSurface::Parameters::Subsurface : FName{},
+						!bIsSubstrateEnabled ? StandardSurface::Parameters::SubsurfaceColor : FName{},
+						!bIsSubstrateEnabled ? StandardSurface::Parameters::SubsurfaceRadius : FName{},
+						!bIsSubstrateEnabled ? StandardSurface::Parameters::SubsurfaceScale : FName{},
+						StandardSurface::Parameters::Sheen,
+						StandardSurface::Parameters::SheenColor,
+						StandardSurface::Parameters::SheenRoughness,
+						StandardSurface::Parameters::Coat,
+						StandardSurface::Parameters::CoatColor,
+						StandardSurface::Parameters::CoatRoughness,
+						bIsSubstrateEnabled ? StandardSurface::Parameters::CoatAnisotropy : FName{},
+						bIsSubstrateEnabled ? StandardSurface::Parameters::CoatRotation : FName{},
+						StandardSurface::Parameters::CoatIOR,
+						StandardSurface::Parameters::CoatNormal,
+						StandardSurface::Parameters::ThinFilmThickness,
+						StandardSurface::Parameters::ThinFilmIOR,
+						bIsSubstrateEnabled ? StandardSurface::Parameters::ThinFilmIOR : FName{},
+						StandardSurface::Parameters::Emission,
+						StandardSurface::Parameters::EmissionColor,
+						bIsSubstrateEnabled ? StandardSurface::Parameters::Opacity : FName{},
+						StandardSurface::Parameters::Normal,
+						StandardSurface::Parameters::Tangent
 					},
 					// StandardSurfaceTransmission Outputs
 					!bIsSubstrateEnabled ?
 					TSet<FName>{
-						UE::Interchange::Materials::PBRMR::Parameters::BaseColor,
-						UE::Interchange::Materials::PBRMR::Parameters::Metallic,
-						UE::Interchange::Materials::PBRMR::Parameters::Specular,
-						UE::Interchange::Materials::PBRMR::Parameters::Roughness,
-						UE::Interchange::Materials::PBRMR::Parameters::Anisotropy,
-						UE::Interchange::Materials::PBRMR::Parameters::EmissiveColor,
-						UE::Interchange::Materials::PBRMR::Parameters::Opacity,
-						UE::Interchange::Materials::PBRMR::Parameters::Normal,
-						UE::Interchange::Materials::PBRMR::Parameters::Tangent,
-						UE::Interchange::Materials::PBRMR::Parameters::Refraction,
-						UE::Interchange::Materials::ThinTranslucent::Parameters::TransmissionColor
+						PBRMR::Parameters::BaseColor,
+						PBRMR::Parameters::Metallic,
+						PBRMR::Parameters::Specular,
+						PBRMR::Parameters::Roughness,
+						PBRMR::Parameters::Anisotropy,
+						PBRMR::Parameters::EmissiveColor,
+						PBRMR::Parameters::Opacity,
+						PBRMR::Parameters::Normal,
+						PBRMR::Parameters::Tangent,
+						PBRMR::Parameters::Refraction,
+						ThinTranslucent::Parameters::TransmissionColor
 					}	:
-					TSet<FName>{UE::Interchange::Materials::StandardSurface::SubstrateMaterial::Outputs::Translucent}
+					TSet<FName>{StandardSurface::SubstrateMaterial::Outputs::Translucent}
 				}
 			},
 			{
@@ -328,22 +330,22 @@ UMaterialXPipelineSettings::UMaterialXPipelineSettings()
 				{
 					// SurfaceUnlit Inputs
 					TSet<FName>{
-						UE::Interchange::Materials::SurfaceUnlit::Parameters::Emission,
-						UE::Interchange::Materials::SurfaceUnlit::Parameters::EmissionColor,
-						UE::Interchange::Materials::SurfaceUnlit::Parameters::Transmission,
-						UE::Interchange::Materials::SurfaceUnlit::Parameters::TransmissionColor,
-						UE::Interchange::Materials::SurfaceUnlit::Parameters::Opacity
+						SurfaceUnlit::Parameters::Emission,
+						SurfaceUnlit::Parameters::EmissionColor,
+						SurfaceUnlit::Parameters::Transmission,
+						SurfaceUnlit::Parameters::TransmissionColor,
+						SurfaceUnlit::Parameters::Opacity
 					},
 					// SurfaceUnlit Outputs
 					!bIsSubstrateEnabled ?
 					TSet<FName>{
-						UE::Interchange::Materials::Common::Parameters::EmissiveColor,
-						UE::Interchange::Materials::Common::Parameters::Opacity,
-						UE::Interchange::Materials::SurfaceUnlit::Outputs::OpacityMask
+						Common::Parameters::EmissiveColor,
+						Common::Parameters::Opacity,
+						SurfaceUnlit::Outputs::OpacityMask
 						} :
 					TSet<FName>{
-						UE::Interchange::Materials::SurfaceUnlit::Substrate::Outputs::OpacityMask,
-						UE::Interchange::Materials::SurfaceUnlit::Substrate::Outputs::SurfaceUnlit
+						SurfaceUnlit::Substrate::Outputs::OpacityMask,
+						SurfaceUnlit::Substrate::Outputs::SurfaceUnlit
 					}
 				}
 			},
@@ -352,36 +354,39 @@ UMaterialXPipelineSettings::UMaterialXPipelineSettings()
 				{
 					// UsdPreviewSurface Inputs
 					TSet<FName>{
-						UE::Interchange::Materials::UsdPreviewSurface::Parameters::DiffuseColor,
-						UE::Interchange::Materials::UsdPreviewSurface::Parameters::EmissiveColor,
-						UE::Interchange::Materials::UsdPreviewSurface::Parameters::SpecularColor,
-						UE::Interchange::Materials::UsdPreviewSurface::Parameters::Metallic,
-						UE::Interchange::Materials::UsdPreviewSurface::Parameters::Roughness,
-						UE::Interchange::Materials::UsdPreviewSurface::Parameters::Clearcoat,
-						UE::Interchange::Materials::UsdPreviewSurface::Parameters::ClearcoatRoughness,
-						UE::Interchange::Materials::UsdPreviewSurface::Parameters::Opacity,
-						UE::Interchange::Materials::UsdPreviewSurface::Parameters::OpacityThreshold,
-						UE::Interchange::Materials::UsdPreviewSurface::Parameters::IOR,
-						UE::Interchange::Materials::UsdPreviewSurface::Parameters::Normal,
-						UE::Interchange::Materials::UsdPreviewSurface::Parameters::Displacement,
-						UE::Interchange::Materials::UsdPreviewSurface::Parameters::Occlusion
+						UsdPreviewSurface::Parameters::DiffuseColor,
+						UsdPreviewSurface::Parameters::EmissiveColor,
+						UsdPreviewSurface::Parameters::SpecularColor,
+						UsdPreviewSurface::Parameters::Metallic,
+						UsdPreviewSurface::Parameters::Roughness,
+						UsdPreviewSurface::Parameters::Clearcoat,
+						UsdPreviewSurface::Parameters::ClearcoatRoughness,
+						UsdPreviewSurface::Parameters::Opacity,
+						UsdPreviewSurface::Parameters::OpacityThreshold,
+						UsdPreviewSurface::Parameters::IOR,
+						UsdPreviewSurface::Parameters::Normal,
+						UsdPreviewSurface::Parameters::Displacement,
+						UsdPreviewSurface::Parameters::Occlusion
 					},
 					// UsdPreviewSurface Outputs
 					!bIsSubstrateEnabled ?
 					TSet<FName>{
-						UE::Interchange::Materials::PBRMR::Parameters::BaseColor,
-						UE::Interchange::Materials::PBRMR::Parameters::Metallic,
-						UE::Interchange::Materials::PBRMR::Parameters::Specular,
-						UE::Interchange::Materials::PBRMR::Parameters::Roughness,
-						UE::Interchange::Materials::PBRMR::Parameters::EmissiveColor,
-						UE::Interchange::Materials::PBRMR::Parameters::Opacity,
-						UE::Interchange::Materials::PBRMR::Parameters::Normal,
-						UE::Interchange::Materials::Common::Parameters::Refraction,
-						UE::Interchange::Materials::Common::Parameters::Occlusion,
-						UE::Interchange::Materials::ClearCoat::Parameters::ClearCoat,
-						UE::Interchange::Materials::ClearCoat::Parameters::ClearCoatRoughness,
+						PBRMR::Parameters::BaseColor,
+						PBRMR::Parameters::Metallic,
+						PBRMR::Parameters::Specular,
+						PBRMR::Parameters::Roughness,
+						PBRMR::Parameters::EmissiveColor,
+						PBRMR::Parameters::Opacity,
+						PBRMR::Parameters::Normal,
+						Common::Parameters::Refraction,
+						Common::Parameters::Occlusion,
+						ClearCoat::Parameters::ClearCoat,
+						ClearCoat::Parameters::ClearCoatRoughness,
 					} :
-					TSet<FName>{TEXT("Substrate UsdPreviewSurface"), TEXT("Opacity")}
+					TSet<FName>{
+						UsdPreviewSurface::SubstrateMaterial::Outputs::FrontMaterial,
+						UsdPreviewSurface::SubstrateMaterial::Outputs::Displacement,
+						UsdPreviewSurface::SubstrateMaterial::Outputs::Occlusion}
 				}
 			},
 			{
@@ -389,14 +394,14 @@ UMaterialXPipelineSettings::UMaterialXPipelineSettings()
 				{
 					// Surface Inputs
 					TSet<FName>{
-						UE::Interchange::Materials::Surface::Parameters::BSDF,
-						UE::Interchange::Materials::Surface::Parameters::EDF,
-						UE::Interchange::Materials::Surface::Parameters::Opacity
+						Surface::Parameters::BSDF,
+						Surface::Parameters::EDF,
+						Surface::Parameters::Opacity
 					},
 					// Surface Outputs
 					bIsSubstrateEnabled ?
-					TSet<FName>{UE::Interchange::Materials::Surface::Outputs::Surface} :
-					TSet<FName>{UE::Interchange::Materials::Surface::Substrate::Outputs::Surface, UE::Interchange::Materials::Surface::Substrate::Outputs::Opacity}
+					TSet<FName>{Surface::Outputs::Surface} :
+					TSet<FName>{Surface::Substrate::Outputs::Surface,Surface::Substrate::Outputs::Opacity}
 				}
 			},
 			// BSDF
