USD: Translate skeletal animation and morph target animations via Interchange.

diff --git a/Engine/Plugins/Importers/USDImporter/Source/USDSchemas/Private/USDSkelSkeletonTranslator.cpp b/Engine/Plugins/Importers/USDImporter/Source/USDSchemas/Private/USDSkelSkeletonTranslator.cpp
--- a/Engine/Plugins/Importers/USDImporter/Source/USDSchemas/Private/USDSkelSkeletonTranslator.cpp
+++ b/Engine/Plugins/Importers/USDImporter/Source/USDSchemas/Private/USDSkelSkeletonTranslator.cpp
@@ -425,49 +425,6 @@ namespace UsdSkelSkeletonTranslatorImpl
 		return OutHash;
 	}
 
-	void SetMorphTargetWeight(USkeletalMeshComponent& SkeletalMeshComponent, const FString& MorphTargetName, float Weight)
-	{
-		USkeletalMesh* SkeletalMesh = SkeletalMeshComponent.GetSkeletalMeshAsset();
-
-		// We try keeping a perfect correspondence between SkeletalMesh->GetMorphTargets() and SkeletalMeshComponent.ActiveMorphTargets
-		int32 IndexInSkeletalMesh = INDEX_NONE;
-		SkeletalMeshComponent.GetSkeletalMeshAsset()->FindMorphTargetAndIndex(*MorphTargetName, IndexInSkeletalMesh);
-		if (IndexInSkeletalMesh == INDEX_NONE)
-		{
-			return;
-		}
-
-		const UMorphTarget* MorphTarget = SkeletalMeshComponent.GetSkeletalMeshAsset()->GetMorphTargets()[IndexInSkeletalMesh];
-		if (!MorphTarget)
-		{
-			return;
-		}
-
-		int32 WeightIndex = INDEX_NONE;
-		if (SkeletalMeshComponent.ActiveMorphTargets.Contains(MorphTarget))
-		{
-			WeightIndex = SkeletalMeshComponent.ActiveMorphTargets[MorphTarget];
-		}
-
-		// Morph target is not at expected location (i.e. after CreateComponents, duplicate for PIE or undo/redo) --> Rebuild ActiveMorphTargets
-		// This may lead to one frame of glitchiness, as we'll reset all weights to zero...
-		if (WeightIndex == INDEX_NONE)
-		{
-			SkeletalMeshComponent.ActiveMorphTargets.Reset();
-			SkeletalMeshComponent.MorphTargetWeights.Reset();
-			TArray<TObjectPtr<UMorphTarget>>& MorphTargets = SkeletalMesh->GetMorphTargets();
-			for (int32 MorphTargetIndex = 0; MorphTargetIndex < MorphTargets.Num(); ++MorphTargetIndex)
-			{
-				SkeletalMeshComponent.ActiveMorphTargets.Add(MorphTargets[MorphTargetIndex], MorphTargetIndex);
-				SkeletalMeshComponent.MorphTargetWeights.Add(0.0f);	   // We'll update these right afterwards when we call UpdateComponents
-			}
-
-			WeightIndex = IndexInSkeletalMesh;
-		}
-
-		SkeletalMeshComponent.MorphTargetWeights[WeightIndex] = Weight;
-	}
-
 	bool LoadAllSkeletalData(
 		const pxr::UsdSkelBinding& InSkeletonBinding,
 		const pxr::UsdSkelCache& InSkelCache,
diff --git a/Engine/Plugins/Interchange/Runtime/Source/Import/Private/Usd/InterchangeUsdTranslator.cpp b/Engine/Plugins/Interchange/Runtime/Source/Import/Private/Usd/InterchangeUsdTranslator.cpp
--- a/Engine/Plugins/Interchange/Runtime/Source/Import/Private/Usd/InterchangeUsdTranslator.cpp
+++ b/Engine/Plugins/Interchange/Runtime/Source/Import/Private/Usd/InterchangeUsdTranslator.cpp
@@ -3,11 +3,22 @@
 #include "Usd/InterchangeUsdTranslator.h"
 
 #include "UnrealUSDWrapper.h"
+#include "USDConversionUtils.h"
+#include "USDGeomMeshConversion.h"
+#include "USDLightConversion.h"
+#include "USDLog.h"
+#include "USDPrimConversion.h"
+#include "USDShadeConversion.h"
+#include "USDSkeletalDataConversion.h"
+#include "USDStageOptions.h"
+#include "USDTypesConversion.h"
 #include "UsdWrappers/SdfPath.h"
 #include "UsdWrappers/UsdAttribute.h"
 #include "UsdWrappers/UsdGeomXformable.h"
 #include "UsdWrappers/UsdPrim.h"
 #include "UsdWrappers/UsdRelationship.h"
+#include "UsdWrappers/UsdSkelAnimQuery.h"
+#include "UsdWrappers/UsdSkelBinding.h"
 #include "UsdWrappers/UsdSkelBlendShape.h"
 #include "UsdWrappers/UsdSkelBlendShapeQuery.h"
 #include "UsdWrappers/UsdSkelCache.h"
@@ -16,18 +27,7 @@
 #include "UsdWrappers/UsdSkelSkinningQuery.h"
 #include "UsdWrappers/UsdStage.h"
 #include "UsdWrappers/UsdTyped.h"
-
-#include "USDClassesModule.h"
-#include "USDConversionUtils.h"
-#include "USDGeomMeshConversion.h"
-#include "USDLightConversion.h"
-#include "USDLog.h"
-#include "USDMaterialUtils.h"
-#include "USDPrimConversion.h"
-#include "USDShadeConversion.h"
-#include "USDSkeletalDataConversion.h"
-#include "USDStageOptions.h"
-#include "USDTypesConversion.h"
+#include "UsdWrappers/VtValue.h"
 
 #include "Async/Async.h"
 #include "HAL/IConsoleManager.h"
@@ -52,6 +52,7 @@
 #include "pxr/usd/usdLux/tokens.h"
 #include "pxr/usd/usdPhysics/tokens.h"
 #include "pxr/usd/usdShade/tokens.h"
+#include "pxr/usd/usdSkel/tokens.h"
 #include "USDIncludesEnd.h"
 #endif	  // USE_USD_SDK
 
@@ -82,6 +83,7 @@ namespace UE::InterchangeUsdTranslator::Private
 	const static FString MaterialPrefix = TEXT("\\Material\\");
 	const static FString MeshPrefix = TEXT("\\Mesh\\");
 	const static FString MorphTargetPrefix = TEXT("\\MorphTarget\\");
+	const static FString BonePrefix = TEXT("\\Bone\\");
 
 	// Information intended to be passed down from parent to children (by value) as we traverse the stage
 	struct FTraversalInfo
@@ -196,6 +198,46 @@ TArray<FString> UInterchangeUSDTranslator::GetSupportedFormats() const
 #if USE_USD_SDK
 namespace UE::InterchangeUsdTranslator::Private
 {
+	FString HashAnimPayloadQuery(const Interchange::FAnimationPayloadQuery& Query)
+	{
+		FSHAHash Hash;
+		FSHA1 SHA1;
+
+		// TODO: Is there a StringView alternative?
+		FString SkeletonPrimPath;
+		FString JointIndexStr;
+		bool bSplit = Query.PayloadKey.UniqueId.Split(TEXT("\\"), &SkeletonPrimPath, &JointIndexStr, ESearchCase::CaseSensitive, ESearchDir::FromEnd);
+		if (!bSplit)
+		{
+			return {};
+		}
+
+		SHA1.UpdateWithString(*SkeletonPrimPath, SkeletonPrimPath.Len());
+
+		SHA1.Update(reinterpret_cast<const uint8*>(&Query.TimeDescription.BakeFrequency), sizeof(Query.TimeDescription.BakeFrequency));
+		SHA1.Update(reinterpret_cast<const uint8*>(&Query.TimeDescription.RangeStartSecond), sizeof(Query.TimeDescription.RangeStartSecond));
+		SHA1.Update(reinterpret_cast<const uint8*>(&Query.TimeDescription.RangeStopSecond), sizeof(Query.TimeDescription.RangeStopSecond));
+
+		SHA1.Final();
+		SHA1.GetHash(&Hash.Hash[0]);
+		return Hash.ToString();
+	}
+
+	FString GetMorphTargetMeshNodeUid(const FString& MeshPrimPath, int32 MeshBlendShapeIndex, const FString& InbetweenName = FString{})
+	{
+		return FString::Printf(TEXT("%s%s\\%d\\%s"), *MorphTargetPrefix, *MeshPrimPath, MeshBlendShapeIndex, *InbetweenName);
+	}
+
+	FString GetMorphTargetMeshPayloadKey(const FString& MeshPrimPath, int32 MeshBlendShapeIndex, const FString& InbetweenName = FString{})
+	{
+		return FString::Printf(TEXT("%s\\%d\\%s"), *MeshPrimPath, MeshBlendShapeIndex, *InbetweenName);
+	}
+
+	FString GetMorphTargetCurvePayloadKey(const FString& SkeletonPrimPath, int32 SkelAnimChannelIndex, const FString& BlendShapePath)
+	{
+		return FString::Printf(TEXT("%s\\%d\\%s"), *SkeletonPrimPath, SkelAnimChannelIndex, *BlendShapePath);
+	}
+
 	FString EncodeTexturePayloadKey(const UsdToUnreal::FTextureParameterValue& Value)
 	{
 		// Encode the compression settings onto the payload key as we need to move that into the
@@ -676,19 +718,20 @@ namespace UE::InterchangeUsdTranslator::Private
 	}
 
 	void AddMorphTargetNodes(
-		const UE::FUsdPrim& Prim,
+		const UE::FUsdPrim& MeshPrim,
 		UInterchangeUSDTranslatorImpl& TranslatorImpl,
 		UInterchangeMeshNode& MeshNode,
-		UInterchangeBaseNodeContainer& NodeContainer
+		UInterchangeBaseNodeContainer& NodeContainer,
+		const FTraversalInfo& Info
 	)
 	{
-		UE::FUsdSkelBlendShapeQuery Query{Prim};
+		UE::FUsdSkelBlendShapeQuery Query{MeshPrim};
 		if (!Query)
 		{
 			return;
 		}
 
-		const FString MeshPrimPath = Prim.GetPrimPath().GetString();
+		const FString MeshPrimPath = MeshPrim.GetPrimPath().GetString();
 
 		TFunction<void(const FString&, int32, const FString&)> AddMorphTargetNode =
 			[&MeshNode, &MeshPrimPath, &NodeContainer](const FString& MorphTargetName, int32 BlendShapeIndex, const FString& InbetweenName)
@@ -699,20 +742,19 @@ namespace UE::InterchangeUsdTranslator::Private
 			//  - It could be possible for different meshes to share the same BlendShape (possibly?), so we really want a separate version of
 			//    a blend shape for each mesh that uses it.
 			//
-			// Despite of that though, we'll just use the blendshape (+inbetween) name for MorphTargetName (so not anything's full path),
-			// so that users can get different blendshapes across the model to combine into a single morph target. Interchange has
-			// an import option to let you control whether they become separate morph targets or not anyway
-			// ("Merge Morph Targets with Same Name")
-			FString IDString = FString::Printf(TEXT("%s\\%d\\%s"), *MeshPrimPath, BlendShapeIndex, *InbetweenName);
-			FString MorphTargetUid = MorphTargetPrefix + IDString;
+			// Despite of that though, we won't use the blendshape's full path as the morph target name, so that users can get different
+			// blendshapes across the model to combine into a single morph target. Interchange has an import option to let you control
+			// whether they become separate morph targets or not anyway ("Merge Morph Targets with Same Name")
+			const FString NodeUid = GetMorphTargetMeshNodeUid(MeshPrimPath, BlendShapeIndex, InbetweenName);
+			const FString PayloadKey = GetMorphTargetMeshPayloadKey(MeshPrimPath, BlendShapeIndex, InbetweenName);
 
 			UInterchangeMeshNode* MorphTargetMeshNode = NewObject<UInterchangeMeshNode>(&NodeContainer);
-			MorphTargetMeshNode->InitializeNode(MorphTargetUid, MorphTargetName, EInterchangeNodeContainerType::TranslatedAsset);
-			MorphTargetMeshNode->SetPayLoadKey(IDString, EInterchangeMeshPayLoadType::MORPHTARGET);
+			MorphTargetMeshNode->InitializeNode(NodeUid, MorphTargetName, EInterchangeNodeContainerType::TranslatedAsset);
+			MorphTargetMeshNode->SetPayLoadKey(PayloadKey, EInterchangeMeshPayLoadType::MORPHTARGET);
 			MorphTargetMeshNode->SetMorphTarget(true);
 			MorphTargetMeshNode->SetMorphTargetName(MorphTargetName);
 			NodeContainer.AddNode(MorphTargetMeshNode);
-			MeshNode.SetMorphTargetDependencyUid(MorphTargetUid);
+			MeshNode.SetMorphTargetDependencyUid(NodeUid);
 		};
 
 		for (size_t Index = 0; Index < Query.GetNumBlendShapes(); ++Index)
@@ -722,17 +764,17 @@ namespace UE::InterchangeUsdTranslator::Private
 			{
 				continue;
 			}
-
 			UE::FUsdPrim BlendShapePrim = BlendShape.GetPrim();
-			FString MorphTargetName = BlendShapePrim.GetName().ToString();
-			AddMorphTargetNode(MorphTargetName, Index, FString{});
+			const FString BlendShapeName = BlendShapePrim.GetName().ToString();
 
-			TArray<UE::FUsdSkelInbetweenShape> Inbetweens = BlendShape.GetInbetweens();
-			for (const UE::FUsdSkelInbetweenShape& Inbetween : Inbetweens)
+			const FString UnusedInbetweenName;
+			AddMorphTargetNode(BlendShapeName, Index, UnusedInbetweenName);
+
+			for (const UE::FUsdSkelInbetweenShape& Inbetween : BlendShape.GetInbetweens())
 			{
-				FString InbetweenName = Inbetween.GetAttr().GetName().ToString();
-				FString InnerMorphTargetName = MorphTargetName + TEXT("_") + InbetweenName;
-				AddMorphTargetNode(InnerMorphTargetName, Index, InbetweenName);
+				const FString InbetweenName = Inbetween.GetAttr().GetName().ToString();
+				const FString MorphTargetName = BlendShapeName + TEXT("_") + InbetweenName;
+				AddMorphTargetNode(MorphTargetName, Index, InbetweenName);
 			}
 		}
 	}
@@ -768,7 +810,7 @@ namespace UE::InterchangeUsdTranslator::Private
 				MeshNode->SetSkeletonDependencyUid(Info.ActiveSkelQuery.GetSkeleton().GetPrimPath().GetString());
 			}
 
-			AddMorphTargetNodes(Prim, TranslatorImpl, *MeshNode, NodeContainer);
+			AddMorphTargetNodes(Prim, TranslatorImpl, *MeshNode, NodeContainer, Info);
 
 			// When returning the payload data later, we'll need at the very least our SkeletonQuery, so
 			// here we store the Info object into the Impl
@@ -839,12 +881,21 @@ namespace UE::InterchangeUsdTranslator::Private
 		NodeContainer.AddNode(MeshNode);
 	}
 
+	void AddSkeletalAnimationNode(
+		const UE::FUsdSkelSkeletonQuery& SkeletonQuery,
+		const TMap<FString, TPair<FString, int32>>& BoneToUidAndBoneIndex,
+		UInterchangeUSDTranslatorImpl& TranslatorImpl,
+		UInterchangeSceneNode& SkeletonPrimNode,
+		UInterchangeBaseNodeContainer& NodeContainer,
+		const FTraversalInfo& Info
+	);
+
 	void AddSkeletonNodes(
 		const UE::FUsdPrim& Prim,
 		UInterchangeUSDTranslatorImpl& TranslatorImpl,
 		UInterchangeSceneNode& SkeletonPrimNode,
 		UInterchangeBaseNodeContainer& NodeContainer,
-		FTraversalInfo Info
+		FTraversalInfo& Info
 	)
 	{
 		// If we're not inside of a SkelRoot, the skeleton shouldn't really do anything
@@ -865,6 +916,7 @@ namespace UE::InterchangeUsdTranslator::Private
 		SkeletonPrimNode.AddSpecializedType(UE::Interchange::FSceneNodeStaticData::GetJointSpecializeTypeString());
 		SkeletonPrimNode.SetCustomBindPoseLocalTransform(&NodeContainer, FTransform::Identity);
 		SkeletonPrimNode.SetCustomTimeZeroLocalTransform(&NodeContainer, FTransform::Identity);
+		const FString SkeletonPrimNodeUid = SkeletonPrimNode.GetUniqueID();
 
 #if WITH_EDITOR
 		// Convert the skeleton bones/joints into ConvertedData
@@ -878,44 +930,67 @@ namespace UE::InterchangeUsdTranslator::Private
 			return;
 		}
 
+		// Maps from the USD-style full bone name (e.g. "shoulder/elbow/hand") to the Uid we used for
+		// the corresponding scene node, and the bone's index on the skeleton's joint order.
+		// We'll need this to parse skeletal animations, if any
+		TMap<FString, TPair<FString, int32>> BoneToUidAndBoneIndex;
+
 		// Recursively traverse ConvertedData spawning the joint translated nodes
-		TFunction<void(const UsdToUnreal::FUsdSkeletonData::FBone&, UInterchangeSceneNode&, const FString&)> RecursiveTraverseBones = nullptr;
+		TFunction<void(int32, UInterchangeSceneNode&, const FString&)> RecursiveTraverseBones = nullptr;
 		RecursiveTraverseBones = [&RecursiveTraverseBones,
+								  &BoneToUidAndBoneIndex,
+								  &SkeletonPrimNodeUid,
 								  &ConvertedData,
 								  &NodeContainer	//
-		](const UsdToUnreal::FUsdSkeletonData::FBone& Bone, UInterchangeSceneNode& ParentNode, const FString& BonePath)
+		](int32 BoneIndex, UInterchangeSceneNode& ParentNode, const FString& BonePath)
 		{
-			// Concatenate a full "bone path" here for uniqueness, because Bone.Name is just the name of this
+			const UsdToUnreal::FUsdSkeletonData::FBone& Bone = ConvertedData.Bones[BoneIndex];
+
+			// Reconcatenate a full "bone path" here for uniqueness, because Bone.Name is just the name of this
 			// single bone/joint itself (e.g. "Elbow")
-			const FString BoneNodeUid = BonePath + TEXT("\\") + Bone.Name;
+			const FString ConcatBonePath = (BonePath.IsEmpty() ? TEXT("") : (BonePath + TEXT("/"))) + Bone.Name;
+
+			// Putting the BonePrefix here avoids the pathological case where the user has skeleton child prims
+			// with names that match the joint names
+			const FString BoneNodeUid = SkeletonPrimNodeUid + BonePrefix + ConcatBonePath;
 
 			UInterchangeSceneNode* BoneNode = NewObject<UInterchangeSceneNode>(&NodeContainer);
 			BoneNode->InitializeNode(BoneNodeUid, Bone.Name, EInterchangeNodeContainerType::TranslatedScene);
-			BoneNode->SetCustomLocalTransform(&NodeContainer, Bone.LocalRestTransform);
-
 			BoneNode->AddSpecializedType(UE::Interchange::FSceneNodeStaticData::GetJointSpecializeTypeString());
-			BoneNode->SetCustomBindPoseLocalTransform(&NodeContainer, Bone.LocalBindTransform);
+
+			// Note that we use our rest transforms for the Interchange bind pose as well: This because Interchange
+			// will put this on the RefSkeleton and so it will make its way to the Skeleton asset. We already kind
+			// of bake in our skeleton bind pose directly into our skinned mesh, so we really just want to put the
+			// rest pose on the skeleton asset/ReferenceSkeleton
+			BoneNode->SetCustomBindPoseLocalTransform(&NodeContainer, Bone.LocalRestTransform);
 			BoneNode->SetCustomTimeZeroLocalTransform(&NodeContainer, Bone.LocalRestTransform);
+			BoneNode->SetCustomLocalTransform(&NodeContainer, Bone.LocalRestTransform);
 
 			NodeContainer.AddNode(BoneNode);
 			NodeContainer.SetNodeParentUid(BoneNodeUid, ParentNode.GetUniqueID());
 
+			BoneToUidAndBoneIndex.Add(ConcatBonePath, {BoneNodeUid, BoneIndex});
+
 			for (int32 ChildIndex : Bone.ChildIndices)
 			{
-				RecursiveTraverseBones(ConvertedData.Bones[ChildIndex], *BoneNode, BoneNodeUid);
+				RecursiveTraverseBones(ChildIndex, *BoneNode, ConcatBonePath);
 			}
 		};
 
 		// Start traversing from the root bones (we may have more than one)
-		const FString SkeletonPrimNodeUid = SkeletonPrimNode.GetUniqueID();
-		for (const UsdToUnreal::FUsdSkeletonData::FBone& Bone : ConvertedData.Bones)
+		for (int32 BoneIndex = 0; BoneIndex < ConvertedData.Bones.Num(); ++BoneIndex)
 		{
+			const UsdToUnreal::FUsdSkeletonData::FBone& Bone = ConvertedData.Bones[BoneIndex];
 			if (Bone.ParentIndex == INDEX_NONE)
 			{
-				RecursiveTraverseBones(Bone, SkeletonPrimNode, SkeletonPrimNodeUid);
+				const FString BonePathRoot = TEXT("");
+				RecursiveTraverseBones(BoneIndex, SkeletonPrimNode, BonePathRoot);
 			}
 		}
 
+		// Handle SkelAnimation prims, if we have any bound for this Skeleton
+		AddSkeletalAnimationNode(SkelQuery, BoneToUidAndBoneIndex, TranslatorImpl, SkeletonPrimNode, NodeContainer, Info);
+
 		// Cache our joint names in order, as this is needed when generating skeletal mesh payloads
 		Info.SkelJointNames = MakeShared<TArray<FString>>();
 		Info.SkelJointNames->Reserve(ConvertedData.Bones.Num());
@@ -1082,6 +1157,224 @@ namespace UE::InterchangeUsdTranslator::Private
 		}
 	}
 
+	void AddSkeletalAnimationNode(
+		const UE::FUsdSkelSkeletonQuery& SkeletonQuery,
+		const TMap<FString, TPair<FString, int32>>& BoneToUidAndBoneIndex,
+		UInterchangeUSDTranslatorImpl& TranslatorImpl,
+		UInterchangeSceneNode& SkeletonPrimNode,
+		UInterchangeBaseNodeContainer& NodeContainer,
+		const FTraversalInfo& Info
+	)
+	{
+		UE::FUsdSkelAnimQuery AnimQuery = SkeletonQuery.GetAnimQuery();
+		if (!AnimQuery)
+		{
+			return;
+		}
+
+		UE::FUsdPrim SkelAnimationPrim = AnimQuery.GetPrim();
+		if (!SkelAnimationPrim)
+		{
+			return;
+		}
+
+		UE::FUsdPrim SkeletonPrim = SkeletonQuery.GetSkeleton();
+		if (!SkeletonPrim)
+		{
+			return;
+		}
+
+		UE::FUsdStage Stage = SkeletonPrim.GetStage();
+
+		const FString SkelAnimationName = SkelAnimationPrim.GetName().ToString();
+		const FString SkelAnimationPrimPath = SkelAnimationPrim.GetPrimPath().GetString();
+		const FString SkeletonPrimPath = SkeletonPrim.GetPrimPath().GetString();
+		const FString UniquePath = SkelAnimationPrimPath + TEXT("\\") + SkeletonPrimPath;
+		const FString NodeUid = AnimationTrackPrefix + UniquePath;
+
+		const UInterchangeSkeletalAnimationTrackNode* ExistingNode = Cast<UInterchangeSkeletalAnimationTrackNode>(NodeContainer.GetNode(NodeUid));
+		if (ExistingNode)
+		{
+			return;
+		};
+
+		UInterchangeSkeletalAnimationTrackNode* SkelAnimNode = NewObject<UInterchangeSkeletalAnimationTrackNode>(&NodeContainer);
+		SkelAnimNode->InitializeNode(NodeUid, SkelAnimationName, EInterchangeNodeContainerType::TranslatedAsset);
+		SkelAnimNode->SetCustomSkeletonNodeUid(SkeletonPrimNode.GetUniqueID());
+
+		// TODO: Uncomment this whenever Interchange supports skeletal animation sections, because
+		// currently it seems that InterchangeLevelSequenceFactory.cpp doesn't even have the string "skel" anywhere.
+		// If we were to add this all we'd get is a warning on the output log about "all referenced actors being missing",
+		// in case it failed to find anything else (e.g. other actual property/transform track) to put on the LevelSequence.
+		// AddTrackSetNode(TranslatorImpl, NodeContainer);
+		// TranslatorImpl->CurrentTrackSet->AddCustomAnimationTrackUid(NodeUid);
+
+		NodeContainer.AddNode(SkelAnimNode);
+
+		// Time info
+		{
+			// TODO: Match the TrackSet framerate whenever Interchange supports skeletal animation sections.
+			// float TrackSetFrameRate = 30.0f;
+			// if (TranslatorImpl->CurrentTrackSet->GetCustomFrameRate(TrackSetFrameRate))
+			// {
+			// 	SkelAnimNode->SetCustomAnimationSampleRate(TrackSetFrameRate);
+			// }
+			SkelAnimNode->SetCustomAnimationSampleRate(Stage.GetFramesPerSecond());
+
+			TOptional<double> StartTimeCode;
+			TOptional<double> StopTimeCode;
+
+			// For now we don't generate LevelSequences for sublayers and will instead put everything on a single
+			// LevelSequence for the entire stage, so we don't need to care so much about sublayer offset/scale like
+			// UsdToUnreal::ConvertSkelAnim does
+			TArray<double> JointTimeSamples;
+			if (AnimQuery.GetJointTransformTimeSamples(JointTimeSamples) && JointTimeSamples.Num() > 0)
+			{
+				StartTimeCode = JointTimeSamples[0];
+				StopTimeCode = JointTimeSamples[JointTimeSamples.Num() - 1];
+			}
+			TArray<double> BlendShapeTimeSamples;
+			if (AnimQuery.GetBlendShapeWeightTimeSamples(BlendShapeTimeSamples) && BlendShapeTimeSamples.Num() > 0)
+			{
+				StartTimeCode = FMath::Min(BlendShapeTimeSamples[0], StartTimeCode.Get(TNumericLimits<double>::Max()));
+				StopTimeCode = FMath::Max(BlendShapeTimeSamples[BlendShapeTimeSamples.Num() - 1], StopTimeCode.Get(TNumericLimits<double>::Lowest()));
+			}
+
+			UE::FUsdStage UsdStage = SkeletonPrim.GetStage();
+			double TimeCodesPerSecond = UsdStage.GetTimeCodesPerSecond();
+			if (StartTimeCode.IsSet())
+			{
+				SkelAnimNode->SetCustomAnimationStartTime(StartTimeCode.GetValue() / TimeCodesPerSecond);
+			}
+			if (StopTimeCode.IsSet())
+			{
+				SkelAnimNode->SetCustomAnimationStopTime(StopTimeCode.GetValue() / TimeCodesPerSecond);
+			}
+		}
+
+		// Joint animation
+		TArray<FString> UsdJointOrder = AnimQuery.GetJointOrder();
+		for (const FString& FullAnimatedBoneName : UsdJointOrder)
+		{
+			const TPair<FString, int32>* FoundPair = BoneToUidAndBoneIndex.Find(FullAnimatedBoneName);
+			if (!FoundPair)
+			{
+				continue;
+			}
+
+			const FString& BoneSceneNodeUid = FoundPair->Key;
+			int32 SkeletonOrderBoneIndex = FoundPair->Value;
+
+			const FString BoneAnimPayloadKey = SkeletonPrimPath + TEXT("\\") + LexToString(SkeletonOrderBoneIndex);
+
+			// When retrieving the payload later, We'll need that bone's index within the Skeleton prim to index into the
+			// InUsdSkeletonQuery.ComputeJointLocalTransforms() results.
+			// Note that we're describing joint transforms with baked frames here. It would have been possible to use transform
+			// curves, but that may have lead to issues when interpolating problematic joint transforms. Instead, we'll bake
+			// using USD, and let it interpolate the transforms however it wants
+			SkelAnimNode->SetAnimationPayloadKeyForSceneNodeUid(BoneSceneNodeUid, BoneAnimPayloadKey, EInterchangeAnimationPayLoadType::BAKED);
+		}
+
+		// Morph targets
+		{
+			UE::FUsdSkelBinding SkelBinding;
+			const bool bTraverseInstanceProxies = true;
+			bool bSuccess = Info.FurthestSkelCache->ComputeSkelBinding(	   //
+				Info.ClosestParentSkelRoot,
+				SkeletonPrim,
+				SkelBinding,
+				bTraverseInstanceProxies
+			);
+			if (!bSuccess)
+			{
+				return;
+			}
+
+			TArray<FString> SkelAnimChannelOrder = AnimQuery.GetBlendShapeOrder();
+
+			TMap<FString, int32> SkelAnimChannelIndices;
+			SkelAnimChannelIndices.Reserve(SkelAnimChannelOrder.Num());
+			for (int32 ChannelIndex = 0; ChannelIndex < SkelAnimChannelOrder.Num(); ++ChannelIndex)
+			{
+				const FString& ChannelName = SkelAnimChannelOrder[ChannelIndex];
+				SkelAnimChannelIndices.Add(ChannelName, ChannelIndex);
+			}
+
+			TArray<UE::FUsdSkelSkinningQuery> SkinningTargets = SkelBinding.GetSkinningTargets();
+			for (const UE::FUsdSkelSkinningQuery& SkinningTarget : SkinningTargets)
+			{
+				// USD lets you "skin" anything that can take the SkelBindingAPI, but we only care about Mesh here as
+				// those are the only ones that can have blendshapes
+				UE::FUsdPrim Prim = SkinningTarget.GetPrim();
+				if (!Prim.IsA(TEXT("Mesh")))
+				{
+					continue;
+				}
+				const FString MeshPrimPath = Prim.GetPrimPath().GetString();
+
+				TArray<FString> BlendShapeChannels;
+				bool bInnerSucces = SkinningTarget.GetBlendShapeOrder(BlendShapeChannels);
+				if (!bInnerSucces)
+				{
+					continue;
+				}
+
+				TArray<UE::FSdfPath> Targets;
+				{
+					UE::FUsdRelationship BlendShapeTargetsRel = SkinningTarget.GetBlendShapeTargetsRel();
+					if (!BlendShapeTargetsRel)
+					{
+						continue;
+					}
+					bInnerSucces = BlendShapeTargetsRel.GetTargets(Targets);
+					if (!bInnerSucces)
+					{
+						continue;
+					}
+				}
+
+				if (BlendShapeChannels.Num() != Targets.Num())
+				{
+					UE_LOG(
+						LogUsd,
+						Warning,
+						TEXT(
+							"Skipping morph target curves for animation of skinned mesh '%s' because the number of entries in the 'skel:blendShapes' attribute (%d) doesn't match the number of entries in the 'skel:blendShapeTargets' attribute (%d)"
+						),
+						*MeshPrimPath,
+						BlendShapeChannels.Num(),
+						Targets.Num()
+					);
+					continue;
+				}
+
+				for (int32 BlendShapeIndex = 0; BlendShapeIndex < Targets.Num(); ++BlendShapeIndex)
+				{
+					const FString& ChannelName = BlendShapeChannels[BlendShapeIndex];
+					int32* FoundSkelAnimChannelIndex = SkelAnimChannelIndices.Find(ChannelName);
+					if (!FoundSkelAnimChannelIndex)
+					{
+						// This channel is not animated by this SkelAnimation prim
+						continue;
+					}
+
+					// Note that we put no inbetween name on the MorphTargetUid: We only need to emit the morph target curve payloads
+					// for the main shapes: We'll provide the inbetween "positions" when providing the curve and Interchange computes
+					// the inbetween curves automatically
+					const FString BlendShapePath = Targets[BlendShapeIndex].GetString();
+					const FString MorphTargetUid = GetMorphTargetMeshNodeUid(MeshPrimPath, BlendShapeIndex);
+					const FString PayloadKey = GetMorphTargetCurvePayloadKey(SkeletonPrimPath, *FoundSkelAnimChannelIndex, BlendShapePath);
+
+					SkelAnimNode->SetAnimationPayloadKeyForMorphTargetNodeUid(	  //
+						MorphTargetUid,
+						PayloadKey,
+						EInterchangeAnimationPayLoadType::MORPHTARGETCURVE
+					);
+				}
+			}
+		}
+	}
+
 	void Traverse(
 		const UE::FUsdPrim& Prim,
 		UInterchangeUSDTranslatorImpl& TranslatorImpl,
@@ -1355,7 +1648,8 @@ namespace UE::InterchangeUsdTranslator::Private
 		FString& OutMorphTargetName
 	)
 	{
-		// Our payloadkey should be something like "/MySkelRoot/MyMeshPrim\2\BlendShapeName_inbetweens:someName"
+		// These payload keys are generated by GetMorphTargetMeshPayloadKey(), and so should take the form
+		// "<mesh prim path>\<mesh blend shape index>\<optional inbetween name>"
 		const bool bCullEmpty = false;
 		TArray<FString> PayloadKeyTokens;
 		PayloadKey.ParseIntoArray(PayloadKeyTokens, TEXT("\\"), bCullEmpty);
@@ -1607,7 +1901,11 @@ namespace UE::InterchangeUsdTranslator::Private
 		return true;
 	}
 
-	bool GetAnimationCurvePayloadData(const UE::FUsdStage& UsdStage, const FString& PayloadKey, Interchange::FAnimationPayloadData& OutPayloadData)
+	bool GetPropertyAnimationCurvePayloadData(
+		const UE::FUsdStage& UsdStage,
+		const FString& PayloadKey,
+		Interchange::FAnimationPayloadData& OutPayloadData
+	)
 	{
 		FString PrimPath;
 		FString UEPropertyNameStr;
@@ -1654,6 +1952,288 @@ namespace UE::InterchangeUsdTranslator::Private
 		return false;
 	}
 
+	bool GetJointAnimationCurvePayloadData(
+		const UInterchangeUSDTranslatorImpl& Impl,
+		const TArray<const UE::Interchange::FAnimationPayloadQuery*>& Queries,
+		TArray<UE::Interchange::FAnimationPayloadData>& OutPayloadData
+	)
+	{
+		if (Queries.Num() == 0)
+		{
+			return false;
+		}
+
+		// We expect all queries to be for the same skeleton, and have the same timing parameters,
+		// since they were grouped up by HashAnimPayloadQuery, so let's just grab one for the params
+		const UE::Interchange::FAnimationPayloadQuery* FirstQuery = Queries[0];
+
+		// Parse payload key.
+		// Here is takes the form "<skeleton prim path>\<joint index in skeleton order>"
+		TArray<FString> PayloadKeyTokens;
+		FirstQuery->PayloadKey.UniqueId.ParseIntoArray(PayloadKeyTokens, TEXT("\\"));
+		if (PayloadKeyTokens.Num() != 2)
+		{
+			return false;
+		}
+
+		// Fetch our cached skeleton query
+		const FString& SkeletonPrimPath = PayloadKeyTokens[0];
+		UE::FUsdSkelSkeletonQuery SkelQuery;
+		{
+			FReadScopeLock ReadLock{Impl.CachedTraversalInfoLock};
+
+			const FTraversalInfo* MeshInfo = Impl.NodeUidToCachedTraversalInfo.Find(SkeletonPrimPath);
+			if (!MeshInfo)
+			{
+				return false;
+			}
+			SkelQuery = MeshInfo->ActiveSkelQuery;
+			if (!SkelQuery)
+			{
+				return false;
+			}
+		}
+
+		UE::FUsdPrim SkeletonPrim = SkelQuery.GetPrim();
+		UE::FUsdStage Stage = SkeletonPrim.GetStage();
+		FUsdStageInfo StageInfo{Stage};
+
+		// Compute the bake ranges and intervals
+		double TimeCodesPerSecond = Stage.GetTimeCodesPerSecond();
+		double BakeFrequency = FirstQuery->TimeDescription.BakeFrequency;
+		double RangeStartSeconds = FirstQuery->TimeDescription.RangeStartSecond;
+		double RangeStopSeconds = FirstQuery->TimeDescription.RangeStopSecond;
+		double SectionLengthSeconds = RangeStopSeconds - RangeStartSeconds;
+		double StartTimeCode = RangeStartSeconds * TimeCodesPerSecond;
+		const int32 NumBakedFrames = FMath::RoundToInt(FMath::Max(SectionLengthSeconds * TimeCodesPerSecond + 1.0, 1.0));
+		double TimeCodeIncrement = (1.0 / BakeFrequency) * TimeCodesPerSecond;
+
+		// Bake all joint transforms via USD into arrays for each separate joint (in whatever order SkelQuery gives us)
+		TArray<TArray<FTransform>> BakedTransforms;
+		for (int32 FrameIndex = 0; FrameIndex < NumBakedFrames; ++FrameIndex)
+		{
+			const double FrameTimeCode = StartTimeCode + FrameIndex * TimeCodeIncrement;
+
+			TArray<FTransform> TransformsForTimeCode;
+			bool bSuccess = SkelQuery.ComputeJointLocalTransforms(TransformsForTimeCode, FrameTimeCode);
+			if (!bSuccess)
+			{
+				break;
+			}
+
+			for (FTransform& Transform : TransformsForTimeCode)
+			{
+				Transform = UsdUtils::ConvertTransformToUESpace(StageInfo, Transform);
+			}
+
+			// Setup our BakedTransforms in here, because we may actually get more or less transforms
+			// from the SkeletonQuery than our AnimSequence wants/expects, given that it can specify
+			// its own animated joint order
+			int32 NumSkelJoints = TransformsForTimeCode.Num();
+			if (FrameIndex == 0)
+			{
+				BakedTransforms.SetNum(NumSkelJoints);
+				for (int32 JointIndex = 0; JointIndex < NumSkelJoints; ++JointIndex)
+				{
+					BakedTransforms[JointIndex].SetNum(NumBakedFrames);
+				}
+			}
+
+			// Transpose our baked transforms into the arrays we'll eventually return
+			for (int32 JointIndex = 0; JointIndex < NumSkelJoints; ++JointIndex)
+			{
+				BakedTransforms[JointIndex][FrameIndex] = TransformsForTimeCode[JointIndex];
+			}
+		}
+
+		// Finally build our payload data return values by picking the desired baked arrays with the payload joint indices
+		OutPayloadData.Reset(Queries.Num());
+		for (int32 QueryIndex = 0; QueryIndex < Queries.Num(); ++QueryIndex)
+		{
+			const UE::Interchange::FAnimationPayloadQuery* Query = Queries[QueryIndex];
+
+			FString IndexStr = Query->PayloadKey.UniqueId.RightChop(SkeletonPrimPath.Len() + 1);	// Also skip the '\'
+			int32 JointIndex = INDEX_NONE;
+			bool bLexed = LexTryParseString(JointIndex, *IndexStr);
+			if (!bLexed)
+			{
+				continue;
+			}
+
+			UE::Interchange::FAnimationPayloadData& PayloadData = OutPayloadData.Emplace_GetRef(Query->SceneNodeUniqueID, Query->PayloadKey);
+			PayloadData.BakeFrequency = BakeFrequency;
+			PayloadData.RangeStartTime = RangeStartSeconds;
+			PayloadData.RangeEndTime = RangeStopSeconds;
+
+			if (BakedTransforms.IsValidIndex(JointIndex))
+			{
+				PayloadData.Transforms = MoveTemp(BakedTransforms[JointIndex]);
+			}
+		}
+
+		return true;
+	}
+
+	bool GetMorphTargetAnimationCurvePayloadData(
+		const UInterchangeUSDTranslatorImpl& Impl,
+		const FString& PayloadKey,
+		Interchange::FAnimationPayloadData& OutPayloadData
+	)
+	{
+		// Here we must output the morph target curve for a particular channel and skinning target, i.e.
+		// the connection of a SkelAnimation blend shape channel to a particular Mesh prim.
+
+		// These payload keys were generated from GetMorphTargetCurvePayloadKey(), so they take the form
+		// "<skeleton prim path>\<skel anim channel index>\<blend shape path>"
+		TArray<FString> PayloadKeyTokens;
+		PayloadKey.ParseIntoArray(PayloadKeyTokens, TEXT("\\"));
+		if (PayloadKeyTokens.Num() != 3)
+		{
+			return false;
+		}
+		const FString& SkeletonPrimPath = PayloadKeyTokens[0];
+		const FString& AnimChannelIndexStr = PayloadKeyTokens[1];
+		const FString& BlendShapePath = PayloadKeyTokens[2];
+
+		const UE::FUsdStage& UsdStage = Impl.UsdStage;
+
+		int32 SkelAnimChannelIndex = INDEX_NONE;
+		bool bLexed = LexTryParseString(SkelAnimChannelIndex, *AnimChannelIndexStr);
+
+		UE::FUsdPrim BlendShapePrim = UsdStage.GetPrimAtPath(UE::FSdfPath{*BlendShapePath});
+		UE::FUsdSkelBlendShape BlendShape{BlendShapePrim};
+		if (!BlendShape || !bLexed || SkelAnimChannelIndex == INDEX_NONE)
+		{
+			return false;
+		}
+		const FString BlendShapeName = BlendShapePrim.GetName().ToString();
+
+		// Fill in the actual morph target curve
+		UE::FUsdSkelAnimQuery AnimQuery;
+		{
+			UE::FUsdSkelSkeletonQuery SkelQuery;
+			{
+				FReadScopeLock ReadLock{Impl.CachedTraversalInfoLock};
+
+				const FTraversalInfo* MeshInfo = Impl.NodeUidToCachedTraversalInfo.Find(SkeletonPrimPath);
+				if (!MeshInfo)
+				{
+					return false;
+				}
+				SkelQuery = MeshInfo->ActiveSkelQuery;
+				if (!SkelQuery)
+				{
+					return false;
+				}
+			}
+
+			AnimQuery = SkelQuery.GetAnimQuery();
+			if (!AnimQuery)
+			{
+				return false;
+			}
+
+			TArray<double> TimeCodes;
+			bool bSuccess = AnimQuery.GetBlendShapeWeightTimeSamples(TimeCodes);
+			if (!bSuccess)
+			{
+				return false;
+			}
+
+			OutPayloadData.Curves.SetNum(1);
+			FRichCurve& Curve = OutPayloadData.Curves[0];
+			Curve.ReserveKeys(TimeCodes.Num());
+
+			const FFrameRate StageFrameRate{static_cast<uint32>(UsdStage.GetTimeCodesPerSecond()), 1};
+			const ERichCurveInterpMode InterpMode = (UsdStage.GetInterpolationType() == EUsdInterpolationType::Linear)
+														? ERichCurveInterpMode::RCIM_Linear
+														: ERichCurveInterpMode::RCIM_Constant;
+
+			TArray<float> Weights;
+			for (double TimeCode : TimeCodes)
+			{
+				bSuccess = AnimQuery.ComputeBlendShapeWeights(Weights, TimeCode);
+				if (!bSuccess || !Weights.IsValidIndex(SkelAnimChannelIndex))
+				{
+					break;
+				}
+
+				int32 FrameNumber = FMath::FloorToInt(TimeCode);
+				float SubFrameNumber = TimeCode - FrameNumber;
+
+				FFrameTime FrameTime{FrameNumber, SubFrameNumber};
+				double FrameTimeSeconds = static_cast<float>(StageFrameRate.AsSeconds(FrameTime));
+
+				FKeyHandle Handle = Curve.AddKey(FrameTimeSeconds, Weights[SkelAnimChannelIndex]);
+				Curve.SetKeyInterpMode(Handle, InterpMode);
+			}
+		}
+
+		TArray<FString> SkelAnimChannels = AnimQuery.GetBlendShapeOrder();
+
+		// Provide inbetween names/positions for this morph target payload
+		TArray<UE::FUsdSkelInbetweenShape> Inbetweens = BlendShape.GetInbetweens();
+		if (Inbetweens.Num() > 0)
+		{
+			// Let's store them into this temp struct so that we can sort them by weight first,
+			// as Interchange seems to expect that given how it will pass these right along into
+			// ResolveWeightsForBlendShape inside InterchangeAnimSequenceFactory.cpp
+			struct FInbetweenAndPosition
+			{
+				FString Name;
+				float Position;
+			};
+			TArray<FInbetweenAndPosition> ParsedInbetweens;
+			ParsedInbetweens.Reset(Inbetweens.Num());
+
+			for (const UE::FUsdSkelInbetweenShape& Inbetween : Inbetweens)
+			{
+				float Position = 0.5f;
+				bool bSuccess = Inbetween.GetWeight(&Position);
+				if (!bSuccess)
+				{
+					continue;
+				}
+
+				// Skip invalid positions. Note that technically positions outside the [0, 1] range seem to be allowed, but
+				// they don't seem to work very well with our inbetween weights resolution function for some reason.
+				// The legacy USD workflows have this exact same check though, so for consistency let's just do the same, and
+				// if becomes an issue we should fix both
+				if (Position > 1.0f || Position < 0.0f || FMath::IsNearlyZero(Position) || FMath::IsNearlyEqual(Position, 1.0f))
+				{
+					continue;
+				}
+
+				const FString MorphTargetName = BlendShapeName + TEXT("_") + Inbetween.GetAttr().GetName().ToString();
+				ParsedInbetweens.Emplace(MorphTargetName, Position);
+			}
+
+			ParsedInbetweens.Sort(
+				[](const FInbetweenAndPosition& LHS, const FInbetweenAndPosition& RHS)
+				{
+					// It's invalid USD to author two inbetweens with the same weight, so let's ignore that case here.
+					// (Reference: https://openusd.org/release/api/_usd_skel__schemas.html#UsdSkel_BlendShape)
+					return LHS.Position < RHS.Position;
+				}
+			);
+
+			OutPayloadData.InbetweenCurveNames.Reset(Inbetweens.Num() + 1);
+			OutPayloadData.InbetweenFullWeights.Reset(Inbetweens.Num());
+
+			// We add the main morph target curve name to InbetweenCurveNames too (having it end up one size bigger than
+			// InbetweenFullWeights) as it seems like that's what Interchange expects. See CreateMorphTargetCurve within
+			// InterchangeAnimSequenceFactory.cpp, and the very end of FFbxMesh::AddAllMeshes within FbxMesh.cpp
+			OutPayloadData.InbetweenCurveNames.Add(BlendShapeName);
+
+			for (const FInbetweenAndPosition& InbetweenAndPosition : ParsedInbetweens)
+			{
+				OutPayloadData.InbetweenCurveNames.Add(InbetweenAndPosition.Name);
+				OutPayloadData.InbetweenFullWeights.Add(InbetweenAndPosition.Position);
+			}
+		}
+
+		return true;
+	}
 }	 // namespace UE::InterchangeUsdTranslator::Private
 #endif	  // USE_USD_SDK
 
@@ -2002,79 +2582,110 @@ TOptional<UE::Interchange::FImportBlockedImage> UInterchangeUSDTranslator::GetBl
 	return BlockData;
 }
 
-TFuture<TOptional<UE::Interchange::FAnimationPayloadData>> UInterchangeUSDTranslator::ResolveAnimationPayloadQuery(
-	const UE::Interchange::FAnimationPayloadQuery& PayloadQuery
+TArray<UE::Interchange::FAnimationPayloadData> UInterchangeUSDTranslator::GetAnimationPayloadData(
+	const TArray<UE::Interchange::FAnimationPayloadQuery>& PayloadQueries
 ) const
 {
 	using namespace UE::Interchange;
 	using namespace UE::InterchangeUsdTranslator::Private;
 
-	return Async(
-		EAsyncExecution::TaskGraph,
-		[this, PayloadQuery]
-		{
-			TOptional<FAnimationPayloadData> Result;
+	TArray<UE::Interchange::FAnimationPayloadData> AnimationPayloads;
 
 #if USE_USD_SDK
-			FAnimationPayloadData AnimationPayLoadData{PayloadQuery.SceneNodeUniqueID, PayloadQuery.PayloadKey};
+	TMap<FString, TArray<const UE::Interchange::FAnimationPayloadQuery*>> BatchedBakeQueries;
+	BatchedBakeQueries.Reserve(PayloadQueries.Num());
 
-			UInterchangeUSDTranslatorImpl* ImplPtr = Impl.Get();
-			if (!ImplPtr)
+	// Inspect the queries we got
+	TArray<TFuture<TArray<UE::Interchange::FAnimationPayloadData>>> AnimationPayloadFutures;
+	for (const UE::Interchange::FAnimationPayloadQuery& PayloadQuery : PayloadQueries)
+	{
+		switch (PayloadQuery.PayloadKey.Type)
+		{
+			case EInterchangeAnimationPayLoadType::CURVE:	 // Fallthrough
+			case EInterchangeAnimationPayLoadType::STEPCURVE:
 			{
-				return Result;
+				// Property track animation queries.
+				//
+				// We're fine handling these in isolation (currently GetAnimationPayloadData is called with
+				// a single query at a time for these): Emit a separate task for each right away
+				AnimationPayloadFutures.Add(Async(
+					EAsyncExecution::TaskGraph,
+					[this, &PayloadQuery]
+					{
+						TArray<FAnimationPayloadData> Result;
+						FAnimationPayloadData AnimationPayLoadData{PayloadQuery.SceneNodeUniqueID, PayloadQuery.PayloadKey};
+						if (GetPropertyAnimationCurvePayloadData(Impl->UsdStage, PayloadQuery.PayloadKey.UniqueId, AnimationPayLoadData))
+						{
+							Result.Emplace(AnimationPayLoadData);
+						}
+						return Result;
+					}
+				));
+				break;
 			}
-
-			switch (PayloadQuery.PayloadKey.Type)
+			case EInterchangeAnimationPayLoadType::BAKED:
 			{
-				case EInterchangeAnimationPayLoadType::CURVE:	 // Fallthrough
-				case EInterchangeAnimationPayLoadType::STEPCURVE:
-				{
-					if (GetAnimationCurvePayloadData(ImplPtr->UsdStage, PayloadQuery.PayloadKey.UniqueId, AnimationPayLoadData))
+				// Joint transform animation queries.
+				//
+				// Currently we'll receive the PayloadQueries for all joints of a skeletal animation on the same GetAnimationPayloadData
+				// call. Unfortunately in USD we must compute all joint transforms every time, even if all we need is data for a single
+				// joint. For efficiency then, we group up all the queries for the separate joints of the same skeleton into one batch
+				// task that we can resolve in one pass
+				const FString BakedQueryHash = HashAnimPayloadQuery(PayloadQuery);
+				TArray<const UE::Interchange::FAnimationPayloadQuery*>& Queries = BatchedBakeQueries.FindOrAdd(BakedQueryHash);
+				Queries.Add(&PayloadQuery);
+				break;
+			}
+			case EInterchangeAnimationPayLoadType::MORPHTARGETCURVE:
+			{
+				// Morph target curve queries.
+				AnimationPayloadFutures.Add(Async(
+					EAsyncExecution::TaskGraph,
+					[this, &PayloadQuery]
 					{
-						Result.Emplace(AnimationPayLoadData);
+						TArray<FAnimationPayloadData> Result;
+						FAnimationPayloadData AnimationPayLoadData{PayloadQuery.SceneNodeUniqueID, PayloadQuery.PayloadKey};
+						if (GetMorphTargetAnimationCurvePayloadData(*Impl, PayloadQuery.PayloadKey.UniqueId, AnimationPayLoadData))
+						{
+							Result.Emplace(AnimationPayLoadData);
+						}
+						return Result;
 					}
-					break;
-				}
-				case EInterchangeAnimationPayLoadType::MORPHTARGETCURVE:
-				case EInterchangeAnimationPayLoadType::BAKED:
-				case EInterchangeAnimationPayLoadType::NONE:
-				default:
-				{
-					break;
-				}
+				));
+				break;
+			}
+			case EInterchangeAnimationPayLoadType::NONE:
+			default:
+			{
+				break;
 			}
-#endif	  // USE_USD_SDK
-
-			return Result;
 		}
-	);
-}
-
-TArray<UE::Interchange::FAnimationPayloadData> UInterchangeUSDTranslator::GetAnimationPayloadData(
-	const TArray<UE::Interchange::FAnimationPayloadQuery>& PayloadQueries
-) const
-{
-	// Note: This TFuture/Async approach is a bit overkill when it comes to the property track case, as each
-	// individual PayloadQueries has a single item in it (check FLevelSequenceHelper::PopulateAnimationTrack).
-	// This will actually help a lot with AnimSequences though, where RetrieveAnimationPayloads shows how all
-	// the skeletal joint tracks are queried in a single call (and the analogous for the morph target curves).
+	}
 
-	TArray<TFuture<TOptional<UE::Interchange::FAnimationPayloadData>>> AnimationPayloadFutures;
-	for (const UE::Interchange::FAnimationPayloadQuery& PayloadQuery : PayloadQueries)
+	// Emit the batched joint transform animation tasks
+	for (const TPair<FString, TArray<const UE::Interchange::FAnimationPayloadQuery*>>& BatchedBakedQueryPair : BatchedBakeQueries)
 	{
-		AnimationPayloadFutures.Add(ResolveAnimationPayloadQuery(PayloadQuery));
+		AnimationPayloadFutures.Add(Async(
+			EAsyncExecution::TaskGraph,
+			[this, &BatchedBakedQueryPair]
+			{
+				const TArray<const UE::Interchange::FAnimationPayloadQuery*>& Queries = BatchedBakedQueryPair.Value;
+				TArray<UE::Interchange::FAnimationPayloadData> Result;
+				if (GetJointAnimationCurvePayloadData(*Impl, Queries, Result))
+				{
+					return Result;
+				}
+				return TArray<UE::Interchange::FAnimationPayloadData>{};
+			}
+		));
 	}
 
-	TArray<UE::Interchange::FAnimationPayloadData> AnimationPayloads;
-	for (TFuture<TOptional<UE::Interchange::FAnimationPayloadData>>& AnimationPayloadFuture : AnimationPayloadFutures)
+	// Wait for all async tasks to complete and collect the results
+	for (TFuture<TArray<UE::Interchange::FAnimationPayloadData>>& AnimationPayloadFuture : AnimationPayloadFutures)
 	{
-		TOptional<UE::Interchange::FAnimationPayloadData> OptionalPayloadData = AnimationPayloadFuture.Get();
-		if (!OptionalPayloadData.IsSet())
-		{
-			continue;
-		}
-		AnimationPayloads.Add(OptionalPayloadData.GetValue());
+		AnimationPayloads.Append(AnimationPayloadFuture.Get());
 	}
+#endif	  // USE_USD_SDK
 
 	return AnimationPayloads;
 }
diff --git a/Engine/Plugins/Runtime/USDCore/Source/UnrealUSDWrapper/Private/UsdWrappers/UsdSkelAnimQuery.cpp b/Engine/Plugins/Runtime/USDCore/Source/UnrealUSDWrapper/Private/UsdWrappers/UsdSkelAnimQuery.cpp
--- a/Engine/Plugins/Runtime/USDCore/Source/UnrealUSDWrapper/Private/UsdWrappers/UsdSkelAnimQuery.cpp
+++ b/Engine/Plugins/Runtime/USDCore/Source/UnrealUSDWrapper/Private/UsdWrappers/UsdSkelAnimQuery.cpp
@@ -142,4 +142,97 @@ namespace UE
 		return {};
 #endif	  // #if USE_USD_SDK
 	}
+
+	bool FUsdSkelAnimQuery::ComputeBlendShapeWeights(TArray<float>& Weights, TOptional<double> TimeCode) const
+	{
+		bool bResult = false;
+
+#if USE_USD_SDK
+		pxr::UsdTimeCode UsdTimeCode = TimeCode.IsSet() ? TimeCode.GetValue() : pxr::UsdTimeCode::Default();
+
+		TUsdStore<pxr::VtArray<float>> UsdWeights;
+		bResult = Impl->PxrUsdSkelAnimQuery.Get().ComputeBlendShapeWeights(&UsdWeights.Get(), UsdTimeCode);
+		if (!bResult)
+		{
+			return bResult;
+		}
+
+		Weights.SetNumUninitialized(UsdWeights.Get().size());
+		FMemory::Memcpy(Weights.GetData(), UsdWeights.Get().cdata(), UsdWeights.Get().size() * sizeof(float));
+#endif	  // #if USE_USD_SDK
+
+		return bResult;
+	}
+
+	bool FUsdSkelAnimQuery::GetJointTransformTimeSamples(TArray<double>& TimeCodes) const
+	{
+		bool bResult = false;
+
+#if USE_USD_SDK
+		TUsdStore<std::vector<double>> UsdTimes;
+		bResult = Impl->PxrUsdSkelAnimQuery.Get().GetJointTransformTimeSamples(&UsdTimes.Get());
+		if (!bResult)
+		{
+			return bResult;
+		}
+
+		TimeCodes.SetNumUninitialized(UsdTimes.Get().size());
+		FMemory::Memcpy(TimeCodes.GetData(), UsdTimes.Get().data(), UsdTimes.Get().size() * sizeof(double));
+#endif	  // #if USE_USD_SDK
+
+		return bResult;
+	}
+
+	bool FUsdSkelAnimQuery::GetBlendShapeWeightTimeSamples(TArray<double>& TimeCodes) const
+	{
+		bool bResult = false;
+
+#if USE_USD_SDK
+		TUsdStore<std::vector<double>> UsdTimes;
+		bResult = Impl->PxrUsdSkelAnimQuery.Get().GetBlendShapeWeightTimeSamples(&UsdTimes.Get());
+		if (!bResult)
+		{
+			return bResult;
+		}
+
+		TimeCodes.SetNumUninitialized(UsdTimes.Get().size());
+		FMemory::Memcpy(TimeCodes.GetData(), UsdTimes.Get().data(), UsdTimes.Get().size() * sizeof(double));
+#endif	  // #if USE_USD_SDK
+
+		return bResult;
+	}
+
+	TArray<FString> FUsdSkelAnimQuery::GetJointOrder() const
+	{
+		TArray<FString> Result;
+
+#if USE_USD_SDK
+		TUsdStore<pxr::VtArray<pxr::TfToken>> UsdJointOrder = Impl->PxrUsdSkelAnimQuery.Get().GetJointOrder();
+		Result.Reserve(UsdJointOrder.Get().size());
+
+		for (const pxr::TfToken& UsdToken : UsdJointOrder.Get())
+		{
+			Result.Add(UTF8_TO_TCHAR(UsdToken.GetString().c_str()));
+		}
+#endif
+
+		return Result;
+	}
+
+	TArray<FString> FUsdSkelAnimQuery::GetBlendShapeOrder() const
+	{
+		TArray<FString> Result;
+
+#if USE_USD_SDK
+		TUsdStore<pxr::VtArray<pxr::TfToken>> UsdJointOrder = Impl->PxrUsdSkelAnimQuery.Get().GetBlendShapeOrder();
+		Result.Reserve(UsdJointOrder.Get().size());
+
+		for (const pxr::TfToken& UsdToken : UsdJointOrder.Get())
+		{
+			Result.Add(UTF8_TO_TCHAR(UsdToken.GetString().c_str()));
+		}
+#endif
+
+		return Result;
+	}
 }	 // namespace UE
diff --git a/Engine/Plugins/Runtime/USDCore/Source/UnrealUSDWrapper/Private/UsdWrappers/UsdSkelSkeletonQuery.cpp b/Engine/Plugins/Runtime/USDCore/Source/UnrealUSDWrapper/Private/UsdWrappers/UsdSkelSkeletonQuery.cpp
--- a/Engine/Plugins/Runtime/USDCore/Source/UnrealUSDWrapper/Private/UsdWrappers/UsdSkelSkeletonQuery.cpp
+++ b/Engine/Plugins/Runtime/USDCore/Source/UnrealUSDWrapper/Private/UsdWrappers/UsdSkelSkeletonQuery.cpp
@@ -2,10 +2,12 @@
 
 #include "UsdWrappers/UsdSkelSkeletonQuery.h"
 
+#include "UnrealUSDWrapper.h"
 #include "USDMemory.h"
 #include "UsdWrappers/UsdPrim.h"
 #include "UsdWrappers/UsdSkelAnimQuery.h"
 #include "UsdWrappers/UsdSkelSkeletonQuery.h"
+#include "UsdWrappers/UsdStage.h"
 
 #if USE_USD_SDK
 #include "USDIncludesStart.h"
@@ -165,4 +167,33 @@ namespace UE
 		return {};
 #endif	  // #if USE_USD_SDK
 	}
+
+	bool FUsdSkelSkeletonQuery::ComputeJointLocalTransforms(TArray<FTransform>& UESpaceTransforms, double TimeCode, bool bAtRest) const
+	{
+		bool bResult = false;
+
+#if USE_USD_SDK
+		TUsdStore<pxr::VtArray<pxr::GfMatrix4d>> UsdTransforms;
+		bResult = Impl->PxrUsdSkelSkeletonQuery.Get().ComputeJointLocalTransforms(&UsdTransforms.Get(), TimeCode, bAtRest);
+		if (!bResult)
+		{
+			return false;
+		}
+
+		UESpaceTransforms.Reset(UsdTransforms.Get().size());
+		for (const pxr::GfMatrix4d& Matrix : UsdTransforms.Get())
+		{
+			// Copy-pasting from UsdToUnreal::ConvertMatrix since we UnrealUSDWrapper can't depend on USDUtilities
+			FMatrix UnrealMatrix(
+				FPlane(Matrix[0][0], Matrix[0][1], Matrix[0][2], Matrix[0][3]),
+				FPlane(Matrix[1][0], Matrix[1][1], Matrix[1][2], Matrix[1][3]),
+				FPlane(Matrix[2][0], Matrix[2][1], Matrix[2][2], Matrix[2][3]),
+				FPlane(Matrix[3][0], Matrix[3][1], Matrix[3][2], Matrix[3][3])
+			);
+			UESpaceTransforms.Emplace(UnrealMatrix);
+		}
+#endif	  // #if USE_USD_SDK
+
+		return bResult;
+	}
 }	 // namespace UE
diff --git a/Engine/Plugins/Runtime/USDCore/Source/UnrealUSDWrapper/Private/UsdWrappers/UsdSkelSkinningQuery.cpp b/Engine/Plugins/Runtime/USDCore/Source/UnrealUSDWrapper/Private/UsdWrappers/UsdSkelSkinningQuery.cpp
--- a/Engine/Plugins/Runtime/USDCore/Source/UnrealUSDWrapper/Private/UsdWrappers/UsdSkelSkinningQuery.cpp
+++ b/Engine/Plugins/Runtime/USDCore/Source/UnrealUSDWrapper/Private/UsdWrappers/UsdSkelSkinningQuery.cpp
@@ -4,6 +4,7 @@
 
 #include "USDMemory.h"
 #include "UsdWrappers/UsdPrim.h"
+#include "UsdWrappers/UsdRelationship.h"
 #include "UsdWrappers/UsdSkelSkinningQuery.h"
 
 #if USE_USD_SDK
@@ -137,10 +138,44 @@ namespace UE
 	FUsdPrim FUsdSkelSkinningQuery::GetPrim() const
 	{
 #if USE_USD_SDK
-		FScopedUsdAllocs Allocs;
 		return FUsdPrim{Impl->PxrUsdSkelSkinningQuery.Get().GetPrim()};
 #else
 		return {};
 #endif	  // #if USE_USD_SDK
 	}
+
+	UE::FUsdRelationship FUsdSkelSkinningQuery::GetBlendShapeTargetsRel() const
+	{
+#if USE_USD_SDK
+		return UE::FUsdRelationship{Impl->PxrUsdSkelSkinningQuery.Get().GetBlendShapeTargetsRel()};
+#else
+		return {};
+#endif	  // #if USE_USD_SDK
+	}
+
+	bool FUsdSkelSkinningQuery::GetBlendShapeOrder(TArray<FString>& BlendShapes) const
+	{
+		bool bResult = false;
+
+#if USE_USD_SDK
+		pxr::VtArray<pxr::TfToken> UsdOrder;
+		{
+			FScopedUsdAllocs Allocs;
+			bResult = Impl->PxrUsdSkelSkinningQuery.Get().GetBlendShapeOrder(&UsdOrder);
+			if (!bResult)
+			{
+				return bResult;
+			}
+		}
+
+		BlendShapes.Reset(UsdOrder.size());
+		for (const pxr::TfToken& BlendShape : UsdOrder)
+		{
+			BlendShapes.Add(UTF8_TO_TCHAR(BlendShape.GetString().c_str()));
+		}
+#endif
+
+		return bResult;
+	}
+
 }	 // namespace UE
diff --git a/Engine/Plugins/Runtime/USDCore/Source/UnrealUSDWrapper/Public/UsdWrappers/UsdSkelAnimQuery.h b/Engine/Plugins/Runtime/USDCore/Source/UnrealUSDWrapper/Public/UsdWrappers/UsdSkelAnimQuery.h
--- a/Engine/Plugins/Runtime/USDCore/Source/UnrealUSDWrapper/Public/UsdWrappers/UsdSkelAnimQuery.h
+++ b/Engine/Plugins/Runtime/USDCore/Source/UnrealUSDWrapper/Public/UsdWrappers/UsdSkelAnimQuery.h
@@ -3,6 +3,7 @@
 #pragma once
 
 #include "CoreTypes.h"
+#include "Misc/Optional.h"
 #include "Templates/UniquePtr.h"
 
 #if USE_USD_SDK
@@ -59,6 +60,14 @@ namespace UE
 
 		UE::FUsdPrim GetPrim() const;
 
+		bool ComputeBlendShapeWeights(TArray<float>& Weights, TOptional<double> TimeCode = {}) const;
+
+		bool GetJointTransformTimeSamples(TArray<double>& TimeCodes) const;
+		bool GetBlendShapeWeightTimeSamples(TArray<double>& TimeCodes) const;
+
+		TArray<FString> GetJointOrder() const;
+		TArray<FString> GetBlendShapeOrder() const;
+
 	private:
 		TUniquePtr<Internal::FUsdSkelAnimQueryImpl> Impl;
 	};
diff --git a/Engine/Plugins/Runtime/USDCore/Source/UnrealUSDWrapper/Public/UsdWrappers/UsdSkelSkeletonQuery.h b/Engine/Plugins/Runtime/USDCore/Source/UnrealUSDWrapper/Public/UsdWrappers/UsdSkelSkeletonQuery.h
--- a/Engine/Plugins/Runtime/USDCore/Source/UnrealUSDWrapper/Public/UsdWrappers/UsdSkelSkeletonQuery.h
+++ b/Engine/Plugins/Runtime/USDCore/Source/UnrealUSDWrapper/Public/UsdWrappers/UsdSkelSkeletonQuery.h
@@ -65,6 +65,8 @@ namespace UE
 		UE::FUsdPrim GetSkeleton() const;
 		UE::FUsdSkelAnimQuery GetAnimQuery() const;
 
+		bool ComputeJointLocalTransforms(TArray<FTransform>& UsdSpaceTransforms, double TimeCode, bool bAtRest = false) const;
+
 	private:
 		TUniquePtr<Internal::FUsdSkelSkeletonQueryImpl> Impl;
 	};
diff --git a/Engine/Plugins/Runtime/USDCore/Source/UnrealUSDWrapper/Public/UsdWrappers/UsdSkelSkinningQuery.h b/Engine/Plugins/Runtime/USDCore/Source/UnrealUSDWrapper/Public/UsdWrappers/UsdSkelSkinningQuery.h
--- a/Engine/Plugins/Runtime/USDCore/Source/UnrealUSDWrapper/Public/UsdWrappers/UsdSkelSkinningQuery.h
+++ b/Engine/Plugins/Runtime/USDCore/Source/UnrealUSDWrapper/Public/UsdWrappers/UsdSkelSkinningQuery.h
@@ -21,6 +21,7 @@ PXR_NAMESPACE_CLOSE_SCOPE
 namespace UE
 {
 	class FUsdPrim;
+	class FUsdRelationship;
 
 	namespace Internal
 	{
@@ -59,6 +60,10 @@ namespace UE
 
 		UE::FUsdPrim GetPrim() const;
 
+		UE::FUsdRelationship GetBlendShapeTargetsRel() const;
+
+		bool GetBlendShapeOrder(TArray<FString>& BlendShapes) const;
+
 	private:
 		TUniquePtr<Internal::FUsdSkelSkinningQueryImpl> Impl;
 	};
