USD: Reorganize function order;

Remove unused includes.

This has no code change whatsoever, it literally just reorganizes the functions into a predictable order/grouping.

diff --git a/Engine/Plugins/Interchange/Runtime/Source/Import/Private/Usd/InterchangeUsdTranslator.cpp b/Engine/Plugins/Interchange/Runtime/Source/Import/Private/Usd/InterchangeUsdTranslator.cpp
--- a/Engine/Plugins/Interchange/Runtime/Source/Import/Private/Usd/InterchangeUsdTranslator.cpp
+++ b/Engine/Plugins/Interchange/Runtime/Source/Import/Private/Usd/InterchangeUsdTranslator.cpp
@@ -27,32 +27,26 @@
 #include "UsdWrappers/UsdSkelSkinningQuery.h"
 #include "UsdWrappers/UsdStage.h"
 #include "UsdWrappers/UsdTyped.h"
-#include "UsdWrappers/VtValue.h"
 
 #include "Async/Async.h"
 #include "HAL/IConsoleManager.h"
 #include "InterchangeCameraNode.h"
-#include "InterchangeImportLog.h"
 #include "InterchangeLightNode.h"
 #include "InterchangeMaterialInstanceNode.h"
 #include "InterchangeMeshNode.h"
 #include "InterchangeSceneNode.h"
 #include "InterchangeTexture2DNode.h"
 #include "InterchangeTranslatorHelper.h"
-#include "Internationalization/Regex.h"
 #include "MovieSceneSection.h"
 #include "Rendering/SkeletalMeshLODImporterData.h"
 #include "StaticMeshAttributes.h"
-#include "StaticMeshOperations.h"
 #include "UDIMUtilities.h"
 
 #if USE_USD_SDK
 #include "USDIncludesStart.h"
 #include "pxr/usd/usdGeom/tokens.h"
 #include "pxr/usd/usdLux/tokens.h"
-#include "pxr/usd/usdPhysics/tokens.h"
 #include "pxr/usd/usdShade/tokens.h"
-#include "pxr/usd/usdSkel/tokens.h"
 #include "USDIncludesEnd.h"
 #endif	  // USE_USD_SDK
 
@@ -155,49 +149,8 @@
 		// here for easy access when parsing the tracks
 		UInterchangeAnimationTrackSetNode* CurrentTrackSet = nullptr;
 	};
-}	 // namespace UE::InterchangeUsdTranslator::Private
 
-UInterchangeUsdTranslatorSettings::UInterchangeUsdTranslatorSettings()
-	: GeometryPurpose((int32)(EUsdPurpose::Default | EUsdPurpose::Proxy | EUsdPurpose::Render | EUsdPurpose::Guide))
-	, RenderContext(TEXT("unreal"))	   // The proper definition of this is on the USDSchemas module, which we can't depend on
-	, MaterialPurpose(*UnrealIdentifiers::MaterialPreviewPurpose)
-	, InterpolationType(EUsdInterpolationType::Linear)
-	, bOverrideStageOptions(false)
-	, StageOptions{
-		  0.01f,			   // MetersPerUnit
-		  EUsdUpAxis::ZAxis	   // UpAxis
-	  }
-{
-}
-
-UInterchangeUSDTranslator::UInterchangeUSDTranslator()
-	: Impl(MakeUnique<UE::InterchangeUsdTranslator::Private::UInterchangeUSDTranslatorImpl>())
-{
-}
-
-EInterchangeTranslatorType UInterchangeUSDTranslator::GetTranslatorType() const
-{
-	return GInterchangeEnableUSDLevelImport ? EInterchangeTranslatorType::Scenes : EInterchangeTranslatorType::Assets;
-}
-
-EInterchangeTranslatorAssetType UInterchangeUSDTranslator::GetSupportedAssetTypes() const
-{
-	return EInterchangeTranslatorAssetType::Materials | EInterchangeTranslatorAssetType::Meshes | EInterchangeTranslatorAssetType::Animations;
-}
-
-TArray<FString> UInterchangeUSDTranslator::GetSupportedFormats() const
-{
-	TArray<FString> Extensions;
-	if (GInterchangeEnableUSDImport)
-	{
-		UnrealUSDWrapper::AddUsdImportFileFormatDescriptions(Extensions);
-	}
-	return Extensions;
-}
-
 #if USE_USD_SDK
-namespace UE::InterchangeUsdTranslator::Private
-{
 	FString HashAnimPayloadQuery(const Interchange::FAnimationPayloadQuery& Query)
 	{
 		FSHAHash Hash;
@@ -268,3 +221,280 @@
 		}
 
 		return true;
+	}
+
+	void FixSkeletalMeshDescriptionColors(FMeshDescription& MeshDescription)
+	{
+		// FSkeletalMeshImportData::GetMeshDescription() will reinterpret our Wedge FColors as linear, and put those
+		// sRGB values disguised as linear into the mesh description. This also seems to disagree with the patch on
+		// cl 32791826, so here we have to fix that up and get our mesh description colors to be actually linear...
+		//
+		// This will hopefully go away once we have our own skinned mesh to FMeshDescription conversion function.
+		//
+		// Note: Weirdly enough skeletal meshes seem to put linear colors on VertexColor output, while static meshes
+		// put sRGB colors? Maybe this is why the comment above the change on 32791826 mentions to remove the ToFColor on
+		// StaticMeshBuilder? This is overall very confusing
+		FStaticMeshAttributes Attributes(MeshDescription);
+		TVertexInstanceAttributesRef<FVector4f> VertexColor = Attributes.GetVertexInstanceColors();
+		for (FVertexInstanceID VertexInstanceID : MeshDescription.VertexInstances().GetElementIDs())
+		{
+			const FColor ActualSRGB = FLinearColor(VertexColor[VertexInstanceID]).ToFColor(false);
+			VertexColor[VertexInstanceID] = FLinearColor{ActualSRGB};
+		}
+	}
+
+	void FixMaterialSlotNames(FMeshDescription& MeshDescription, const TArray<UsdUtils::FUsdPrimMaterialSlot>& MeshAssingmentSlots)
+	{
+		// Fixup material slot names to match the material that is assigned. For Interchange it is better to have the material
+		// slot names match what is assigned into them, as it will use those names to "merge identical slots" depending on the
+		// import options.
+		//
+		// Note: These names must also match what is set via MeshNode->SetSlotMaterialDependencyUid(SlotName, MaterialUid)
+		FStaticMeshAttributes StaticMeshAttributes(MeshDescription);
+		for (int32 MaterialSlotIndex = 0; MaterialSlotIndex < StaticMeshAttributes.GetPolygonGroupMaterialSlotNames().GetNumElements();
+			 ++MaterialSlotIndex)
+		{
+			int32 MaterialIndex = 0;
+			LexFromString(MaterialIndex, *StaticMeshAttributes.GetPolygonGroupMaterialSlotNames()[MaterialSlotIndex].ToString());
+
+			if (MeshAssingmentSlots.IsValidIndex(MaterialIndex))
+			{
+				const FString Source = MeshAssingmentSlots[MaterialIndex].MaterialSource;
+				StaticMeshAttributes.GetPolygonGroupMaterialSlotNames()[MaterialSlotIndex] = *Source;
+			}
+		}
+	}
+
+	void UpdateTraversalInfo(FTraversalInfo& Info, const UE::FUsdPrim& CurrentPrim)
+	{
+		if (CurrentPrim.IsA(TEXT("SkelRoot")))
+		{
+			if (!Info.ClosestParentSkelRoot)
+			{
+				// The root-most skel cache should handle any nested UsdSkel prims as well
+				Info.FurthestSkelCache = MakeShared<UE::FUsdSkelCache>();
+
+				const bool bTraverseInstanceProxies = true;
+				Info.FurthestSkelCache->Populate(CurrentPrim, bTraverseInstanceProxies);
+			}
+
+			Info.ClosestParentSkelRoot = CurrentPrim;
+		}
+
+		if (Info.ClosestParentSkelRoot && CurrentPrim.HasAPI(TEXT("SkelBindingAPI")))
+		{
+			UE::FUsdStage Stage = CurrentPrim.GetStage();
+
+			if (UE::FUsdRelationship SkelRel = CurrentPrim.GetRelationship(TEXT("skel:skeleton")))
+			{
+				TArray<UE::FSdfPath> Targets;
+				if (SkelRel.GetTargets(Targets) && Targets.Num() > 0)
+				{
+					UE::FUsdPrim TargetSkeleton = Stage.GetPrimAtPath(Targets[0]);
+					if (TargetSkeleton && TargetSkeleton.IsA(TEXT("Skeleton")))
+					{
+						Info.ActiveSkelQuery = Info.FurthestSkelCache->GetSkelQuery(TargetSkeleton);
+					}
+				}
+			}
+		}
+	}
+
+	bool ReadBools(
+		const UE::FUsdStage& UsdStage,
+		const TArray<double>& UsdTimeSamples,
+		const TFunction<bool(double)>& ReaderFunc,
+		Interchange::FAnimationPayloadData& OutPayloadData
+	)
+	{
+		OutPayloadData.StepCurves.SetNum(1);
+		FInterchangeStepCurve& Curve = OutPayloadData.StepCurves[0];
+		TArray<float>& KeyTimes = Curve.KeyTimes;
+		TArray<bool>& BooleanKeyValues = Curve.BooleanKeyValues.Emplace();
+
+		KeyTimes.Reserve(UsdTimeSamples.Num());
+		BooleanKeyValues.Reserve(UsdTimeSamples.Num());
+
+		const FFrameRate StageFrameRate{static_cast<uint32>(UsdStage.GetTimeCodesPerSecond()), 1};
+
+		double LastTimeSample = TNumericLimits<double>::Lowest();
+		for (const double UsdTimeSample : UsdTimeSamples)
+		{
+			// We never want to evaluate the same time twice
+			if (FMath::IsNearlyEqual(UsdTimeSample, LastTimeSample))
+			{
+				continue;
+			}
+			LastTimeSample = UsdTimeSample;
+
+			int32 FrameNumber = FMath::FloorToInt(UsdTimeSample);
+			float SubFrameNumber = UsdTimeSample - FrameNumber;
+
+			FFrameTime FrameTime{FrameNumber, SubFrameNumber};
+			double FrameTimeSeconds = static_cast<float>(StageFrameRate.AsSeconds(FrameTime));
+
+			bool UEValue = ReaderFunc(UsdTimeSample);
+
+			KeyTimes.Add(FrameTimeSeconds);
+			BooleanKeyValues.Add(UEValue);
+		}
+
+		return true;
+	}
+
+	bool ReadFloats(
+		const UE::FUsdStage& UsdStage,
+		const TArray<double>& UsdTimeSamples,
+		const TFunction<float(double)>& ReaderFunc,
+		Interchange::FAnimationPayloadData& OutPayloadData
+	)
+	{
+		OutPayloadData.Curves.SetNum(1);
+		FRichCurve& Curve = OutPayloadData.Curves[0];
+
+		const FFrameRate StageFrameRate{static_cast<uint32>(UsdStage.GetTimeCodesPerSecond()), 1};
+		const ERichCurveInterpMode InterpMode = (UsdStage.GetInterpolationType() == EUsdInterpolationType::Linear)
+													? ERichCurveInterpMode::RCIM_Linear
+													: ERichCurveInterpMode::RCIM_Constant;
+
+		double LastTimeSample = TNumericLimits<double>::Lowest();
+		for (const double UsdTimeSample : UsdTimeSamples)
+		{
+			// We never want to evaluate the same time twice
+			if (FMath::IsNearlyEqual(UsdTimeSample, LastTimeSample))
+			{
+				continue;
+			}
+			LastTimeSample = UsdTimeSample;
+
+			int32 FrameNumber = FMath::FloorToInt(UsdTimeSample);
+			float SubFrameNumber = UsdTimeSample - FrameNumber;
+
+			FFrameTime FrameTime{FrameNumber, SubFrameNumber};
+			double FrameTimeSeconds = static_cast<float>(StageFrameRate.AsSeconds(FrameTime));
+
+			float UEValue = ReaderFunc(UsdTimeSample);
+
+			FKeyHandle Handle = Curve.AddKey(FrameTimeSeconds, UEValue);
+			Curve.SetKeyInterpMode(Handle, InterpMode);
+		}
+
+		return true;
+	}
+
+	bool ReadColors(
+		const UE::FUsdStage& UsdStage,
+		const TArray<double>& UsdTimeSamples,
+		const TFunction<FLinearColor(double)>& ReaderFunc,
+		Interchange::FAnimationPayloadData& OutPayloadData
+	)
+	{
+		OutPayloadData.Curves.SetNum(4);
+		FRichCurve& RCurve = OutPayloadData.Curves[0];
+		FRichCurve& GCurve = OutPayloadData.Curves[1];
+		FRichCurve& BCurve = OutPayloadData.Curves[2];
+		FRichCurve& ACurve = OutPayloadData.Curves[3];
+
+		const FFrameRate StageFrameRate{static_cast<uint32>(UsdStage.GetTimeCodesPerSecond()), 1};
+		const ERichCurveInterpMode InterpMode = (UsdStage.GetInterpolationType() == EUsdInterpolationType::Linear)
+													? ERichCurveInterpMode::RCIM_Linear
+													: ERichCurveInterpMode::RCIM_Constant;
+
+		double LastTimeSample = TNumericLimits<double>::Lowest();
+		for (const double UsdTimeSample : UsdTimeSamples)
+		{
+			// We never want to evaluate the same time twice
+			if (FMath::IsNearlyEqual(UsdTimeSample, LastTimeSample))
+			{
+				continue;
+			}
+			LastTimeSample = UsdTimeSample;
+
+			int32 FrameNumber = FMath::FloorToInt(UsdTimeSample);
+			float SubFrameNumber = UsdTimeSample - FrameNumber;
+
+			FFrameTime FrameTime{FrameNumber, SubFrameNumber};
+			double FrameTimeSeconds = static_cast<float>(StageFrameRate.AsSeconds(FrameTime));
+
+			FLinearColor UEValue = ReaderFunc(UsdTimeSample);
+
+			FKeyHandle RHandle = RCurve.AddKey(FrameTimeSeconds, UEValue.R);
+			FKeyHandle GHandle = GCurve.AddKey(FrameTimeSeconds, UEValue.G);
+			FKeyHandle BHandle = BCurve.AddKey(FrameTimeSeconds, UEValue.B);
+			FKeyHandle AHandle = ACurve.AddKey(FrameTimeSeconds, UEValue.A);
+
+			RCurve.SetKeyInterpMode(RHandle, InterpMode);
+			GCurve.SetKeyInterpMode(GHandle, InterpMode);
+			BCurve.SetKeyInterpMode(BHandle, InterpMode);
+			ACurve.SetKeyInterpMode(AHandle, InterpMode);
+		}
+
+		return true;
+	}
+
+	bool ReadTransforms(
+		const UE::FUsdStage& UsdStage,
+		const TArray<double>& UsdTimeSamples,
+		const TFunction<FTransform(double)>& ReaderFunc,
+		Interchange::FAnimationPayloadData& OutPayloadData
+	)
+	{
+		OutPayloadData.Curves.SetNum(9);
+		FRichCurve& TransXCurve = OutPayloadData.Curves[0];
+		FRichCurve& TransYCurve = OutPayloadData.Curves[1];
+		FRichCurve& TransZCurve = OutPayloadData.Curves[2];
+		FRichCurve& RotXCurve = OutPayloadData.Curves[3];
+		FRichCurve& RotYCurve = OutPayloadData.Curves[4];
+		FRichCurve& RotZCurve = OutPayloadData.Curves[5];
+		FRichCurve& ScaleXCurve = OutPayloadData.Curves[6];
+		FRichCurve& ScaleYCurve = OutPayloadData.Curves[7];
+		FRichCurve& ScaleZCurve = OutPayloadData.Curves[8];
+
+		const FFrameRate StageFrameRate{static_cast<uint32>(UsdStage.GetTimeCodesPerSecond()), 1};
+		const ERichCurveInterpMode InterpMode = (UsdStage.GetInterpolationType() == EUsdInterpolationType::Linear)
+													? ERichCurveInterpMode::RCIM_Linear
+													: ERichCurveInterpMode::RCIM_Constant;
+
+		double LastTimeSample = TNumericLimits<double>::Lowest();
+		for (const double UsdTimeSample : UsdTimeSamples)
+		{
+			// We never want to evaluate the same time twice
+			if (FMath::IsNearlyEqual(UsdTimeSample, LastTimeSample))
+			{
+				continue;
+			}
+			LastTimeSample = UsdTimeSample;
+
+			int32 FrameNumber = FMath::FloorToInt(UsdTimeSample);
+			float SubFrameNumber = UsdTimeSample - FrameNumber;
+
+			FFrameTime FrameTime{FrameNumber, SubFrameNumber};
+			double FrameTimeSeconds = static_cast<float>(StageFrameRate.AsSeconds(FrameTime));
+
+			FTransform UEValue = ReaderFunc(UsdTimeSample);
+			FVector Location = UEValue.GetLocation();
+			FRotator Rotator = UEValue.Rotator();
+			FVector Scale = UEValue.GetScale3D();
+
+			FKeyHandle HandleTransX = TransXCurve.AddKey(FrameTimeSeconds, Location.X);
+			FKeyHandle HandleTransY = TransYCurve.AddKey(FrameTimeSeconds, Location.Y);
+			FKeyHandle HandleTransZ = TransZCurve.AddKey(FrameTimeSeconds, Location.Z);
+			FKeyHandle HandleRotX = RotXCurve.AddKey(FrameTimeSeconds, Rotator.Roll);
+			FKeyHandle HandleRotY = RotYCurve.AddKey(FrameTimeSeconds, Rotator.Pitch);
+			FKeyHandle HandleRotZ = RotZCurve.AddKey(FrameTimeSeconds, Rotator.Yaw);
+			FKeyHandle HandleScaleX = ScaleXCurve.AddKey(FrameTimeSeconds, Scale.X);
+			FKeyHandle HandleScaleY = ScaleYCurve.AddKey(FrameTimeSeconds, Scale.Y);
+			FKeyHandle HandleScaleZ = ScaleZCurve.AddKey(FrameTimeSeconds, Scale.Z);
+
+			TransXCurve.SetKeyInterpMode(HandleTransX, InterpMode);
+			TransYCurve.SetKeyInterpMode(HandleTransY, InterpMode);
+			TransZCurve.SetKeyInterpMode(HandleTransZ, InterpMode);
+			RotXCurve.SetKeyInterpMode(HandleRotX, InterpMode);
+			RotYCurve.SetKeyInterpMode(HandleRotY, InterpMode);
+			RotZCurve.SetKeyInterpMode(HandleRotZ, InterpMode);
+			ScaleXCurve.SetKeyInterpMode(HandleScaleX, InterpMode);
+			ScaleYCurve.SetKeyInterpMode(HandleScaleY, InterpMode);
+			ScaleZCurve.SetKeyInterpMode(HandleScaleZ, InterpMode);
+		}
+
+		return true;
@@ -881,163 +1111,4 @@
 		NodeContainer.AddNode(MeshNode);
 	}
 
-	void AddSkeletalAnimationNode(
-		const UE::FUsdSkelSkeletonQuery& SkeletonQuery,
-		const TMap<FString, TPair<FString, int32>>& BoneToUidAndBoneIndex,
-		UInterchangeUSDTranslatorImpl& TranslatorImpl,
-		UInterchangeSceneNode& SkeletonPrimNode,
-		UInterchangeBaseNodeContainer& NodeContainer,
-		const FTraversalInfo& Info
-	);
-
-	void AddSkeletonNodes(
-		const UE::FUsdPrim& Prim,
-		UInterchangeUSDTranslatorImpl& TranslatorImpl,
-		UInterchangeSceneNode& SkeletonPrimNode,
-		UInterchangeBaseNodeContainer& NodeContainer,
-		FTraversalInfo& Info
-	)
-	{
-		// If we're not inside of a SkelRoot, the skeleton shouldn't really do anything
-		if (!Info.FurthestSkelCache.IsValid())
-		{
-			return;
-		}
-
-		// By the time we get here we've already emitted a scene node for the skeleton prim itself, so we just
-		// need to emit a node hierarchy that mirrors the joints.
-
-		// Make the prim node into an Interchange joint/bone itself. By doing this we solve three issues:
-		//  - It becomes easy to identify our SkeletonDependencyUid when parsing Mesh nodes: It's just the skeleton prim path
-		//    (as opposed to having to target the translated node of the first root joint of the skeleton);
-		//  - We automatically handle USD skeletons with multiple root bones: We'll only ever have one "true"
-		//    root bone anyway: The SkeletonPrimNode itself;
-		//  - If a skeleton has no bones at all somehow, we'll still make one "bone" for it (this node).
-		SkeletonPrimNode.AddSpecializedType(UE::Interchange::FSceneNodeStaticData::GetJointSpecializeTypeString());
-		SkeletonPrimNode.SetCustomBindPoseLocalTransform(&NodeContainer, FTransform::Identity);
-		SkeletonPrimNode.SetCustomTimeZeroLocalTransform(&NodeContainer, FTransform::Identity);
-		const FString SkeletonPrimNodeUid = SkeletonPrimNode.GetUniqueID();
-
-#if WITH_EDITOR
-		// Convert the skeleton bones/joints into ConvertedData
-		UE::FUsdSkelSkeletonQuery SkelQuery = Info.FurthestSkelCache->GetSkelQuery(Prim);
-		const bool bEnsureAtLeastOneBone = false;
-		const bool bEnsureSingleRootBone = false;
-		UsdToUnreal::FUsdSkeletonData ConvertedData;
-		const bool bSuccess = UsdToUnreal::ConvertSkeleton(SkelQuery, ConvertedData, bEnsureAtLeastOneBone, bEnsureSingleRootBone);
-		if (!bSuccess)
-		{
-			return;
-		}
-
-		// Maps from the USD-style full bone name (e.g. "shoulder/elbow/hand") to the Uid we used for
-		// the corresponding scene node, and the bone's index on the skeleton's joint order.
-		// We'll need this to parse skeletal animations, if any
-		TMap<FString, TPair<FString, int32>> BoneToUidAndBoneIndex;
-
-		// Recursively traverse ConvertedData spawning the joint translated nodes
-		TFunction<void(int32, UInterchangeSceneNode&, const FString&)> RecursiveTraverseBones = nullptr;
-		RecursiveTraverseBones = [&RecursiveTraverseBones,
-								  &BoneToUidAndBoneIndex,
-								  &SkeletonPrimNodeUid,
-								  &ConvertedData,
-								  &NodeContainer	//
-		](int32 BoneIndex, UInterchangeSceneNode& ParentNode, const FString& BonePath)
-		{
-			const UsdToUnreal::FUsdSkeletonData::FBone& Bone = ConvertedData.Bones[BoneIndex];
-
-			// Reconcatenate a full "bone path" here for uniqueness, because Bone.Name is just the name of this
-			// single bone/joint itself (e.g. "Elbow")
-			const FString ConcatBonePath = (BonePath.IsEmpty() ? TEXT("") : (BonePath + TEXT("/"))) + Bone.Name;
-
-			// Putting the BonePrefix here avoids the pathological case where the user has skeleton child prims
-			// with names that match the joint names
-			const FString BoneNodeUid = SkeletonPrimNodeUid + BonePrefix + ConcatBonePath;
-
-			UInterchangeSceneNode* BoneNode = NewObject<UInterchangeSceneNode>(&NodeContainer);
-			BoneNode->InitializeNode(BoneNodeUid, Bone.Name, EInterchangeNodeContainerType::TranslatedScene);
-			BoneNode->AddSpecializedType(UE::Interchange::FSceneNodeStaticData::GetJointSpecializeTypeString());
-
-			// Note that we use our rest transforms for the Interchange bind pose as well: This because Interchange
-			// will put this on the RefSkeleton and so it will make its way to the Skeleton asset. We already kind
-			// of bake in our skeleton bind pose directly into our skinned mesh, so we really just want to put the
-			// rest pose on the skeleton asset/ReferenceSkeleton
-			BoneNode->SetCustomBindPoseLocalTransform(&NodeContainer, Bone.LocalRestTransform);
-			BoneNode->SetCustomTimeZeroLocalTransform(&NodeContainer, Bone.LocalRestTransform);
-			BoneNode->SetCustomLocalTransform(&NodeContainer, Bone.LocalRestTransform);
-
-			NodeContainer.AddNode(BoneNode);
-			NodeContainer.SetNodeParentUid(BoneNodeUid, ParentNode.GetUniqueID());
-
-			BoneToUidAndBoneIndex.Add(ConcatBonePath, {BoneNodeUid, BoneIndex});
-
-			for (int32 ChildIndex : Bone.ChildIndices)
-			{
-				RecursiveTraverseBones(ChildIndex, *BoneNode, ConcatBonePath);
-			}
-		};
-
-		// Start traversing from the root bones (we may have more than one)
-		for (int32 BoneIndex = 0; BoneIndex < ConvertedData.Bones.Num(); ++BoneIndex)
-		{
-			const UsdToUnreal::FUsdSkeletonData::FBone& Bone = ConvertedData.Bones[BoneIndex];
-			if (Bone.ParentIndex == INDEX_NONE)
-			{
-				const FString BonePathRoot = TEXT("");
-				RecursiveTraverseBones(BoneIndex, SkeletonPrimNode, BonePathRoot);
-			}
-		}
-
-		// Handle SkelAnimation prims, if we have any bound for this Skeleton
-		AddSkeletalAnimationNode(SkelQuery, BoneToUidAndBoneIndex, TranslatorImpl, SkeletonPrimNode, NodeContainer, Info);
-
-		// Cache our joint names in order, as this is needed when generating skeletal mesh payloads
-		Info.SkelJointNames = MakeShared<TArray<FString>>();
-		Info.SkelJointNames->Reserve(ConvertedData.Bones.Num());
-		for (const UsdToUnreal::FUsdSkeletonData::FBone& Bone : ConvertedData.Bones)
-		{
-			Info.SkelJointNames->Add(Bone.Name);
-		}
-		{
-			FWriteScopeLock ScopedInfoWriteLock{TranslatorImpl.CachedTraversalInfoLock};
-			TranslatorImpl.NodeUidToCachedTraversalInfo.Add(SkeletonPrimNodeUid, Info);
-		}
-#endif
-	}
-
-	void UpdateTraversalInfo(FTraversalInfo& Info, const UE::FUsdPrim& CurrentPrim)
-	{
-		if (CurrentPrim.IsA(TEXT("SkelRoot")))
-		{
-			if (!Info.ClosestParentSkelRoot)
-			{
-				// The root-most skel cache should handle any nested UsdSkel prims as well
-				Info.FurthestSkelCache = MakeShared<UE::FUsdSkelCache>();
-
-				const bool bTraverseInstanceProxies = true;
-				Info.FurthestSkelCache->Populate(CurrentPrim, bTraverseInstanceProxies);
-			}
-
-			Info.ClosestParentSkelRoot = CurrentPrim;
-		}
-
-		if (Info.ClosestParentSkelRoot && CurrentPrim.HasAPI(TEXT("SkelBindingAPI")))
-		{
-			UE::FUsdStage Stage = CurrentPrim.GetStage();
-
-			if (UE::FUsdRelationship SkelRel = CurrentPrim.GetRelationship(TEXT("skel:skeleton")))
-			{
-				TArray<UE::FSdfPath> Targets;
-				if (SkelRel.GetTargets(Targets) && Targets.Num() > 0)
-				{
-					UE::FUsdPrim TargetSkeleton = Stage.GetPrimAtPath(Targets[0]);
-					if (TargetSkeleton && TargetSkeleton.IsA(TEXT("Skeleton")))
-					{
-						Info.ActiveSkelQuery = Info.FurthestSkelCache->GetSkelQuery(TargetSkeleton);
-					}
-				}
-			}
-		}
-	}
-
 	void AddTrackSetNode(UInterchangeUSDTranslatorImpl& Impl, UInterchangeBaseNodeContainer& NodeContainer)
@@ -1375,6 +1446,121 @@
 		}
 	}
 
+	void AddSkeletonNodes(
+		const UE::FUsdPrim& Prim,
+		UInterchangeUSDTranslatorImpl& TranslatorImpl,
+		UInterchangeSceneNode& SkeletonPrimNode,
+		UInterchangeBaseNodeContainer& NodeContainer,
+		FTraversalInfo& Info
+	)
+	{
+		// If we're not inside of a SkelRoot, the skeleton shouldn't really do anything
+		if (!Info.FurthestSkelCache.IsValid())
+		{
+			return;
+		}
+
+		// By the time we get here we've already emitted a scene node for the skeleton prim itself, so we just
+		// need to emit a node hierarchy that mirrors the joints.
+
+		// Make the prim node into an Interchange joint/bone itself. By doing this we solve three issues:
+		//  - It becomes easy to identify our SkeletonDependencyUid when parsing Mesh nodes: It's just the skeleton prim path
+		//    (as opposed to having to target the translated node of the first root joint of the skeleton);
+		//  - We automatically handle USD skeletons with multiple root bones: We'll only ever have one "true"
+		//    root bone anyway: The SkeletonPrimNode itself;
+		//  - If a skeleton has no bones at all somehow, we'll still make one "bone" for it (this node).
+		SkeletonPrimNode.AddSpecializedType(UE::Interchange::FSceneNodeStaticData::GetJointSpecializeTypeString());
+		SkeletonPrimNode.SetCustomBindPoseLocalTransform(&NodeContainer, FTransform::Identity);
+		SkeletonPrimNode.SetCustomTimeZeroLocalTransform(&NodeContainer, FTransform::Identity);
+		const FString SkeletonPrimNodeUid = SkeletonPrimNode.GetUniqueID();
+
+#if WITH_EDITOR
+		// Convert the skeleton bones/joints into ConvertedData
+		UE::FUsdSkelSkeletonQuery SkelQuery = Info.FurthestSkelCache->GetSkelQuery(Prim);
+		const bool bEnsureAtLeastOneBone = false;
+		const bool bEnsureSingleRootBone = false;
+		UsdToUnreal::FUsdSkeletonData ConvertedData;
+		const bool bSuccess = UsdToUnreal::ConvertSkeleton(SkelQuery, ConvertedData, bEnsureAtLeastOneBone, bEnsureSingleRootBone);
+		if (!bSuccess)
+		{
+			return;
+		}
+
+		// Maps from the USD-style full bone name (e.g. "shoulder/elbow/hand") to the Uid we used for
+		// the corresponding scene node, and the bone's index on the skeleton's joint order.
+		// We'll need this to parse skeletal animations, if any
+		TMap<FString, TPair<FString, int32>> BoneToUidAndBoneIndex;
+
+		// Recursively traverse ConvertedData spawning the joint translated nodes
+		TFunction<void(int32, UInterchangeSceneNode&, const FString&)> RecursiveTraverseBones = nullptr;
+		RecursiveTraverseBones = [&RecursiveTraverseBones,
+								  &BoneToUidAndBoneIndex,
+								  &SkeletonPrimNodeUid,
+								  &ConvertedData,
+								  &NodeContainer	//
+		](int32 BoneIndex, UInterchangeSceneNode& ParentNode, const FString& BonePath)
+		{
+			const UsdToUnreal::FUsdSkeletonData::FBone& Bone = ConvertedData.Bones[BoneIndex];
+
+			// Reconcatenate a full "bone path" here for uniqueness, because Bone.Name is just the name of this
+			// single bone/joint itself (e.g. "Elbow")
+			const FString ConcatBonePath = (BonePath.IsEmpty() ? TEXT("") : (BonePath + TEXT("/"))) + Bone.Name;
+
+			// Putting the BonePrefix here avoids the pathological case where the user has skeleton child prims
+			// with names that match the joint names
+			const FString BoneNodeUid = SkeletonPrimNodeUid + BonePrefix + ConcatBonePath;
+
+			UInterchangeSceneNode* BoneNode = NewObject<UInterchangeSceneNode>(&NodeContainer);
+			BoneNode->InitializeNode(BoneNodeUid, Bone.Name, EInterchangeNodeContainerType::TranslatedScene);
+			BoneNode->AddSpecializedType(UE::Interchange::FSceneNodeStaticData::GetJointSpecializeTypeString());
+
+			// Note that we use our rest transforms for the Interchange bind pose as well: This because Interchange
+			// will put this on the RefSkeleton and so it will make its way to the Skeleton asset. We already kind
+			// of bake in our skeleton bind pose directly into our skinned mesh, so we really just want to put the
+			// rest pose on the skeleton asset/ReferenceSkeleton
+			BoneNode->SetCustomBindPoseLocalTransform(&NodeContainer, Bone.LocalRestTransform);
+			BoneNode->SetCustomTimeZeroLocalTransform(&NodeContainer, Bone.LocalRestTransform);
+			BoneNode->SetCustomLocalTransform(&NodeContainer, Bone.LocalRestTransform);
+
+			NodeContainer.AddNode(BoneNode);
+			NodeContainer.SetNodeParentUid(BoneNodeUid, ParentNode.GetUniqueID());
+
+			BoneToUidAndBoneIndex.Add(ConcatBonePath, {BoneNodeUid, BoneIndex});
+
+			for (int32 ChildIndex : Bone.ChildIndices)
+			{
+				RecursiveTraverseBones(ChildIndex, *BoneNode, ConcatBonePath);
+			}
+		};
+
+		// Start traversing from the root bones (we may have more than one)
+		for (int32 BoneIndex = 0; BoneIndex < ConvertedData.Bones.Num(); ++BoneIndex)
+		{
+			const UsdToUnreal::FUsdSkeletonData::FBone& Bone = ConvertedData.Bones[BoneIndex];
+			if (Bone.ParentIndex == INDEX_NONE)
+			{
+				const FString BonePathRoot = TEXT("");
+				RecursiveTraverseBones(BoneIndex, SkeletonPrimNode, BonePathRoot);
+			}
+		}
+
+		// Handle SkelAnimation prims, if we have any bound for this Skeleton
+		AddSkeletalAnimationNode(SkelQuery, BoneToUidAndBoneIndex, TranslatorImpl, SkeletonPrimNode, NodeContainer, Info);
+
+		// Cache our joint names in order, as this is needed when generating skeletal mesh payloads
+		Info.SkelJointNames = MakeShared<TArray<FString>>();
+		Info.SkelJointNames->Reserve(ConvertedData.Bones.Num());
+		for (const UsdToUnreal::FUsdSkeletonData::FBone& Bone : ConvertedData.Bones)
+		{
+			Info.SkelJointNames->Add(Bone.Name);
+		}
+		{
+			FWriteScopeLock ScopedInfoWriteLock{TranslatorImpl.CachedTraversalInfoLock};
+			TranslatorImpl.NodeUidToCachedTraversalInfo.Add(SkeletonPrimNodeUid, Info);
+		}
+#endif
+	}
+
 	void Traverse(
 		const UE::FUsdPrim& Prim,
 		UInterchangeUSDTranslatorImpl& TranslatorImpl,
@@ -1481,46 +1667,4 @@
 		}
 	}
 
-	void FixSkeletalMeshDescriptionColors(FMeshDescription& MeshDescription)
-	{
-		// FSkeletalMeshImportData::GetMeshDescription() will reinterpret our Wedge FColors as linear, and put those
-		// sRGB values disguised as linear into the mesh description. This also seems to disagree with the patch on
-		// cl 32791826, so here we have to fix that up and get our mesh description colors to be actually linear...
-		//
-		// This will hopefully go away once we have our own skinned mesh to FMeshDescription conversion function.
-		//
-		// Note: Weirdly enough skeletal meshes seem to put linear colors on VertexColor output, while static meshes
-		// put sRGB colors? Maybe this is why the comment above the change on 32791826 mentions to remove the ToFColor on
-		// StaticMeshBuilder? This is overall very confusing
-		FStaticMeshAttributes Attributes(MeshDescription);
-		TVertexInstanceAttributesRef<FVector4f> VertexColor = Attributes.GetVertexInstanceColors();
-		for (FVertexInstanceID VertexInstanceID : MeshDescription.VertexInstances().GetElementIDs())
-		{
-			const FColor ActualSRGB = FLinearColor(VertexColor[VertexInstanceID]).ToFColor(false);
-			VertexColor[VertexInstanceID] = FLinearColor{ActualSRGB};
-		}
-	}
-
-	void FixMaterialSlotNames(FMeshDescription& MeshDescription, const TArray<UsdUtils::FUsdPrimMaterialSlot>& MeshAssingmentSlots)
-	{
-		// Fixup material slot names to match the material that is assigned. For Interchange it is better to have the material
-		// slot names match what is assigned into them, as it will use those names to "merge identical slots" depending on the
-		// import options.
-		//
-		// Note: These names must also match what is set via MeshNode->SetSlotMaterialDependencyUid(SlotName, MaterialUid)
-		FStaticMeshAttributes StaticMeshAttributes(MeshDescription);
-		for (int32 MaterialSlotIndex = 0; MaterialSlotIndex < StaticMeshAttributes.GetPolygonGroupMaterialSlotNames().GetNumElements();
-			 ++MaterialSlotIndex)
-		{
-			int32 MaterialIndex = 0;
-			LexFromString(MaterialIndex, *StaticMeshAttributes.GetPolygonGroupMaterialSlotNames()[MaterialSlotIndex].ToString());
-
-			if (MeshAssingmentSlots.IsValidIndex(MaterialIndex))
-			{
-				const FString Source = MeshAssingmentSlots[MaterialIndex].MaterialSource;
-				StaticMeshAttributes.GetPolygonGroupMaterialSlotNames()[MaterialSlotIndex] = *Source;
-			}
-		}
-	}
-
 	bool GetStaticMeshPayloadData(
@@ -1701,206 +1845,6 @@
 		return UsdUtils::ApplyBlendShape(OutMeshDescription, BlendShape.GetPrim(), Weight, InbetweenName);
 	}
 
-	bool ReadBools(
-		const UE::FUsdStage& UsdStage,
-		const TArray<double>& UsdTimeSamples,
-		const TFunction<bool(double)>& ReaderFunc,
-		Interchange::FAnimationPayloadData& OutPayloadData
-	)
-	{
-		OutPayloadData.StepCurves.SetNum(1);
-		FInterchangeStepCurve& Curve = OutPayloadData.StepCurves[0];
-		TArray<float>& KeyTimes = Curve.KeyTimes;
-		TArray<bool>& BooleanKeyValues = Curve.BooleanKeyValues.Emplace();
-
-		KeyTimes.Reserve(UsdTimeSamples.Num());
-		BooleanKeyValues.Reserve(UsdTimeSamples.Num());
-
-		const FFrameRate StageFrameRate{static_cast<uint32>(UsdStage.GetTimeCodesPerSecond()), 1};
-
-		double LastTimeSample = TNumericLimits<double>::Lowest();
-		for (const double UsdTimeSample : UsdTimeSamples)
-		{
-			// We never want to evaluate the same time twice
-			if (FMath::IsNearlyEqual(UsdTimeSample, LastTimeSample))
-			{
-				continue;
-			}
-			LastTimeSample = UsdTimeSample;
-
-			int32 FrameNumber = FMath::FloorToInt(UsdTimeSample);
-			float SubFrameNumber = UsdTimeSample - FrameNumber;
-
-			FFrameTime FrameTime{FrameNumber, SubFrameNumber};
-			double FrameTimeSeconds = static_cast<float>(StageFrameRate.AsSeconds(FrameTime));
-
-			bool UEValue = ReaderFunc(UsdTimeSample);
-
-			KeyTimes.Add(FrameTimeSeconds);
-			BooleanKeyValues.Add(UEValue);
-		}
-
-		return true;
-	}
-
-	bool ReadFloats(
-		const UE::FUsdStage& UsdStage,
-		const TArray<double>& UsdTimeSamples,
-		const TFunction<float(double)>& ReaderFunc,
-		Interchange::FAnimationPayloadData& OutPayloadData
-	)
-	{
-		OutPayloadData.Curves.SetNum(1);
-		FRichCurve& Curve = OutPayloadData.Curves[0];
-
-		const FFrameRate StageFrameRate{static_cast<uint32>(UsdStage.GetTimeCodesPerSecond()), 1};
-		const ERichCurveInterpMode InterpMode = (UsdStage.GetInterpolationType() == EUsdInterpolationType::Linear)
-													? ERichCurveInterpMode::RCIM_Linear
-													: ERichCurveInterpMode::RCIM_Constant;
-
-		double LastTimeSample = TNumericLimits<double>::Lowest();
-		for (const double UsdTimeSample : UsdTimeSamples)
-		{
-			// We never want to evaluate the same time twice
-			if (FMath::IsNearlyEqual(UsdTimeSample, LastTimeSample))
-			{
-				continue;
-			}
-			LastTimeSample = UsdTimeSample;
-
-			int32 FrameNumber = FMath::FloorToInt(UsdTimeSample);
-			float SubFrameNumber = UsdTimeSample - FrameNumber;
-
-			FFrameTime FrameTime{FrameNumber, SubFrameNumber};
-			double FrameTimeSeconds = static_cast<float>(StageFrameRate.AsSeconds(FrameTime));
-
-			float UEValue = ReaderFunc(UsdTimeSample);
-
-			FKeyHandle Handle = Curve.AddKey(FrameTimeSeconds, UEValue);
-			Curve.SetKeyInterpMode(Handle, InterpMode);
-		}
-
-		return true;
-	}
-
-	bool ReadColors(
-		const UE::FUsdStage& UsdStage,
-		const TArray<double>& UsdTimeSamples,
-		const TFunction<FLinearColor(double)>& ReaderFunc,
-		Interchange::FAnimationPayloadData& OutPayloadData
-	)
-	{
-		OutPayloadData.Curves.SetNum(4);
-		FRichCurve& RCurve = OutPayloadData.Curves[0];
-		FRichCurve& GCurve = OutPayloadData.Curves[1];
-		FRichCurve& BCurve = OutPayloadData.Curves[2];
-		FRichCurve& ACurve = OutPayloadData.Curves[3];
-
-		const FFrameRate StageFrameRate{static_cast<uint32>(UsdStage.GetTimeCodesPerSecond()), 1};
-		const ERichCurveInterpMode InterpMode = (UsdStage.GetInterpolationType() == EUsdInterpolationType::Linear)
-													? ERichCurveInterpMode::RCIM_Linear
-													: ERichCurveInterpMode::RCIM_Constant;
-
-		double LastTimeSample = TNumericLimits<double>::Lowest();
-		for (const double UsdTimeSample : UsdTimeSamples)
-		{
-			// We never want to evaluate the same time twice
-			if (FMath::IsNearlyEqual(UsdTimeSample, LastTimeSample))
-			{
-				continue;
-			}
-			LastTimeSample = UsdTimeSample;
-
-			int32 FrameNumber = FMath::FloorToInt(UsdTimeSample);
-			float SubFrameNumber = UsdTimeSample - FrameNumber;
-
-			FFrameTime FrameTime{FrameNumber, SubFrameNumber};
-			double FrameTimeSeconds = static_cast<float>(StageFrameRate.AsSeconds(FrameTime));
-
-			FLinearColor UEValue = ReaderFunc(UsdTimeSample);
-
-			FKeyHandle RHandle = RCurve.AddKey(FrameTimeSeconds, UEValue.R);
-			FKeyHandle GHandle = GCurve.AddKey(FrameTimeSeconds, UEValue.G);
-			FKeyHandle BHandle = BCurve.AddKey(FrameTimeSeconds, UEValue.B);
-			FKeyHandle AHandle = ACurve.AddKey(FrameTimeSeconds, UEValue.A);
-
-			RCurve.SetKeyInterpMode(RHandle, InterpMode);
-			GCurve.SetKeyInterpMode(GHandle, InterpMode);
-			BCurve.SetKeyInterpMode(BHandle, InterpMode);
-			ACurve.SetKeyInterpMode(AHandle, InterpMode);
-		}
-
-		return true;
-	}
-
-	bool ReadTransforms(
-		const UE::FUsdStage& UsdStage,
-		const TArray<double>& UsdTimeSamples,
-		const TFunction<FTransform(double)>& ReaderFunc,
-		Interchange::FAnimationPayloadData& OutPayloadData
-	)
-	{
-		OutPayloadData.Curves.SetNum(9);
-		FRichCurve& TransXCurve = OutPayloadData.Curves[0];
-		FRichCurve& TransYCurve = OutPayloadData.Curves[1];
-		FRichCurve& TransZCurve = OutPayloadData.Curves[2];
-		FRichCurve& RotXCurve = OutPayloadData.Curves[3];
-		FRichCurve& RotYCurve = OutPayloadData.Curves[4];
-		FRichCurve& RotZCurve = OutPayloadData.Curves[5];
-		FRichCurve& ScaleXCurve = OutPayloadData.Curves[6];
-		FRichCurve& ScaleYCurve = OutPayloadData.Curves[7];
-		FRichCurve& ScaleZCurve = OutPayloadData.Curves[8];
-
-		const FFrameRate StageFrameRate{static_cast<uint32>(UsdStage.GetTimeCodesPerSecond()), 1};
-		const ERichCurveInterpMode InterpMode = (UsdStage.GetInterpolationType() == EUsdInterpolationType::Linear)
-													? ERichCurveInterpMode::RCIM_Linear
-													: ERichCurveInterpMode::RCIM_Constant;
-
-		double LastTimeSample = TNumericLimits<double>::Lowest();
-		for (const double UsdTimeSample : UsdTimeSamples)
-		{
-			// We never want to evaluate the same time twice
-			if (FMath::IsNearlyEqual(UsdTimeSample, LastTimeSample))
-			{
-				continue;
-			}
-			LastTimeSample = UsdTimeSample;
-
-			int32 FrameNumber = FMath::FloorToInt(UsdTimeSample);
-			float SubFrameNumber = UsdTimeSample - FrameNumber;
-
-			FFrameTime FrameTime{FrameNumber, SubFrameNumber};
-			double FrameTimeSeconds = static_cast<float>(StageFrameRate.AsSeconds(FrameTime));
-
-			FTransform UEValue = ReaderFunc(UsdTimeSample);
-			FVector Location = UEValue.GetLocation();
-			FRotator Rotator = UEValue.Rotator();
-			FVector Scale = UEValue.GetScale3D();
-
-			FKeyHandle HandleTransX = TransXCurve.AddKey(FrameTimeSeconds, Location.X);
-			FKeyHandle HandleTransY = TransYCurve.AddKey(FrameTimeSeconds, Location.Y);
-			FKeyHandle HandleTransZ = TransZCurve.AddKey(FrameTimeSeconds, Location.Z);
-			FKeyHandle HandleRotX = RotXCurve.AddKey(FrameTimeSeconds, Rotator.Roll);
-			FKeyHandle HandleRotY = RotYCurve.AddKey(FrameTimeSeconds, Rotator.Pitch);
-			FKeyHandle HandleRotZ = RotZCurve.AddKey(FrameTimeSeconds, Rotator.Yaw);
-			FKeyHandle HandleScaleX = ScaleXCurve.AddKey(FrameTimeSeconds, Scale.X);
-			FKeyHandle HandleScaleY = ScaleYCurve.AddKey(FrameTimeSeconds, Scale.Y);
-			FKeyHandle HandleScaleZ = ScaleZCurve.AddKey(FrameTimeSeconds, Scale.Z);
-
-			TransXCurve.SetKeyInterpMode(HandleTransX, InterpMode);
-			TransYCurve.SetKeyInterpMode(HandleTransY, InterpMode);
-			TransZCurve.SetKeyInterpMode(HandleTransZ, InterpMode);
-			RotXCurve.SetKeyInterpMode(HandleRotX, InterpMode);
-			RotYCurve.SetKeyInterpMode(HandleRotY, InterpMode);
-			RotZCurve.SetKeyInterpMode(HandleRotZ, InterpMode);
-			ScaleXCurve.SetKeyInterpMode(HandleScaleX, InterpMode);
-			ScaleYCurve.SetKeyInterpMode(HandleScaleY, InterpMode);
-			ScaleZCurve.SetKeyInterpMode(HandleScaleZ, InterpMode);
-		}
-
-		return true;
-	}
-
 	bool GetPropertyAnimationCurvePayloadData(
 		const UE::FUsdStage& UsdStage,
 		const FString& PayloadKey,
@@ -2234,8 +2178,46 @@
 
 		return true;
 	}
-}	 // namespace UE::InterchangeUsdTranslator::Private
 #endif	  // USE_USD_SDK
+}	 // namespace UE::InterchangeUsdTranslator::Private
 
+UInterchangeUsdTranslatorSettings::UInterchangeUsdTranslatorSettings()
+	: GeometryPurpose((int32)(EUsdPurpose::Default | EUsdPurpose::Proxy | EUsdPurpose::Render | EUsdPurpose::Guide))
+	, RenderContext(TEXT("unreal"))	   // The proper definition of this is on the USDSchemas module, which we can't depend on
+	, MaterialPurpose(*UnrealIdentifiers::MaterialPreviewPurpose)
+	, InterpolationType(EUsdInterpolationType::Linear)
+	, bOverrideStageOptions(false)
+	, StageOptions{
+		  0.01f,			   // MetersPerUnit
+		  EUsdUpAxis::ZAxis	   // UpAxis
+	  }
+{
+}
+
+UInterchangeUSDTranslator::UInterchangeUSDTranslator()
+	: Impl(MakeUnique<UE::InterchangeUsdTranslator::Private::UInterchangeUSDTranslatorImpl>())
+{
+}
+
+EInterchangeTranslatorType UInterchangeUSDTranslator::GetTranslatorType() const
+{
+	return GInterchangeEnableUSDLevelImport ? EInterchangeTranslatorType::Scenes : EInterchangeTranslatorType::Assets;
+}
+
+EInterchangeTranslatorAssetType UInterchangeUSDTranslator::GetSupportedAssetTypes() const
+{
+	return EInterchangeTranslatorAssetType::Materials | EInterchangeTranslatorAssetType::Meshes | EInterchangeTranslatorAssetType::Animations;
+}
+
+TArray<FString> UInterchangeUSDTranslator::GetSupportedFormats() const
+{
+	TArray<FString> Extensions;
+	if (GInterchangeEnableUSDImport)
+	{
+		UnrealUSDWrapper::AddUsdImportFileFormatDescriptions(Extensions);
+	}
+	return Extensions;
+}
+
 bool UInterchangeUSDTranslator::Translate(UInterchangeBaseNodeContainer& NodeContainer) const
 {
