USD: Fix Mac build broken with cl 34761413.

diff --git a/Engine/Plugins/Interchange/Runtime/Source/Import/Private/Usd/InterchangeUsdTranslator.cpp b/Engine/Plugins/Interchange/Runtime/Source/Import/Private/Usd/InterchangeUsdTranslator.cpp
--- a/Engine/Plugins/Interchange/Runtime/Source/Import/Private/Usd/InterchangeUsdTranslator.cpp
+++ b/Engine/Plugins/Interchange/Runtime/Source/Import/Private/Usd/InterchangeUsdTranslator.cpp
@@ -2149,7 +2149,9 @@ namespace UE::InterchangeUsdTranslator::Private
 				}
 
 				const FString MorphTargetName = BlendShapeName + TEXT("_") + Inbetween.GetAttr().GetName().ToString();
-				ParsedInbetweens.Emplace(MorphTargetName, Position);
+				FInbetweenAndPosition& NewEntry = ParsedInbetweens.Emplace_GetRef();
+				NewEntry.Name = MorphTargetName;
+				NewEntry.Position = Position;
 			}
 
 			ParsedInbetweens.Sort(
