[Interchange] SkeletalMesh import fixes for T0 usage.
Interchange: Mesh is mirror with the skeleton (180 deg X axis rotation).

#rnx

diff --git a/Engine/Plugins/Interchange/Runtime/Source/Import/Private/Mesh/InterchangeSkeletalMeshFactory.cpp b/Engine/Plugins/Interchange/Runtime/Source/Import/Private/Mesh/InterchangeSkeletalMeshFactory.cpp
--- a/Engine/Plugins/Interchange/Runtime/Source/Import/Private/Mesh/InterchangeSkeletalMeshFactory.cpp
+++ b/Engine/Plugins/Interchange/Runtime/Source/Import/Private/Mesh/InterchangeSkeletalMeshFactory.cpp
@@ -328,10 +328,13 @@ namespace UE
 				return nullptr;
 			}
 
+			//We assume normalize weight method in this bind pose conversion
 			void SkinVertexPositionToTimeZero(UE::Interchange::FMeshPayloadData& LodMeshPayload
 				, const UInterchangeBaseNodeContainer* NodeContainer
 				, const FString& RootJointNodeId
-				, const FTransform& MeshGlobalTransform)
+				, const UInterchangeMeshNode* MeshNode
+				, const UInterchangeSceneNode* SceneNode
+				, const FTransform& SceneNodeTransform)
 			{
 				TRACE_CPUPROFILER_EVENT_SCOPE(SkinVertexPositionToTimeZero)
 				FMeshDescription& MeshDescription = LodMeshPayload.MeshDescription;
@@ -364,9 +367,6 @@ namespace UE
 					CommonPipelineDataFactoryNode->GetCustomGlobalOffsetTransform(GlobalOffsetTransform);
 				}
 
-				//We assume normalize weight method in this bind pose conversion
-
-				const FTransform MeshGlobalTransformInverse = MeshGlobalTransform.Inverse();
 				const int32 JointCount = JointNames.Num();
 				for (int32 JointIndex = 0; JointIndex < JointCount; ++JointIndex)
 				{
@@ -377,26 +377,29 @@ namespace UE
 					{
 						continue;
 					}
-					
-					FTransform JointBindPoseGlobalTransform;
-					if (!ensure(JointNode->GetCustomBindPoseGlobalTransform(NodeContainer, GlobalOffsetTransform, JointBindPoseGlobalTransform)))
+
+					FTransform BindPose;
+					if (!ensure(JointNode->GetCustomBindPoseGlobalTransform(NodeContainer, FTransform::Identity, BindPose)))
 					{
 						//BindPose will fall back on LocalTransform in case its not present.
 						//If neither is present: No Value to convert from, skip this joint.
 						continue;
 					}
 
-					FTransform JointTimeZeroGlobalTransform;
-					if (!JointNode->GetCustomTimeZeroGlobalTransform(NodeContainer, GlobalOffsetTransform, JointTimeZeroGlobalTransform))
+					FTransform T0;
+					if (!JointNode->GetCustomTimeZeroGlobalTransform(NodeContainer, FTransform::Identity, T0))
 					{
 						//If there is no time zero global transform we cannot set the bind pose to time zero.
 						//We must skip this joint.
 						continue;
 					}
 
-					//JointBindPoseGlobalTransform Contains the GlobalTransform as the GlobalTransform calculation fallsback onto Transform (from BindPoseTransform or T0Transforms if they are not present)
-					//Multiply both transform to get a matrix that will transform the mesh vertices from the bind pose skinning to the time zero skinning
-					const FMatrix VertexTransformMatrix = (JointBindPoseGlobalTransform.Inverse() * JointTimeZeroGlobalTransform).ToMatrixWithScale();
+					FMatrix TransformMatrix = FMatrix::Identity;
+					JointNode->GetGlobalBindPoseReferenceForMeshUID(MeshNode->GetUniqueID(), TransformMatrix);
+					FTransform TransformMatrixTransform;
+					TransformMatrixTransform.SetFromMatrix(TransformMatrix);
+
+					FMatrix VertexTransformMatrix = (SceneNodeTransform.Inverse() * ((TransformMatrixTransform * BindPose.Inverse()) * (T0 * SceneNodeTransform.Inverse()))).ToMatrixWithScale();
 
 					//Iterate all bone vertices
 					for (FVertexID VertexID : MeshDescription.Vertices().GetElementIDs())
@@ -443,7 +446,7 @@ namespace UE
 						//Normalized, in case the weight is different then 1
 						lDstVertex /= Weight;
 						//Set the new vertex position in the mesh description
-						VertexPositions[VertexID] = lDstVertex;
+						VertexPositions[VertexID] = FVector3f(FVector4f(SceneNodeTransform.TransformFVector4(FVector(lDstVertex))));
 					}
 				}
 			}
@@ -546,12 +549,7 @@ namespace UE
 					{
 						//We need to rebind the mesh at time 0. Skeleton joint have the time zero transform, so we need to apply the skinning to the mesh
 						//With the skeleton transform at time zero
-						FTransform MeshGlobalTransform = FTransform::Identity;
-						if (MeshNodeContext.SceneGlobalTransform.IsSet())
-						{
-							MeshGlobalTransform = MeshNodeContext.SceneGlobalTransform.GetValue();
-						}
-						SkinVertexPositionToTimeZero(LodMeshPayload.GetValue(), NodeContainer, RootJointNodeId, MeshGlobalTransform);
+						SkinVertexPositionToTimeZero(LodMeshPayload.GetValue(), NodeContainer, RootJointNodeId, MeshNodeContext.MeshNode, MeshNodeContext.SceneNode, MeshNodeContext.SceneGlobalTransform.Get(FTransform::Identity));
 					}
 
 					const int32 RefBoneCount = RefBonesBinary.Num();
diff --git a/Engine/Plugins/Interchange/Runtime/Source/Nodes/Private/InterchangeSceneNode.cpp b/Engine/Plugins/Interchange/Runtime/Source/Nodes/Private/InterchangeSceneNode.cpp
--- a/Engine/Plugins/Interchange/Runtime/Source/Nodes/Private/InterchangeSceneNode.cpp
+++ b/Engine/Plugins/Interchange/Runtime/Source/Nodes/Private/InterchangeSceneNode.cpp
@@ -47,6 +47,14 @@ namespace UE
 			return SlotMaterialDependenciesString;
 		}
 
+		const FString& FSceneNodeStaticData::GetMeshToGlobalBindPoseReferencesString()
+		{
+			static FString MeshToGlobalBindPoseReferncesString(TEXT("__MeshToGlobalBindPoseReferences__"));
+			return MeshToGlobalBindPoseReferncesString;
+		}
+
+		
+
 		const FString& FSceneNodeStaticData::GetMorphTargetCurveWeightsKey()
 		{
 			static FString MorphTargetCurvesKey(TEXT("__MorphTargetCurveWeights__Key"));
@@ -58,6 +66,7 @@ namespace UE
 UInterchangeSceneNode::UInterchangeSceneNode()
 {
 	NodeSpecializeTypes.Initialize(Attributes, UE::Interchange::FSceneNodeStaticData::GetNodeSpecializeTypeBaseKey().ToString());
+	MeshToGlobalBindPoseReferences.Initialize(Attributes.ToSharedRef(), UE::Interchange::FSceneNodeStaticData::GetMeshToGlobalBindPoseReferencesString());
 	SlotMaterialDependencies.Initialize(Attributes.ToSharedRef(), UE::Interchange::FSceneNodeStaticData::GetSlotMaterialDependenciesString());
 	MorphTargetCurveWeights.Initialize(Attributes.ToSharedRef(), UE::Interchange::FSceneNodeStaticData::GetMorphTargetCurveWeightsKey());
 }
@@ -416,4 +425,17 @@ bool UInterchangeSceneNode::SetCustomAnimationAssetUidToPlay(const FString& Attr
 bool UInterchangeSceneNode::GetCustomAnimationAssetUidToPlay(FString& AttributeValue) const
 {
 	IMPLEMENT_NODE_ATTRIBUTE_GETTER(AnimationAssetUidToPlay, FString);
-}
\ No newline at end of file
+}
+
+bool UInterchangeSceneNode::GetGlobalBindPoseReferenceForMeshUID(const FString& MeshUID, FMatrix& GlobalBindPoseReference) const
+{
+	return MeshToGlobalBindPoseReferences.GetValue(MeshUID, GlobalBindPoseReference);
+}
+
+void UInterchangeSceneNode::SetGlobalBindPoseReferenceForMeshUIDs(const TMap<FString, FMatrix>& GlobalBindPoseReferenceForMeshUIDs)
+{
+	for (const TPair<FString, FMatrix>& Entry: GlobalBindPoseReferenceForMeshUIDs)
+	{
+		MeshToGlobalBindPoseReferences.SetKeyValue(Entry.Key, Entry.Value);
+	}
+}
diff --git a/Engine/Plugins/Interchange/Runtime/Source/Nodes/Public/InterchangeSceneNode.h b/Engine/Plugins/Interchange/Runtime/Source/Nodes/Public/InterchangeSceneNode.h
--- a/Engine/Plugins/Interchange/Runtime/Source/Nodes/Public/InterchangeSceneNode.h
+++ b/Engine/Plugins/Interchange/Runtime/Source/Nodes/Public/InterchangeSceneNode.h
@@ -22,6 +22,7 @@ namespace UE
 			static const FString& GetTransformSpecializeTypeString();
 			static const FString& GetJointSpecializeTypeString();
 			static const FString& GetLodGroupSpecializeTypeString();
+			static const FString& GetMeshToGlobalBindPoseReferencesString();
 			static const FString& GetSlotMaterialDependenciesString();
 			static const FString& GetMorphTargetCurveWeightsKey();
 		};
@@ -221,6 +222,16 @@ public:
 	UFUNCTION(BlueprintCallable, Category = "Interchange | Node | SkeletalMesh")
 	bool GetCustomAnimationAssetUidToPlay(FString& AttributeValue) const;
 
+	/** Set the Global Bind Pose Referenced for MeshUIDs. */
+	UFUNCTION(BlueprintCallable, Category = "Interchange | Node | Joint")
+	void SetGlobalBindPoseReferenceForMeshUIDs(const TMap<FString, FMatrix>& GlobalBindPoseReferenceForMeshUIDs);
+
+	/** Get the Global Bind Pose Reference for given MeshUID. */
+	UFUNCTION(BlueprintCallable, Category = "Interchange | Node | Joint")
+	bool GetGlobalBindPoseReferenceForMeshUID(const FString& MeshUID, FMatrix& GlobalBindPoseReference) const;
+
+
+
 private:
 
 	bool GetGlobalTransformInternal(const UE::Interchange::FAttributeKey LocalTransformKey
@@ -247,7 +258,7 @@ private:
 
 	//A scene node can reference an asset. Asset can be Mesh, Light, camera...
 	const UE::Interchange::FAttributeKey Macro_CustomAssetInstanceUidKey = UE::Interchange::FAttributeKey(TEXT("AssetInstanceUid"));
-	
+
 	//A scene node can represent many special types
 	UE::Interchange::TArrayAttributeHelper<FString> NodeSpecializeTypes;
 
@@ -257,6 +268,9 @@ private:
 	//A scene node can have different MorphTarget curve settings:
 	UE::Interchange::TMapAttributeHelper<FString, float> MorphTargetCurveWeights;
 
+	//BindPose References per Mesh for a JointNode.
+	UE::Interchange::TMapAttributeHelper<FString, FMatrix> MeshToGlobalBindPoseReferences;
+
 	//A scene node can reference an animation asset on top of base asset:
 	const UE::Interchange::FAttributeKey Macro_CustomAnimationAssetUidToPlayKey = UE::Interchange::FAttributeKey(TEXT("AnimationAssetUidToPlay"));
 
diff --git a/Engine/Plugins/Interchange/Runtime/Source/Parsers/Fbx/Private/FbxMesh.cpp b/Engine/Plugins/Interchange/Runtime/Source/Parsers/Fbx/Private/FbxMesh.cpp
--- a/Engine/Plugins/Interchange/Runtime/Source/Parsers/Fbx/Private/FbxMesh.cpp
+++ b/Engine/Plugins/Interchange/Runtime/Source/Parsers/Fbx/Private/FbxMesh.cpp
@@ -1991,7 +1991,7 @@ void FFbxMesh::AddAllMeshes(FbxScene* SDKScene, FbxGeometryConverter* SDKGeometr
 	} // for GeometryCount
 }
 
-bool FFbxMesh::GetGlobalJointBindPoseTransform(FFbxParser* Parser, FbxScene* SDKScene, FbxNode* Joint, FbxAMatrix& GlobalBindPoseJointMatrix, bool& bBadBindPoseMessageDisplay)
+bool FFbxMesh::GetGlobalJointBindPoseTransform(FFbxParser* Parser, FbxScene* SDKScene, FbxNode* Joint, FbxAMatrix& GlobalBindPoseJointMatrix, TMap<FString, FMatrix>& MeshIdToGlobalBindPoseReferenceMap, bool& bBadBindPoseMessageDisplay)
 {
 	FbxManager* SDKManager = SDKScene->GetFbxManager();
 	//First look for Cluster and then look in the bind pose if no cluster was found
@@ -2001,11 +2001,13 @@ bool FFbxMesh::GetGlobalJointBindPoseTransform(FFbxParser* Parser, FbxScene* SDK
 	const int32 GeometryCount = SDKScene->GetGeometryCount();
 	for (int32 GeometryIndex = 0; GeometryIndex < GeometryCount; ++GeometryIndex)
 	{
-		const FbxGeometry* Geometry = SDKScene->GetGeometry(GeometryIndex);
+		FbxGeometry* Geometry = SDKScene->GetGeometry(GeometryIndex);
 		if (!ensure(Geometry))
 		{
 			continue;
 		}
+		FString MeshUniqueID = Parser->GetFbxHelper()->GetMeshUniqueID(Geometry);
+
 		const int32 GeometryDeformerCount = Geometry->GetDeformerCount(FbxDeformer::eSkin);
 		for (int32 GeometryDeformerIndex = 0; GeometryDeformerIndex < GeometryDeformerCount; ++GeometryDeformerIndex)
 		{
@@ -2027,12 +2029,26 @@ bool FFbxMesh::GetGlobalJointBindPoseTransform(FFbxParser* Parser, FbxScene* SDK
 				}
 				if (Joint == Cluster->GetLink())
 				{
+					uint64 ClusterId = Cluster->GetUniqueID();
+					uint64 MeshId = Geometry->GetUniqueID();
+
 					Cluster->GetTransformLinkMatrix(GlobalBindPoseJointMatrix);
-					return true;
+
+					FbxAMatrix GlobalBindPoseReferenceMatrix;
+					Cluster->GetTransformMatrix(GlobalBindPoseReferenceMatrix);
+					
+					MeshIdToGlobalBindPoseReferenceMap.Add(MeshUniqueID, FFbxConvert::ConvertMatrix<FMatrix>(GlobalBindPoseReferenceMatrix));
+
+					//Presumes only 1 cluster per geometry that matches criteria:
+					break;
 				}
 			}
 		}
 	}
+	if (MeshIdToGlobalBindPoseReferenceMap.Num() > 0)
+	{
+		return true;
+	}
 
 	auto AcquireBindPoseMatrix = [](FbxPose* CurrentPose, FbxAMatrix& GlobalBindPoseJointMatrix, FbxNode* Joint)
 	{
diff --git a/Engine/Plugins/Interchange/Runtime/Source/Parsers/Fbx/Private/FbxMesh.h b/Engine/Plugins/Interchange/Runtime/Source/Parsers/Fbx/Private/FbxMesh.h
--- a/Engine/Plugins/Interchange/Runtime/Source/Parsers/Fbx/Private/FbxMesh.h
+++ b/Engine/Plugins/Interchange/Runtime/Source/Parsers/Fbx/Private/FbxMesh.h
@@ -196,7 +196,7 @@ namespace UE
 				{}
 
 				void AddAllMeshes(FbxScene* SDKScene, FbxGeometryConverter* SDKGeometryConverter, UInterchangeBaseNodeContainer& NodeContainer, TMap<FString, TSharedPtr<FPayloadContextBase>>& PayloadContexts);
-				static bool GetGlobalJointBindPoseTransform(FFbxParser* Parser, FbxScene* SDKScene, FbxNode* Joint, FbxAMatrix& GlobalBindPoseJointMatrix, bool& bBadBindPoseMessageDisplay);
+				static bool GetGlobalJointBindPoseTransform(FFbxParser* Parser, FbxScene* SDKScene, FbxNode* Joint, FbxAMatrix& GlobalBindPoseJointMatrix, TMap<FString, FMatrix>& MeshIdToGlobalBindPoseReferenceMap, bool& bBadBindPoseMessageDisplay);
 
 				TArray<FMorphTargetAnimationBuildingData>& GetMorphTargetAnimationsBuildingData() { return MorphTargetAnimationsBuildingData; }
 			protected:
diff --git a/Engine/Plugins/Interchange/Runtime/Source/Parsers/Fbx/Private/FbxScene.cpp b/Engine/Plugins/Interchange/Runtime/Source/Parsers/Fbx/Private/FbxScene.cpp
--- a/Engine/Plugins/Interchange/Runtime/Source/Parsers/Fbx/Private/FbxScene.cpp
+++ b/Engine/Plugins/Interchange/Runtime/Source/Parsers/Fbx/Private/FbxScene.cpp
@@ -217,34 +217,50 @@ namespace UE
 					//Add the joint specialized type
 					UnrealNode->AddSpecializedType(FSceneNodeStaticData::GetJointSpecializeTypeString());
 					//Get the bind pose transform for this joint
-					FbxAMatrix GlobalBindPoseJointMatrix;
-					if (FFbxMesh::GetGlobalJointBindPoseTransform(&Parser, SDKScene, Node, GlobalBindPoseJointMatrix, bBadBindPoseMessageDisplay))
+					FbxAMatrix GlobalBindPoseJointMatrix = SDKScene->GetAnimationEvaluator()->GetNodeGlobalTransform(Node, 0);
+					TMap<FString, FMatrix> MeshIdToGlobalBindPoseReferenceMap;
+
+					FFbxMesh::GetGlobalJointBindPoseTransform(&Parser, SDKScene, Node, GlobalBindPoseJointMatrix, MeshIdToGlobalBindPoseReferenceMap, bBadBindPoseMessageDisplay); //if false it would use the previously calculated
+
+					FTransform GlobalBindPoseJointTransform = GetConvertedTransform(GlobalBindPoseJointMatrix);
+					UnrealNode->SetGlobalBindPoseReferenceForMeshUIDs(MeshIdToGlobalBindPoseReferenceMap);
+
+					FbxNode* ParentNode = Node->GetParent();
+					
+					if (ParentNode != nullptr)
 					{
-						FTransform GlobalBindPoseJointTransform = GetConvertedTransform(GlobalBindPoseJointMatrix);
-						//We grab the fbx parent node to compute the local transform
-						if (FbxNode* ParentNode = Node->GetParent())
-						{
-							FbxAMatrix GlobalFbxParentMatrix = ParentNode->EvaluateGlobalTransform();
-							FFbxMesh::GetGlobalJointBindPoseTransform(&Parser, SDKScene, ParentNode, GlobalFbxParentMatrix, bBadBindPoseMessageDisplay);
-							FbxAMatrix	LocalFbxMatrix = GlobalFbxParentMatrix.Inverse() * GlobalBindPoseJointMatrix;
-							FTransform LocalBindPoseJointTransform = GetConvertedTransform(LocalFbxMatrix);
-							UnrealNode->SetCustomBindPoseLocalTransform(&NodeContainer, LocalBindPoseJointTransform, bResetCache);
-						}
-						else
-						{
-							//No parent, set the same matrix has the global
-							UnrealNode->SetCustomBindPoseLocalTransform(&NodeContainer, GlobalBindPoseJointTransform, bResetCache);
-						}
+						FbxAMatrix GlobalFbxParentMatrix = SDKScene->GetAnimationEvaluator()->GetNodeGlobalTransform(ParentNode, 0);
+						TMap<FString, FMatrix> ParentMeshIdToGlobalBindPoseReferenceMap;
+						FFbxMesh::GetGlobalJointBindPoseTransform(&Parser, SDKScene, ParentNode, GlobalFbxParentMatrix, ParentMeshIdToGlobalBindPoseReferenceMap, bBadBindPoseMessageDisplay);
+						
+						FbxAMatrix LocalFbxMatrix = GlobalFbxParentMatrix.Inverse() * GlobalBindPoseJointMatrix;
+						FTransform LocalBindPoseJointTransform = GetConvertedTransform(LocalFbxMatrix);
+
+						UnrealNode->SetCustomBindPoseLocalTransform(&NodeContainer, LocalBindPoseJointTransform, bResetCache);
+					}
+					else
+					{
+						//No parent, set the same matrix has the global
+						UnrealNode->SetCustomBindPoseLocalTransform(&NodeContainer, GlobalBindPoseJointTransform, bResetCache);
 					}
 
 					//Get time Zero transform for this joint
 					{
+						//NOTE:
+						// Legacy FBX uses the following Matrix calculation for moving Vertices to T0:
+						//		VertexTransformMatrix = ((TransformMatrix * BindPose.Inverse()) * (T0 * GlobalMeshTransformMatrix.Inverse()));
+						//					TransformMatrix				:= GlobalBindPoseReferenceForMeshUIDs (this is joint and Mesh dependent) => seems very FBX specific
+						//					BindPose					:= GlobalBindPose
+						//					T0							:= TimeZero
+						//					GlobalMeshTransformMatrix	:= Mesh's Node's GlobalTransform * GeometricTransform (Interchange.SceneNodeTransform)
+
 						//Set the global node transform
-						FbxAMatrix GlobalFbxMatrix = Node->EvaluateGlobalTransform(FBXSDK_TIME_ZERO);
+						FbxAMatrix GlobalFbxMatrix = SDKScene->GetAnimationEvaluator()->GetNodeGlobalTransform(Node, 0);
 						FTransform GlobalTransform = GetConvertedTransform(GlobalFbxMatrix);
-						if (FbxNode* ParentNode = Node->GetParent())
+
+						if (ParentNode != nullptr)
 						{
-							FbxAMatrix GlobalFbxParentMatrix = ParentNode->EvaluateGlobalTransform(FBXSDK_TIME_ZERO);
+							FbxAMatrix GlobalFbxParentMatrix = SDKScene->GetAnimationEvaluator()->GetNodeGlobalTransform(ParentNode, 0);
 							FbxAMatrix	LocalFbxMatrix = GlobalFbxParentMatrix.Inverse() * GlobalFbxMatrix;
 							FTransform LocalTransform = GetConvertedTransform(LocalFbxMatrix);
 							UnrealNode->SetCustomTimeZeroLocalTransform(&NodeContainer, LocalTransform, bResetCache);
@@ -321,6 +337,7 @@ namespace UE
 							Geometry.SetT(Translation);
 							Geometry.SetR(Rotation);
 							Geometry.SetS(Scaling);
+
 							FTransform GeometricTransform = GetConvertedTransform(Geometry);
 
 							//Get the pivot geometry offset 
