[Interchange] Clearing Collisions on ReImport (if ImportCollision is turned on).

#rnx

diff --git a/Engine/Plugins/Interchange/Runtime/Source/Import/Private/Mesh/InterchangeStaticMeshFactory.cpp b/Engine/Plugins/Interchange/Runtime/Source/Import/Private/Mesh/InterchangeStaticMeshFactory.cpp
--- a/Engine/Plugins/Interchange/Runtime/Source/Import/Private/Mesh/InterchangeStaticMeshFactory.cpp
+++ b/Engine/Plugins/Interchange/Runtime/Source/Import/Private/Mesh/InterchangeStaticMeshFactory.cpp
@@ -429,6 +429,10 @@ UInterchangeFactoryBase::FImportAssetResult UInterchangeStaticMeshFactory::Impor
 			LodDataNode->GetImportCollision(bImportCollision);
 			if(bImportCollision)
 			{
+				if (bReimport)
+				{
+					StaticMesh->GetBodySetup()->AggGeom.EmptyElements();
+				}
 				bImportedCustomCollision |= ImportBoxCollision(Arguments, StaticMesh, LodDataNode);
 				bImportedCustomCollision |= ImportCapsuleCollision(Arguments, StaticMesh, LodDataNode);
 				bImportedCustomCollision |= ImportSphereCollision(Arguments, StaticMesh, LodDataNode);
