[Interchange - FBX] Not importing extra uv sets when having VertexAttributeFromMeshVertexColor.

#rnx

diff --git a/Engine/Plugins/Interchange/Runtime/Source/Parsers/Fbx/Private/FbxMesh.cpp b/Engine/Plugins/Interchange/Runtime/Source/Parsers/Fbx/Private/FbxMesh.cpp
--- a/Engine/Plugins/Interchange/Runtime/Source/Parsers/Fbx/Private/FbxMesh.cpp
+++ b/Engine/Plugins/Interchange/Runtime/Source/Parsers/Fbx/Private/FbxMesh.cpp
@@ -160,7 +160,7 @@ namespace UE::Interchange::Private
 		if(NamedVertexAttributes.Num() > 0)
 		{
 			FSkeletalMeshAttributes MeshAttributes(MeshDescription);
-			MeshAttributes.Register();
+			MeshAttributes.Register(true);
 			TVertexAttributesRef<FVector3f> VertexPositions = MeshAttributes.GetVertexPositions();
 
 			TMap<FString, FName> ValidAttributes;
