[MLDeformer] The recent importer changes broke the import of specific data that was needed on Skeletal Meshes to work with ML Deformer. This change fixes that. Also fixed a small issue in the meta tag of a property.

#rnx

diff --git a/Engine/Plugins/Animation/MLDeformer/MLDeformerFramework/Source/MLDeformerFramework/Public/MLDeformerMorphModel.h b/Engine/Plugins/Animation/MLDeformer/MLDeformerFramework/Source/MLDeformerFramework/Public/MLDeformerMorphModel.h
--- a/Engine/Plugins/Animation/MLDeformer/MLDeformerFramework/Source/MLDeformerFramework/Public/MLDeformerMorphModel.h
+++ b/Engine/Plugins/Animation/MLDeformer/MLDeformerFramework/Source/MLDeformerFramework/Public/MLDeformerMorphModel.h
@@ -395,7 +395,7 @@ private:
 	 * You can use this feather out influence of the ML Deformer in specific areas, such as neck line seams, where the head mesh connects with the body.
 	 * The painted vertex color values will be like a weight multiplier on the ML deformer deltas applied to that vertex. You can invert the mask as well.
 	 */
-	UPROPERTY(EditAnywhere, Category = "Morph Targets", meta = (EditCondition = "CanDynamicallyUpdateMorphTargets"))
+	UPROPERTY(EditAnywhere, Category = "Morph Targets", meta = (EditCondition = "CanDynamicallyUpdateMorphTargets()"))
 	EMLDeformerMaskChannel MaskChannel = EMLDeformerMaskChannel::Disabled;
 
 #if WITH_EDITORONLY_DATA
diff --git a/Engine/Plugins/Interchange/Runtime/Source/Import/Private/Mesh/InterchangeSkeletalMeshFactory.cpp b/Engine/Plugins/Interchange/Runtime/Source/Import/Private/Mesh/InterchangeSkeletalMeshFactory.cpp
--- a/Engine/Plugins/Interchange/Runtime/Source/Import/Private/Mesh/InterchangeSkeletalMeshFactory.cpp
+++ b/Engine/Plugins/Interchange/Runtime/Source/Import/Private/Mesh/InterchangeSkeletalMeshFactory.cpp
@@ -518,6 +518,8 @@ namespace UE
 				MorphTargetMeshDescriptionsPerMorphTargetName.Reserve(MorphTargetCount);
 
 				//Fill the lod mesh description using all combined mesh part
+				TArray<SkeletalMeshImportData::FMeshInfo> MeshInfos;
+				int32 MeshInfoVertexOffset = 0;
 				for(TPair<const FMeshNodeContext*, TFuture<TOptional<UE::Interchange::FMeshPayloadData>>>& MeshNodeContextAndFuture: LodMeshPayloadPerTranslatorPayloadKey)
 				{
 					if (!MeshNodeContextAndFuture.Key)
@@ -532,6 +534,14 @@ namespace UE
 						UE_LOG(LogInterchangeImport, Warning, TEXT("Invalid skeletal mesh payload key [%s] for SkeletalMesh asset %s."), *MeshNodeContext.TranslatorPayloadKey.UniqueId, *Arguments.AssetName);
 						continue;
 					}
+
+					// Generate the mesh info.
+					SkeletalMeshImportData::FMeshInfo MeshInfo;
+					MeshInfo.Name = FName(MeshNodeContext.SceneNode->GetDisplayLabel());
+					MeshNodeContext.MeshNode->GetCustomVertexCount(MeshInfo.NumVertices);
+					MeshInfo.StartImportedVertex = MeshInfoVertexOffset;
+					MeshInfoVertexOffset += MeshInfo.NumVertices;
+					MeshInfos.Add(MeshInfo);
 					
 					const int32 VertexOffset = LodMeshDescription.Vertices().Num();
 
@@ -653,6 +663,7 @@ namespace UE
 
 				DestinationImportData = FSkeletalMeshImportData::CreateFromMeshDescription(LodMeshDescription);
 				DestinationImportData.RefBonesBinary = RefBonesBinary;
+				DestinationImportData.MeshInfos = MoveTemp(MeshInfos);
 
 				bool bMergeMorphTargetWithSameName = false;
 				SkeletalMeshFactoryNode->GetCustomMergeMorphTargetShapeWithSameName(bMergeMorphTargetWithSameName);
