[MLDeformer] Fixed some missing nullptr check which triggered a static analyzer warning.

#rnx

diff --git a/Engine/Plugins/Interchange/Runtime/Source/Import/Private/Mesh/InterchangeSkeletalMeshFactory.cpp b/Engine/Plugins/Interchange/Runtime/Source/Import/Private/Mesh/InterchangeSkeletalMeshFactory.cpp
--- a/Engine/Plugins/Interchange/Runtime/Source/Import/Private/Mesh/InterchangeSkeletalMeshFactory.cpp
+++ b/Engine/Plugins/Interchange/Runtime/Source/Import/Private/Mesh/InterchangeSkeletalMeshFactory.cpp
@@ -536,12 +536,25 @@ namespace UE
 					}
 
 					// Generate the mesh info.
-					SkeletalMeshImportData::FMeshInfo MeshInfo;
-					MeshInfo.Name = FName(MeshNodeContext.SceneNode->GetDisplayLabel());
-					MeshNodeContext.MeshNode->GetCustomVertexCount(MeshInfo.NumVertices);
-					MeshInfo.StartImportedVertex = MeshInfoVertexOffset;
-					MeshInfoVertexOffset += MeshInfo.NumVertices;
-					MeshInfos.Add(MeshInfo);
+					if (MeshNodeContext.SceneNode && MeshNodeContext.MeshNode)
+					{
+						SkeletalMeshImportData::FMeshInfo MeshInfo;
+						MeshInfo.Name = FName(MeshNodeContext.SceneNode->GetDisplayLabel());
+						if (MeshNodeContext.MeshNode->GetCustomVertexCount(MeshInfo.NumVertices))
+						{
+							MeshInfo.StartImportedVertex = MeshInfoVertexOffset;
+							MeshInfoVertexOffset += MeshInfo.NumVertices;
+							MeshInfos.Add(MeshInfo);
+						}
+					}
+					else
+					{
+						int32 NumVerts = 0;
+						if (MeshNodeContext.MeshNode && MeshNodeContext.MeshNode->GetCustomVertexCount(NumVerts))
+						{
+							MeshInfoVertexOffset += NumVerts;
+						}
+					}
 					
 					const int32 VertexOffset = LodMeshDescription.Vertices().Num();
 
