Interchange:
- Add enum instead of a boolean for collision in the Mesh Pipeline (Box, Sphere. Capsule and K-DOP)
- Deprecate method with boolean collision in the StaticMeshLodDataNode in favor of a method taking an enum
- remove old K-DOP function from the StaticMeshFactory, using now the one providing in GeomfitUtils (used by the StaticMeshEditor to be consistent)

GeomFitUtils:
- Add default boolean to the collision functions to avoid launching a rendering update on the StaticMesh (already done by Interchange later)
- Improve K-DOP collision function to use a MeshDescription when a RenderData is not present

diff --git a/Engine/Plugins/Interchange/Editor/Source/InterchangeEditor/Private/InterchangeFbxAssetImportDataConverter.cpp b/Engine/Plugins/Interchange/Editor/Source/InterchangeEditor/Private/InterchangeFbxAssetImportDataConverter.cpp
--- a/Engine/Plugins/Interchange/Editor/Source/InterchangeEditor/Private/InterchangeFbxAssetImportDataConverter.cpp
+++ b/Engine/Plugins/Interchange/Editor/Source/InterchangeEditor/Private/InterchangeFbxAssetImportDataConverter.cpp
@@ -122,7 +122,7 @@ namespace UE::Interchange::Private
 
 	void FillFbxStaticMeshImportData(const UInterchangeGenericAssetsPipeline* GenericAssetPipeline, UFbxStaticMeshImportData* DestinationStaticMeshImportData)
 	{
-		DestinationStaticMeshImportData->bAutoGenerateCollision = GenericAssetPipeline->MeshPipeline->bImportCollision;
+		DestinationStaticMeshImportData->bAutoGenerateCollision = GenericAssetPipeline->MeshPipeline->Collision != EInterchangeMeshCollision::None;
 		DestinationStaticMeshImportData->bBuildNanite = GenericAssetPipeline->MeshPipeline->bBuildNanite;
 		DestinationStaticMeshImportData->bBuildReversedIndexBuffer = GenericAssetPipeline->MeshPipeline->bBuildReversedIndexBuffer;
 		DestinationStaticMeshImportData->bCombineMeshes = GenericAssetPipeline->MeshPipeline->bCombineStaticMeshes;
@@ -292,7 +292,7 @@ namespace UE::Interchange::Private
 			FillInterchangeGenericAssetsPipelineFromFbxMeshImportData(GenericAssetPipeline, StaticMeshImportData);
 		}
 
-		GenericAssetPipeline->MeshPipeline->bImportCollision = StaticMeshImportData->bAutoGenerateCollision;
+		GenericAssetPipeline->MeshPipeline->Collision = StaticMeshImportData->bAutoGenerateCollision ? EInterchangeMeshCollision::Convex18DOP : EInterchangeMeshCollision::None;
 		GenericAssetPipeline->MeshPipeline->bBuildNanite = StaticMeshImportData->bBuildNanite;
 		GenericAssetPipeline->MeshPipeline->bBuildReversedIndexBuffer = StaticMeshImportData->bBuildReversedIndexBuffer;
 		GenericAssetPipeline->MeshPipeline->bCombineStaticMeshes = StaticMeshImportData->bCombineMeshes;
diff --git a/Engine/Plugins/Interchange/Runtime/Source/FactoryNodes/Private/InterchangeStaticMeshLodDataNode.cpp b/Engine/Plugins/Interchange/Runtime/Source/FactoryNodes/Private/InterchangeStaticMeshLodDataNode.cpp
--- a/Engine/Plugins/Interchange/Runtime/Source/FactoryNodes/Private/InterchangeStaticMeshLodDataNode.cpp
+++ b/Engine/Plugins/Interchange/Runtime/Source/FactoryNodes/Private/InterchangeStaticMeshLodDataNode.cpp
@@ -322,14 +322,24 @@ bool UInterchangeStaticMeshLodDataNode::SetOneConvexHullPerUCX(bool AttributeVal
 	IMPLEMENT_NODE_ATTRIBUTE_SETTER_NODELEGATE(OneConvexHullPerUCX, bool)
 }
 
+bool UInterchangeStaticMeshLodDataNode::GetImportCollision(EInterchangeMeshCollision& AttributeValue) const
+{
+	IMPLEMENT_NODE_ATTRIBUTE_GETTER(ImportCollision, EInterchangeMeshCollision);
+}
+
+bool UInterchangeStaticMeshLodDataNode::SetImportCollision(EInterchangeMeshCollision AttributeValue)
+{
+	IMPLEMENT_NODE_ATTRIBUTE_SETTER_NODELEGATE(ImportCollision, EInterchangeMeshCollision)
+}
+
 bool UInterchangeStaticMeshLodDataNode::GetImportCollision(bool& AttributeValue) const
 {
-	IMPLEMENT_NODE_ATTRIBUTE_GETTER(ImportCollision, bool);
+	return false;
 }
 
 bool UInterchangeStaticMeshLodDataNode::SetImportCollision(bool AttributeValue)
 {
-	IMPLEMENT_NODE_ATTRIBUTE_SETTER_NODELEGATE(ImportCollision, bool)
+	return false;
 }
 
 bool UInterchangeStaticMeshLodDataNode::IsEditorOnlyDataDefined()
diff --git a/Engine/Plugins/Interchange/Runtime/Source/FactoryNodes/Public/InterchangeStaticMeshLodDataNode.h b/Engine/Plugins/Interchange/Runtime/Source/FactoryNodes/Public/InterchangeStaticMeshLodDataNode.h
--- a/Engine/Plugins/Interchange/Runtime/Source/FactoryNodes/Public/InterchangeStaticMeshLodDataNode.h
+++ b/Engine/Plugins/Interchange/Runtime/Source/FactoryNodes/Public/InterchangeStaticMeshLodDataNode.h
@@ -3,6 +3,7 @@
 #pragma once
 
 #include "CoreMinimal.h"
+#include "InterchangeMeshDefinitions.h"
 #include "UObject/ObjectMacros.h"
 #include "Nodes/InterchangeFactoryBaseNode.h"
 
@@ -128,9 +129,15 @@ public:
 	bool SetOneConvexHullPerUCX(bool AttributeValue);
 
 	UFUNCTION(BlueprintCallable, Category = "Interchange | Node | StaticMeshLodData")
-	bool GetImportCollision(bool& AttributeValue) const;
+	bool GetImportCollision(EInterchangeMeshCollision& AttributeValue) const;
 
 	UFUNCTION(BlueprintCallable, Category = "Interchange | Node | StaticMeshLodData")
+	bool SetImportCollision(EInterchangeMeshCollision AttributeValue);
+
+	UE_DEPRECATED(5.5, "GetImportCollision using a bool has been deprecated, please use the function with the enum instead.")
+	bool GetImportCollision(bool& AttributeValue) const;
+
+	UE_DEPRECATED(5.5, "SetImportCollision using a bool has been deprecated, please use the function with the enum instead.")
 	bool SetImportCollision(bool AttributeValue);
 
 
diff --git a/Engine/Plugins/Interchange/Runtime/Source/Import/Private/Mesh/InterchangeStaticMeshFactory.cpp b/Engine/Plugins/Interchange/Runtime/Source/Import/Private/Mesh/InterchangeStaticMeshFactory.cpp
--- a/Engine/Plugins/Interchange/Runtime/Source/Import/Private/Mesh/InterchangeStaticMeshFactory.cpp
+++ b/Engine/Plugins/Interchange/Runtime/Source/Import/Private/Mesh/InterchangeStaticMeshFactory.cpp
@@ -4,6 +4,7 @@
 
 #if WITH_EDITOR
 #include "BSPOps.h"
+#include "Editor/UnrealEd/Private/GeomFitUtils.h"
 #endif
 #include "Components.h"
 #include "Engine/Polys.h"
@@ -301,7 +302,7 @@ UInterchangeFactoryBase::FImportAssetResult UInterchangeStaticMeshFactory::Impor
 	TArray<FMeshDescription>& LodMeshDescriptions = ImportAssetObjectData.LodMeshDescriptions;
 	LodMeshDescriptions.SetNum(LodCount);
 
-	bool bImportCollision = false;
+	EInterchangeMeshCollision Collision = EInterchangeMeshCollision::None;
 	bool bImportedCustomCollision = false;
 	int32 CurrentLodIndex = 0;
 	for (int32 LodIndex = 0; LodIndex < LodCount; ++LodIndex)
@@ -426,8 +427,8 @@ UInterchangeFactoryBase::FImportAssetResult UInterchangeStaticMeshFactory::Impor
 		// Import collision geometry
 		if (CurrentLodIndex == 0)
 		{
-			LodDataNode->GetImportCollision(bImportCollision);
-			if(bImportCollision)
+			LodDataNode->GetImportCollision(Collision);
+			if(Collision != EInterchangeMeshCollision::None)
 			{
 				if (bReimport)
 				{
@@ -459,7 +460,7 @@ UInterchangeFactoryBase::FImportAssetResult UInterchangeStaticMeshFactory::Impor
 	}
 #endif // WITH_EDITOR
 
-	ImportAssetObjectData.bImportCollision = bImportCollision;
+	ImportAssetObjectData.Collision = Collision;
 	ImportAssetObjectData.bImportedCustomCollision = bImportedCustomCollision;
 
 	// Getting the file Hash will cache it into the source data
@@ -661,19 +662,47 @@ UInterchangeFactoryBase::FImportAssetResult UInterchangeStaticMeshFactory::EndIm
 	}
 #endif // WITH_EDITOR
 
-	if(ImportAssetObjectData.bImportCollision)
+#if WITH_EDITOR
+	if(!ImportAssetObjectData.bImportedCustomCollision)
 	{
-		if (!ImportAssetObjectData.bImportedCustomCollision)
+		constexpr bool bUpdateRendering = false;
+		switch(ImportAssetObjectData.Collision)
 		{
-			GenerateKDopCollision(StaticMesh);
+		case EInterchangeMeshCollision::Box:
+			GenerateBoxAsSimpleCollision(StaticMesh, bUpdateRendering);
+			break;
+		case EInterchangeMeshCollision::Sphere:
+			GenerateSphereAsSimpleCollision(StaticMesh, bUpdateRendering);
+			break;
+		case EInterchangeMeshCollision::Capsule:
+			GenerateSphylAsSimpleCollision(StaticMesh, bUpdateRendering);
+			break;
+		case EInterchangeMeshCollision::Convex10DOP_X:
+			GenerateKDopAsSimpleCollision(StaticMesh, TArray<FVector>(KDopDir10X, sizeof(KDopDir10X) / sizeof(FVector)), bUpdateRendering);
+			break;
+		case EInterchangeMeshCollision::Convex10DOP_Y:
+			GenerateKDopAsSimpleCollision(StaticMesh, TArray<FVector>(KDopDir10Y, sizeof(KDopDir10Y) / sizeof(FVector)), bUpdateRendering);
+			break;
+		case EInterchangeMeshCollision::Convex10DOP_Z:
+			GenerateKDopAsSimpleCollision(StaticMesh, TArray<FVector>(KDopDir10Z, sizeof(KDopDir10Z) / sizeof(FVector)), bUpdateRendering);
+			break;
+		case EInterchangeMeshCollision::Convex18DOP:
+			GenerateKDopAsSimpleCollision(StaticMesh, TArray<FVector>(KDopDir18, sizeof(KDopDir18) / sizeof(FVector)), bUpdateRendering);
+			break;
+		case EInterchangeMeshCollision::Convex26DOP:
+			GenerateKDopAsSimpleCollision(StaticMesh, TArray<FVector>(KDopDir26, sizeof(KDopDir26) / sizeof(FVector)), bUpdateRendering);
+			break;
+		default:
+			break;
 		}
+	}
+#endif
 #if WITH_EDITORONLY_DATA
-		else
-		{
-			StaticMesh->bCustomizedCollision = true;
-		}
-#endif // WITH_EDITORONLY_DATA
+	else
+	{
+		StaticMesh->bCustomizedCollision = true;
 	}
+#endif // WITH_EDITORONLY_DATA
 #if WITH_EDITOR
 	//Lod group need to use the static mesh API and cannot use the apply delegate
 	if (!Arguments.ReimportObject)
@@ -750,7 +779,7 @@ void UInterchangeStaticMeshFactory::BuildFromMeshDescriptions(UStaticMesh& Stati
 	BuildMeshDescriptionsParams.bFastBuild = true;
 	// For the time being at runtime collision is set to complex one
 	// TODO: Revisit pipeline options for collision. bImportCollision is not enough.
-	BuildMeshDescriptionsParams.bAllowCpuAccess = ImportAssetObjectData.bImportCollision;
+	BuildMeshDescriptionsParams.bAllowCpuAccess = ImportAssetObjectData.Collision != EInterchangeMeshCollision::None;
 	StaticMesh.bAllowCPUAccess = BuildMeshDescriptionsParams.bAllowCpuAccess;
 
 	StaticMesh.BuildFromMeshDescriptions(MeshDescriptionPointers, BuildMeshDescriptionsParams);
@@ -758,7 +787,7 @@ void UInterchangeStaticMeshFactory::BuildFromMeshDescriptions(UStaticMesh& Stati
 	// TODO: Expand support for different collision types
 	if (ensure(StaticMesh.GetRenderData()))
 	{
-		if (ImportAssetObjectData.bImportCollision && !ImportAssetObjectData.bImportedCustomCollision)
+		if (ImportAssetObjectData.Collision != EInterchangeMeshCollision::None && !ImportAssetObjectData.bImportedCustomCollision)
 		{
 			if (StaticMesh.GetBodySetup() == nullptr)
 			{
@@ -1608,161 +1637,6 @@ bool UInterchangeStaticMeshFactory::ImportConvexCollision(const FImportAssetObje
 	return bResult;
 }
 
-
-bool UInterchangeStaticMeshFactory::GenerateKDopCollision(UStaticMesh* StaticMesh)
-{
-#if WITH_EDITOR
-
-	constexpr float RecipSqrt2 = UE_INV_SQRT_2;
-	constexpr int32 KCount = 18;
-	static const FVector3f KDopDir18[KCount] =
-	{
-		FVector3f( 1.0f,  0.0f,  0.0f),
-		FVector3f(-1.0f,  0.0f,  0.0f),
-		FVector3f( 0.0f,  1.0f,  0.0f),
-		FVector3f( 0.0f, -1.0f,  0.0f),
-		FVector3f( 0.0f,  0.0f,  1.0f),
-		FVector3f( 0.0f,  0.0f, -1.0f),
-		FVector3f( 0.0f,  RecipSqrt2,  RecipSqrt2),
-		FVector3f( 0.0f, -RecipSqrt2, -RecipSqrt2),
-		FVector3f( 0.0f,  RecipSqrt2, -RecipSqrt2),
-		FVector3f( 0.0f, -RecipSqrt2,  RecipSqrt2),
-		FVector3f( RecipSqrt2,  0.0f,  RecipSqrt2),
-		FVector3f(-RecipSqrt2,  0.0f, -RecipSqrt2),
-		FVector3f( RecipSqrt2,  0.0f, -RecipSqrt2),
-		FVector3f(-RecipSqrt2,  0.0f,  RecipSqrt2),
-		FVector3f( RecipSqrt2,  RecipSqrt2,  0.0f),
-		FVector3f(-RecipSqrt2, -RecipSqrt2,  0.0f),
-		FVector3f( RecipSqrt2, -RecipSqrt2,  0.0f),
-		FVector3f(-RecipSqrt2,  RecipSqrt2,  0.0f)
-	};
-
-	UBodySetup* BodySetup = StaticMesh->GetBodySetup();
-
-	// Initialize maximum distances from each kdop plane to the minimum value possible
-	TArray<float> MaxDist;
-	MaxDist.Reserve(KCount);
-	for (int32 Count = 0; Count < KCount; Count++)
-	{
-		MaxDist.Add(-MAX_FLT);
-	}
-
-	// Construct temporary UModel for kdop creation.
-	UModel* TempModel = nullptr;
-	{
-		FGCScopeGuard GCScopeGuard;
-		TempModel = NewObject<UModel>();
-	}
-	TempModel->RootOutside = true;
-	TempModel->EmptyModel(true, true);
-	TempModel->Polys->ClearFlags(RF_Transactional);
-
-	// Get the vertex positions for the final LOD0 mesh
-	FMeshDescription* MeshDescription = StaticMesh->GetMeshDescription(0);
-	check(MeshDescription);
-
-	FStaticMeshConstAttributes Attributes(*MeshDescription);
-	TVertexAttributesRef<const FVector3f> VertexPositions = Attributes.GetVertexPositions();
-
-	// For each vertex, project along each kdop direction, to find the max in that direction.
-	for (FVector3f VertexPosition : VertexPositions.GetRawArray())
-	{
-		for (int32 Index = 0; Index < KCount; Index++)
-		{
-			float Dist = FVector3f::DotProduct(VertexPosition, KDopDir18[Index]);
-			MaxDist[Index] = FMath::Max(Dist, MaxDist[Index]);
-		}
-	}
-
-	// Inflate kdop to ensure it is not degenerate
-	constexpr float MinSize = 0.1f;
-	for (int32 Index = 0; Index < KCount; Index++)
-	{
-		MaxDist[Index] += MinSize;
-	}
-
-	// Now we have the planes of the kdop, we work out the face polygons.
-	TArray<FPlane4f> Planes;
-	Planes.Reserve(KCount);
-	for (int32 Index = 0; Index < KCount; Index++)
-	{
-		Planes.Add(FPlane4f(KDopDir18[Index], MaxDist[Index]));
-	}
-
-	for (int32 PlaneIndex = 0; PlaneIndex < Planes.Num(); PlaneIndex++)
-	{
-		FPoly& Polygon = TempModel->Polys->Element.Emplace_GetRef();
-
-		Polygon.Init();
-		Polygon.Normal = Planes[PlaneIndex];
-
-		FVector3f AxisX, AxisY;
-		Polygon.Normal.FindBestAxisVectors(AxisX, AxisY);
-
-		FVector3f Base = Planes[PlaneIndex] * Planes[PlaneIndex].W;
-
-		Polygon.Vertices.Emplace(Base + AxisX * UE_OLD_HALF_WORLD_MAX + AxisY * UE_OLD_HALF_WORLD_MAX);
-		Polygon.Vertices.Emplace(Base + AxisX * UE_OLD_HALF_WORLD_MAX - AxisY * UE_OLD_HALF_WORLD_MAX);
-		Polygon.Vertices.Emplace(Base - AxisX * UE_OLD_HALF_WORLD_MAX - AxisY * UE_OLD_HALF_WORLD_MAX);
-		Polygon.Vertices.Emplace(Base - AxisX * UE_OLD_HALF_WORLD_MAX + AxisY * UE_OLD_HALF_WORLD_MAX);
-
-		for (int32 OtherPlaneIndex = 0; OtherPlaneIndex < Planes.Num(); OtherPlaneIndex++)
-		{
-			if (PlaneIndex != OtherPlaneIndex)
-			{
-				if (!Polygon.Split(-FVector3f(Planes[OtherPlaneIndex]), Planes[OtherPlaneIndex] * Planes[OtherPlaneIndex].W))
-				{
-					Polygon.Vertices.Empty();
-					break;
-				}
-			}
-		}
-
-		if (Polygon.Vertices.Num() < 3)
-		{
-			// If poly resulted in no verts, remove from array
-			TempModel->Polys->Element.RemoveAt(TempModel->Polys->Element.Num() - 1);
-		}
-		else
-		{
-			Polygon.iLink = PlaneIndex;
-
-			constexpr bool bSilent = true;
-			Polygon.CalcNormal(bSilent);
-		}
-	}
-
-	if (TempModel->Polys->Element.Num() < 4)
-	{
-		return false;
-	}
-
-	// Build bounding box.
-	TempModel->BuildBound();
-
-	// Build BSP for the brush.
-	FBSPOps::bspBuild(TempModel, FBSPOps::BSP_Good, 15, 70, 1, 0);
-
-	constexpr bool bNoRemapSurfs = true;
-	FBSPOps::bspRefresh(TempModel, bNoRemapSurfs);
-	FBSPOps::bspBuildBounds(TempModel);
-
-	bool bRemoveExisting = true;
-	BodySetup->CreateFromModel(TempModel, bRemoveExisting);
-
-	TempModel->ClearInternalFlags(EInternalObjectFlags::Async);
-	TempModel->Polys->ClearInternalFlags(EInternalObjectFlags::Async);
-
-	return true;
-
-#else // #if WITH_EDITOR
-
-	return false;
-
-#endif
-}
-
-
 bool UInterchangeStaticMeshFactory::ImportSockets(const FImportAssetObjectParams& Arguments, UStaticMesh* StaticMesh, const UInterchangeStaticMeshFactoryNode* FactoryNode)
 {
 	TArray<FString> SocketUids;
diff --git a/Engine/Plugins/Interchange/Runtime/Source/Import/Public/Mesh/InterchangeStaticMeshFactory.h b/Engine/Plugins/Interchange/Runtime/Source/Import/Public/Mesh/InterchangeStaticMeshFactory.h
--- a/Engine/Plugins/Interchange/Runtime/Source/Import/Public/Mesh/InterchangeStaticMeshFactory.h
+++ b/Engine/Plugins/Interchange/Runtime/Source/Import/Public/Mesh/InterchangeStaticMeshFactory.h
@@ -4,6 +4,7 @@
 
 #include "CoreMinimal.h"
 #include "InterchangeFactoryBase.h"
+#include "InterchangeMeshDefinitions.h"
 #include "MeshDescription.h"
 #include "UObject/Object.h"
 #include "UObject/ObjectMacros.h"
@@ -78,7 +79,7 @@ private:
 	{
 		bool bIsAppGame = false;
 		bool bImportedCustomCollision = false;
-		bool bImportCollision = false;
+		EInterchangeMeshCollision Collision = EInterchangeMeshCollision::None;
 		TArray<FMeshDescription> LodMeshDescriptions;
 	};
 	FImportAssetObjectData ImportAssetObjectData;
diff --git a/Engine/Plugins/Interchange/Runtime/Source/Nodes/Public/InterchangeMeshDefinitions.h b/Engine/Plugins/Interchange/Runtime/Source/Nodes/Public/InterchangeMeshDefinitions.h
new file mode 100644
--- /dev/null
+++ b/Engine/Plugins/Interchange/Runtime/Source/Nodes/Public/InterchangeMeshDefinitions.h
@@ -0,0 +1,59 @@
+// Copyright Epic Games, Inc. All Rights Reserved.
+
+#pragma once
+
+#include "UObject/NameTypes.h"
+#include "Types/AttributeStorage.h"
+
+#include "InterchangeMeshDefinitions.generated.h"
+
+UENUM(BlueprintType)
+enum class EInterchangeMeshCollision : uint8
+{
+	/** Generates a new box collision mesh encompassing the static mesh*/
+	Box,
+
+	/** Generates a new sphere collision mesh encompassing the static mesh*/
+	Sphere,
+
+	/** Generates a new capsule collision mesh encompassing the static mesh*/
+	Capsule,
+
+	/** Generates a new axis-aligned box collision mesh with the 4 X-axis aligned edges beveled (10 total sides) encompassing the static mesh*/
+	Convex10DOP_X UMETA(DisplayName="10DOP-X Simplified"),
+
+	/** Generates a new axis-aligned box collision mesh with the 4 Y-axis aligned edges beveled (10 total sides) encompassing the static mesh*/
+	Convex10DOP_Y UMETA(DisplayName = "10DOP-Y Simplified"),
+
+	/** Generates a new axis-aligned box collision mesh with the 4 Z-axis aligned edges beveled (10 total sides) encompassing the static mesh*/
+	Convex10DOP_Z UMETA(DisplayName = "10DOP-Z Simplified"),
+
+	/** Generates a new axis-aligned box collision mesh with all edges beveled (18 total sides) encompassing the static mesh*/
+	Convex18DOP UMETA(DisplayName = "18DOP Simplified"),
+	
+	/** Generates a new axis-aligned box collision mesh with all edges and corners beveled (26 total sides) encompassing the static mesh*/
+	Convex26DOP UMETA(DisplayName = "26DOP Simplified"),
+	
+	None = 0xFF
+};
+
+namespace UE
+{
+	namespace Interchange
+	{
+		template<> struct TAttributeTypeTraits<EInterchangeMeshCollision>
+		{
+			static constexpr EAttributeTypes GetType()
+			{
+				return EAttributeTypes::UInt8;
+			}
+			static FString ToString(const uint16& Value)
+			{
+				int32 ValueConv = Value;
+				FStringFormatOrderedArguments OrderedArguments;
+				OrderedArguments.Add(FStringFormatArg(ValueConv));
+				return FString::Format(TEXT("{0}"), OrderedArguments);
+			}
+		};
+	}
+}
\ No newline at end of file
diff --git a/Engine/Plugins/Interchange/Runtime/Source/Pipelines/Private/InterchangeGenericMeshPipeline.cpp b/Engine/Plugins/Interchange/Runtime/Source/Pipelines/Private/InterchangeGenericMeshPipeline.cpp
--- a/Engine/Plugins/Interchange/Runtime/Source/Pipelines/Private/InterchangeGenericMeshPipeline.cpp
+++ b/Engine/Plugins/Interchange/Runtime/Source/Pipelines/Private/InterchangeGenericMeshPipeline.cpp
@@ -18,6 +18,7 @@
 #include "Nodes/InterchangeBaseNode.h"
 #include "Nodes/InterchangeBaseNodeContainer.h"
 #include "PhysicsEngine/PhysicsAsset.h"
+#include "InterchangeCustomVersion.h"
 
 #include UE_INLINE_GENERATED_CPP_BY_NAME(InterchangeGenericMeshPipeline)
 
@@ -105,7 +106,7 @@ void UInterchangeGenericMeshPipeline::AdjustSettingsForContext(const FInterchang
 			bCombineStaticMeshes = true;
 			bBuildNanite = false;
 			LodGroup = NAME_None;
-			bImportCollision = false;
+			Collision = EInterchangeMeshCollision::None;
 			bImportCollisionAccordingToMeshName = false;
 			bGenerateLightmapUVs = false;
 			bGenerateDistanceFieldAsIfTwoSided = false;
@@ -127,7 +128,7 @@ void UInterchangeGenericMeshPipeline::AdjustSettingsForContext(const FInterchang
 					bCombineStaticMeshes = true;
 					LodGroup = NAME_None;
 					bSupportFaceRemap = false;
-					bImportCollision = false;
+					Collision = EInterchangeMeshCollision::None;
 					bImportCollisionAccordingToMeshName = false;
 					bGenerateLightmapUVs = false;
 					bGenerateDistanceFieldAsIfTwoSided = false;
@@ -349,6 +350,16 @@ void UInterchangeGenericMeshPipeline::GetMeshesInformationFromTranslatedData(con
 	}
 }
 
+void UInterchangeGenericMeshPipeline::PostLoad()
+{
+	Super::PostLoad();
+
+	if(!bImportCollision_DEPRECATED)
+	{
+		Collision = EInterchangeMeshCollision::None;
+	}
+}
+
 UInterchangePipelineMeshesUtilities* UInterchangeGenericMeshPipeline::CreateMeshPipelineUtilities(UInterchangeBaseNodeContainer* InBaseNodeContainer
 	, const UInterchangeGenericMeshPipeline* Pipeline)
 {
diff --git a/Engine/Plugins/Interchange/Runtime/Source/Pipelines/Private/InterchangeGenericStaticMeshPipeline.cpp b/Engine/Plugins/Interchange/Runtime/Source/Pipelines/Private/InterchangeGenericStaticMeshPipeline.cpp
--- a/Engine/Plugins/Interchange/Runtime/Source/Pipelines/Private/InterchangeGenericStaticMeshPipeline.cpp
+++ b/Engine/Plugins/Interchange/Runtime/Source/Pipelines/Private/InterchangeGenericStaticMeshPipeline.cpp
@@ -21,15 +21,6 @@
 #include "UObject/ObjectMacros.h"
 
 
-enum class EMeshCollisionType
-{
-	None,
-	Box,
-	Sphere,
-	Capsule,
-	Convex
-};
-
 namespace UE::Interchange::Private
 {
 	FString GetNodeName(const UInterchangeBaseNodeContainer& NodeContainer, const FString& NodeUid)
@@ -59,32 +50,32 @@ namespace UE::Interchange::Private
 	}
 }
 
-static TTuple<EMeshCollisionType, FString> GetCollisionMeshType(const UInterchangeBaseNodeContainer& NodeContainer, const FString& NodeUid, const TArray<FString>& AllNodeUids)
+static TTuple<EInterchangeMeshCollision, FString> GetCollisionMeshType(const UInterchangeBaseNodeContainer& NodeContainer, const FString& NodeUid, const TArray<FString>& AllNodeUids)
 {
 	FString MeshName = UE::Interchange::Private::GetNodeName(NodeContainer, NodeUid);
-	EMeshCollisionType CollisionType = EMeshCollisionType::None;
+	EInterchangeMeshCollision CollisionType = EInterchangeMeshCollision::None;
 
 	// Determine if the mesh name is a potential collision mesh
 
 	if (MeshName.StartsWith(TEXT("UBX_")))
 	{
-		CollisionType = EMeshCollisionType::Box;
+		CollisionType = EInterchangeMeshCollision::Box;
 	}
 	else if (MeshName.StartsWith(TEXT("UCX_")) || MeshName.StartsWith(TEXT("MCDCX_")))
 	{
-		CollisionType = EMeshCollisionType::Convex;
+		CollisionType = EInterchangeMeshCollision::Convex18DOP;
 	}
 	else if (MeshName.StartsWith(TEXT("USP_")))
 	{
-		CollisionType = EMeshCollisionType::Sphere;
+		CollisionType = EInterchangeMeshCollision::Sphere;
 	}
 	else if (MeshName.StartsWith(TEXT("UCP_")))
 	{
-		CollisionType = EMeshCollisionType::Capsule;
+		CollisionType = EInterchangeMeshCollision::Capsule;
 	}
 	else
 	{
-		return { EMeshCollisionType::None, FString() };
+		return { EInterchangeMeshCollision::None, FString() };
 	}
 
 	// We have a mesh name with a collision type suffix.
@@ -126,13 +117,13 @@ static TTuple<EMeshCollisionType, FString> GetCollisionMeshType(const UInterchan
 
 	// Mesh had a collision type prefix, but no corresponding mesh, so don't treat it as a collision mesh
 
-	return { EMeshCollisionType::None, FString() };
+	return { EInterchangeMeshCollision::None, FString() };
 }
 
 
 static bool IsCollisionMeshUid(const UInterchangeBaseNodeContainer& NodeContainer, const FString& MeshUid, const TArray<FString>& MeshUids)
 {
-	return GetCollisionMeshType(NodeContainer, MeshUid, MeshUids).Get<0>() != EMeshCollisionType::None;
+	return GetCollisionMeshType(NodeContainer, MeshUid, MeshUids).Get<0>() != EInterchangeMeshCollision::None;
 }
 
 
@@ -140,8 +131,8 @@ static void BuildMeshToCollisionMeshMap(const UInterchangeBaseNodeContainer& Nod
 {
 	for (const FString& MeshUid : MeshUids)
 	{
-		TTuple<EMeshCollisionType, FString> CollisionType = GetCollisionMeshType(NodeContainer, MeshUid, MeshUids);
-		if (CollisionType.Get<0>() != EMeshCollisionType::None)
+		TTuple<EInterchangeMeshCollision, FString> CollisionType = GetCollisionMeshType(NodeContainer, MeshUid, MeshUids);
+		if (CollisionType.Get<0>() != EInterchangeMeshCollision::None)
 		{
 			MeshToCollisionMeshMap.FindOrAdd(FString(CollisionType.Get<1>())).Emplace(MeshUid);
 		}
@@ -279,7 +270,7 @@ void UInterchangeGenericMeshPipeline::ExecutePreImportPipelineStaticMesh()
 
 					if (MeshUidsPerLodIndex.Num() > 0)
 					{
-						if (bImportCollision && bImportCollisionAccordingToMeshName)
+						if (Collision != EInterchangeMeshCollision::None && bImportCollisionAccordingToMeshName)
 						{
 							if (const TArray<FString>* CorrespondingCollisionMeshes = MeshToCollisionMeshMap.Find(MeshUid))
 							{
@@ -323,7 +314,7 @@ void UInterchangeGenericMeshPipeline::ExecutePreImportPipelineStaticMesh()
 
 					if (MeshUidsPerLodIndex.Num() > 0)
 					{
-						if (bImportCollision && bImportCollisionAccordingToMeshName)
+						if (Collision != EInterchangeMeshCollision::None && bImportCollisionAccordingToMeshName)
 						{
 							if (const TArray<FString>* CorrespondingCollisionMeshes = MeshToCollisionMeshMap.Find(MeshUid))
 							{
@@ -505,7 +496,7 @@ UInterchangeStaticMeshLodDataNode* UInterchangeGenericMeshPipeline::CreateStatic
 
 	StaticMeshLodDataNode->InitializeNode(NodeUID, DisplayLabel, EInterchangeNodeContainerType::FactoryData);
 	StaticMeshLodDataNode->SetOneConvexHullPerUCX(bOneConvexHullPerUCX);
-	StaticMeshLodDataNode->SetImportCollision(bImportCollision);
+	StaticMeshLodDataNode->SetImportCollision(Collision);
 	BaseNodeContainer->AddNode(StaticMeshLodDataNode);
 	return StaticMeshLodDataNode;
 }
@@ -587,26 +578,30 @@ void UInterchangeGenericMeshPipeline::AddLodDataToStaticMesh(UInterchangeStaticM
 
 			if (bImportCollisionAccordingToMeshName)
 			{
-				TTuple<EMeshCollisionType, FStringView> MeshType = GetCollisionMeshType(*BaseNodeContainer, NodeUid, NodeUids);
+				TTuple<EInterchangeMeshCollision, FStringView> MeshType = GetCollisionMeshType(*BaseNodeContainer, NodeUid, NodeUids);
 				switch (MeshType.Get<0>())
 				{
-				case EMeshCollisionType::None:
+				case EInterchangeMeshCollision::None:
 					LodDataNode->AddMeshUid(NodeUid);
 					break;
 
-				case EMeshCollisionType::Box:
+				case EInterchangeMeshCollision::Box:
 					LodDataNode->AddBoxCollisionMeshUid(NodeUid);
 					break;
 
-				case EMeshCollisionType::Sphere:
+				case EInterchangeMeshCollision::Sphere:
 					LodDataNode->AddSphereCollisionMeshUid(NodeUid);
 					break;
 
-				case EMeshCollisionType::Capsule:
+				case EInterchangeMeshCollision::Capsule:
 					LodDataNode->AddCapsuleCollisionMeshUid(NodeUid);
 					break;
 
-				case EMeshCollisionType::Convex:
+				case EInterchangeMeshCollision::Convex10DOP_X:
+				case EInterchangeMeshCollision::Convex10DOP_Y:
+				case EInterchangeMeshCollision::Convex10DOP_Z:
+				case EInterchangeMeshCollision::Convex18DOP:
+				case EInterchangeMeshCollision::Convex26DOP:
 					LodDataNode->AddConvexCollisionMeshUid(NodeUid);
 					break;
 				}
diff --git a/Engine/Plugins/Interchange/Runtime/Source/Pipelines/Public/InterchangeGenericMeshPipeline.h b/Engine/Plugins/Interchange/Runtime/Source/Pipelines/Public/InterchangeGenericMeshPipeline.h
--- a/Engine/Plugins/Interchange/Runtime/Source/Pipelines/Public/InterchangeGenericMeshPipeline.h
+++ b/Engine/Plugins/Interchange/Runtime/Source/Pipelines/Public/InterchangeGenericMeshPipeline.h
@@ -3,6 +3,7 @@
 #pragma once
 
 #include "InterchangeGenericAssetsPipelineSharedSettings.h"
+#include "InterchangeMeshDefinitions.h"
 #include "InterchangePipelineBase.h"
 #include "InterchangeSourceData.h"
 #include "Nodes/InterchangeBaseNodeContainer.h"
@@ -55,11 +56,14 @@ public:
 	UPROPERTY(EditAnywhere, BlueprintReadWrite, Category = "Static Meshes", meta=(DisplayName="LOD Group"))
 	FName LodGroup = NAME_None;
 
+	UPROPERTY(meta = (DeprecatedProperty, DeprecationMessage = "Use Collision instead."))
+	bool bImportCollision_DEPRECATED = true;
+
 	/** If enabled, custom collision will be imported. If enabled and there is no custom collision, a generic collision will be automatically generated.
 	 * If disabled, no collision will be created or imported.
 	 */
 	UPROPERTY(EditAnywhere, BlueprintReadWrite, Category = "Static Meshes", meta = (SubCategory = "Collision"))
-	bool bImportCollision = true;
+	EInterchangeMeshCollision Collision = EInterchangeMeshCollision::Convex18DOP;
 
 	/**
 	 * If enabled, meshes with certain prefixes will be imported as collision primitives for the mesh with the corresponding unprefixed name.
@@ -70,11 +74,11 @@ public:
 	 * USP_ Sphere collision
 	 * UCX_ Convex collision
 	 */
-	UPROPERTY(EditAnywhere, BlueprintReadWrite, Category = "Static Meshes", meta = (SubCategory = "Collision", editcondition = "bImportCollision"))
+	UPROPERTY(EditAnywhere, BlueprintReadWrite, Category = "Static Meshes", meta = (SubCategory = "Collision", editcondition = "Collision != EInterchangeMeshCollision::None"))
 	bool bImportCollisionAccordingToMeshName = true;
 
 	/** If enabled, each UCX collision mesh will be imported as a single convex hull. If disabled, a UCX mesh will be decomposed into its separate pieces and a convex hull generated for each. */
-	UPROPERTY(EditAnywhere, BlueprintReadWrite, Category = "Static Meshes", meta = (SubCategory = "Collision", editcondition = "bImportCollision && bImportCollisionAccordingToMeshName"))
+	UPROPERTY(EditAnywhere, BlueprintReadWrite, Category = "Static Meshes", meta = (SubCategory = "Collision", editcondition = "Collision != EInterchangeMeshCollision::None && bImportCollisionAccordingToMeshName"))
 	bool bOneConvexHullPerUCX = true;
 
 	//////	Static Meshes Build settings Properties //////
@@ -365,6 +369,10 @@ private:
 		, bool& bContainStaticMeshAnimationNode
 		, bool& bIgnoreStaticMeshes) const;
 
+public:
+
+	virtual void PostLoad() override;
+
 protected:
 	UInterchangeBaseNodeContainer* BaseNodeContainer = nullptr;
 	TArray<const UInterchangeSourceData*> SourceDatas;
diff --git a/Engine/Source/Editor/UnrealEd/Private/GeomFitUtils.cpp b/Engine/Source/Editor/UnrealEd/Private/GeomFitUtils.cpp
--- a/Engine/Source/Editor/UnrealEd/Private/GeomFitUtils.cpp
+++ b/Engine/Source/Editor/UnrealEd/Private/GeomFitUtils.cpp
@@ -56,10 +56,13 @@ static bool PromptToRemoveExistingCollision(UStaticMesh* StaticMesh)
 // This function takes the current collision model, and fits a k-DOP around it.
 // It uses the array of k unit-length direction vectors to define the k bounding planes.
 
-int32 GenerateKDopAsSimpleCollision(UStaticMesh* StaticMesh, const TArray<FVector> &Dirs)
+int32 GenerateKDopAsSimpleCollision(UStaticMesh* StaticMesh, const TArray<FVector> &Dirs, bool bUpdate)
 {
 	// Make sure rendering is done - so we are not changing data being used by collision drawing.
-	FlushRenderingCommands();
+	if(bUpdate)
+	{
+		FlushRenderingCommands();
+	}
 
 	if (!PromptToRemoveExistingCollision(StaticMesh))
 	{
@@ -68,12 +71,39 @@ int32 GenerateKDopAsSimpleCollision(UStaticMesh* StaticMesh, const TArray<FVecto
 
 	UBodySetup* bs = StaticMesh->GetBodySetup();
 
-	const FStaticMeshLODResources& RenderData = StaticMesh->GetRenderData()->LODResources[0];
+	int32 NumVertices = 0;
+	TFunction<UE::Math::TVector<double>(int32)> GetPointFunc;
+	if(FStaticMeshRenderData* RenderData = StaticMesh->GetRenderData())
+	{
+		const FStaticMeshLODResources& LODResource = RenderData->LODResources[0];
+		NumVertices = LODResource.GetNumVertices();
+		GetPointFunc = [&](int32 VertIdx)
+		{
+			return FVector{ LODResource.VertexBuffers.PositionVertexBuffer.VertexPosition(VertIdx) };
+		};
+	}
+	else if(const FMeshDescription* MeshDescription = StaticMesh->GetMeshDescription(0))
+	{		
+		NumVertices = MeshDescription->Vertices().Num();
+		GetPointFunc = [&](int32 VertIdx)
+		{
+			return FVector(MeshDescription->GetVertexPosition(VertIdx));
+		};
+	}
+
+	if(!GetPointFunc)
+	{
+		UE_LOG(LogStaticMesh, Error, TEXT("Couldn't find a valid RenderData or MeshDescription for the StaticMesh"));
+		return INDEX_NONE;
+	}
+
 	TArray<FVector> HullVertices;
-	UE::Geometry::FitKDOPVertices3<double>(Dirs, RenderData.GetNumVertices(),
-		[&](int32 VertIdx) { return (FVector)RenderData.VertexBuffers.PositionVertexBuffer.VertexPosition(VertIdx); }, HullVertices);
+	UE::Geometry::FitKDOPVertices3<double>(Dirs, NumVertices, GetPointFunc, HullVertices);
 
-	bs->Modify();
+	if(bUpdate)
+	{
+		bs->Modify();
+	}
 
 	// Create new GUID
 	bs->InvalidatePhysicsData();
@@ -89,7 +119,10 @@ int32 GenerateKDopAsSimpleCollision(UStaticMesh* StaticMesh, const TArray<FVecto
 	RefreshCollisionChange(*StaticMesh);
 
 	// Mark staticmesh as dirty, to help make sure it gets saved.
-	StaticMesh->MarkPackageDirty();
+	if(bUpdate)
+	{
+		StaticMesh->MarkPackageDirty();
+	}
 
 	StaticMesh->bCustomizedCollision = true;	//mark the static mesh for collision customization
 
@@ -105,7 +138,7 @@ void ComputeBoundingBox(UStaticMesh* StaticMesh, FVector& Center, FVector& Exten
 	BoundingBox.GetCenterAndExtents(Center, Extents);
 }
 
-int32 GenerateBoxAsSimpleCollision(UStaticMesh* StaticMesh)
+int32 GenerateBoxAsSimpleCollision(UStaticMesh* StaticMesh, bool bUpdate)
 {
 	if (!PromptToRemoveExistingCollision(StaticMesh))
 	{
@@ -119,7 +152,10 @@ int32 GenerateBoxAsSimpleCollision(UStaticMesh* StaticMesh)
 	StaticMesh->GetMeshDescription(0)->ComputeBoundingBox().GetCenterAndExtents(Center, Extents);
 	Extents *= (FVector)bs->BuildScale3D;
 
-	bs->Modify();
+	if(bUpdate)
+	{
+		bs->Modify();
+	}
 
 	// Create new GUID
 	bs->InvalidatePhysicsData();
@@ -135,7 +171,10 @@ int32 GenerateBoxAsSimpleCollision(UStaticMesh* StaticMesh)
 	RefreshCollisionChange(*StaticMesh);
 
 	// Mark staticmesh as dirty, to help make sure it gets saved.
-	StaticMesh->MarkPackageDirty();
+	if(bUpdate)
+	{
+		StaticMesh->MarkPackageDirty();
+	}
 
 	StaticMesh->bCustomizedCollision = true;	//mark the static mesh for collision customization
 
@@ -288,7 +327,7 @@ static void CalcBoundingSphere2(const FMeshDescription* MeshDescription, FSphere
 
 // // //
 
-int32 GenerateSphereAsSimpleCollision(UStaticMesh* StaticMesh)
+int32 GenerateSphereAsSimpleCollision(UStaticMesh* StaticMesh, bool bUpdate)
 {
 	if (!PromptToRemoveExistingCollision(StaticMesh))
 	{
@@ -317,7 +356,10 @@ int32 GenerateSphereAsSimpleCollision(UStaticMesh* StaticMesh)
 		return INDEX_NONE;
 	}
 
-	bs->Modify();
+	if(bUpdate)
+	{
+		bs->Modify();
+	}
 
 	// Create new GUID
 	bs->InvalidatePhysicsData();
@@ -331,7 +373,10 @@ int32 GenerateSphereAsSimpleCollision(UStaticMesh* StaticMesh)
 	RefreshCollisionChange(*StaticMesh);
 
 	// Mark staticmesh as dirty, to help make sure it gets saved.
-	StaticMesh->MarkPackageDirty();
+	if(bUpdate)
+	{
+		StaticMesh->MarkPackageDirty();
+	}
 
 	StaticMesh->bCustomizedCollision = true;	//mark the static mesh for collision customization
 	return bs->AggGeom.SphereElems.Num() - 1;
@@ -429,7 +474,7 @@ static void CalcBoundingSphyl(const FMeshDescription* MeshDescription, FSphere&
 
 // // //
 
-int32 GenerateSphylAsSimpleCollision(UStaticMesh* StaticMesh)
+int32 GenerateSphylAsSimpleCollision(UStaticMesh* StaticMesh, bool bUpdate)
 {
 	if (!PromptToRemoveExistingCollision(StaticMesh))
 	{
@@ -461,7 +506,10 @@ int32 GenerateSphylAsSimpleCollision(UStaticMesh* StaticMesh)
 		length = SMALL_NUMBER;
 	}
 
-	bs->Modify();
+	if(bUpdate)
+	{
+		bs->Modify();
+	}
 
 	// Create new GUID
 	bs->InvalidatePhysicsData();
@@ -477,7 +525,10 @@ int32 GenerateSphylAsSimpleCollision(UStaticMesh* StaticMesh)
 	RefreshCollisionChange(*StaticMesh);
 
 	// Mark staticmesh as dirty, to help make sure it gets saved.
-	StaticMesh->MarkPackageDirty();
+	if(bUpdate)
+	{
+		StaticMesh->MarkPackageDirty();
+	}
 
 	StaticMesh->bCustomizedCollision = true;	//mark the static mesh for collision customization
 
diff --git a/Engine/Source/Editor/UnrealEd/Private/GeomFitUtils.h b/Engine/Source/Editor/UnrealEd/Private/GeomFitUtils.h
--- a/Engine/Source/Editor/UnrealEd/Private/GeomFitUtils.h
+++ b/Engine/Source/Editor/UnrealEd/Private/GeomFitUtils.h
@@ -108,10 +108,10 @@ const FVector KDopDir26[26] =
 };
 
 // Utilities
-UNREALED_API int32 GenerateKDopAsSimpleCollision(UStaticMesh* StaticMesh, const TArray<FVector> &dirs);
-UNREALED_API int32 GenerateBoxAsSimpleCollision(UStaticMesh* StaticMesh);
-UNREALED_API int32 GenerateSphereAsSimpleCollision(UStaticMesh* StaticMesh);
-UNREALED_API int32 GenerateSphylAsSimpleCollision(UStaticMesh* StaticMesh);
+UNREALED_API int32 GenerateKDopAsSimpleCollision(UStaticMesh* StaticMesh, const TArray<FVector>& dirs, bool bUpdate = true);
+UNREALED_API int32 GenerateBoxAsSimpleCollision(UStaticMesh* StaticMesh, bool bUpdate = true);
+UNREALED_API int32 GenerateSphereAsSimpleCollision(UStaticMesh* StaticMesh, bool bUpdate = true);
+UNREALED_API int32 GenerateSphylAsSimpleCollision(UStaticMesh* StaticMesh, bool bUpdate = true);
 UNREALED_API void ComputeBoundingBox(UStaticMesh* StaticMesh, FVector& Center, FVector& Extents);
 
 /**
