MaterialX:
- Fix Place2d not correctly loaded due to change of path.

diff --git a/Engine/Plugins/Interchange/Runtime/Source/Import/Private/MaterialX/MaterialXUtils/MaterialXManager.cpp b/Engine/Plugins/Interchange/Runtime/Source/Import/Private/MaterialX/MaterialXUtils/MaterialXManager.cpp
--- a/Engine/Plugins/Interchange/Runtime/Source/Import/Private/MaterialX/MaterialXUtils/MaterialXManager.cpp
+++ b/Engine/Plugins/Interchange/Runtime/Source/Import/Private/MaterialX/MaterialXUtils/MaterialXManager.cpp
@@ -240,9 +240,11 @@ FMaterialXManager::FMaterialXManager()
 			{mx::Category::AbsorptionVDF,			FMaterialXMaterialFunction{TInPlaceType<EInterchangeMaterialXVDF>{}, EInterchangeMaterialXVDF::Absorption}},
 			{mx::Category::AnisotropicVDF,			FMaterialXMaterialFunction{TInPlaceType<EInterchangeMaterialXVDF>{}, EInterchangeMaterialXVDF::Anisotropic}},
 			// Utility nodes
-			{mx::Category::ArtisticIOR,				FMaterialXMaterialFunction{TInPlaceType<FString>{}, TEXT("/Interchange/Functions/MX_Artistic_IOR.MX_Artistic_IOR")}},
-			{mx::Category::RoughnessAnisotropy,		FMaterialXMaterialFunction{TInPlaceType<FString>{}, TEXT("/Interchange/Functions/MX_Roughness_Anisotropy.MX_Roughness_Anisotropy")}},
-			{mx::Category::RoughnessDual,			FMaterialXMaterialFunction{TInPlaceType<FString>{}, TEXT("/Interchange/Functions/MX_Roughness_Dual.MX_Roughness_Dual")}}
+			{mx::Category::ArtisticIOR,				FMaterialXMaterialFunction{TInPlaceType<FString>{}, TEXT("/InterchangeAssets/Functions/MX_Artistic_IOR.MX_Artistic_IOR")}},
+			{mx::Category::RoughnessAnisotropy,		FMaterialXMaterialFunction{TInPlaceType<FString>{}, TEXT("/InterchangeAssets/Functions/MX_Roughness_Anisotropy.MX_Roughness_Anisotropy")}},
+			{mx::Category::RoughnessDual,			FMaterialXMaterialFunction{TInPlaceType<FString>{}, TEXT("/InterchangeAssets/Functions/MX_Roughness_Dual.MX_Roughness_Dual")}},
+			// Math
+			{mx::Category::Place2D,					FMaterialXMaterialFunction{TInPlaceType<FString>{}, TEXT("/InterchangeAssets/Functions/MX_Place2D.MX_Place2D")}},
 		};
 
 		if(bIsSubstrateEnabled)
