[Interchange] Camera types not respected at level import fix.

#rnx

diff --git a/Engine/Plugins/Interchange/Runtime/Source/Import/Private/Scene/InterchangeActorFactory.cpp b/Engine/Plugins/Interchange/Runtime/Source/Import/Private/Scene/InterchangeActorFactory.cpp
--- a/Engine/Plugins/Interchange/Runtime/Source/Import/Private/Scene/InterchangeActorFactory.cpp
+++ b/Engine/Plugins/Interchange/Runtime/Source/Import/Private/Scene/InterchangeActorFactory.cpp
@@ -19,9 +19,11 @@ namespace UE::InterchangeActorFactory::Private
 	// component, and all the camera stuff on the camera component. This agrees with how the actor/root component is bound on
 	// LevelSequences, and is likely what users expect because when you place a camera actor on the level and move it around, you always
 	// affect the root component transform
+
+	template <class T>
 	void ApplyAllCameraCustomAttributes(
 		const UInterchangeFactoryBase::FImportSceneObjectsParams& CreateSceneObjectsParams,
-		UInterchangeActorFactoryNode* CameraFactoryNode,
+		T* CameraFactoryNode,
 		USceneComponent* RootSceneComponent,
 		USceneComponent* ChildCameraComponent
 	)
@@ -36,7 +38,7 @@ namespace UE::InterchangeActorFactory::Private
 		UInterchangeBaseNodeContainer* const NodeContainer = const_cast<UInterchangeBaseNodeContainer* const>(CreateSceneObjectsParams.NodeContainer);
 
 		// Create a temp factory node so we don't modify our existing nodes with our changes
-		UInterchangePhysicalCameraFactoryNode* FactoryNodeCopy = NewObject<UInterchangePhysicalCameraFactoryNode>(NodeContainer, NAME_None);
+		T* FactoryNodeCopy = NewObject<T>(NodeContainer, NAME_None);
 		FactoryNodeCopy->InitializeNode(
 			CameraFactoryNode->GetUniqueID(),
 			CameraFactoryNode->GetDisplayLabel(),
@@ -94,12 +96,12 @@ UObject* UInterchangeActorFactory::ImportSceneObject_GameThread(const UInterchan
 				if (FactoryNode->IsA<UInterchangePhysicalCameraFactoryNode>())
 				{
 					UCineCameraComponent* CameraComponent = Cast<UCineCameraComponent>(ObjectToUpdate);
-					ApplyAllCameraCustomAttributes(CreateSceneObjectsParams, FactoryNode, RootComponent, CameraComponent);
+					ApplyAllCameraCustomAttributes(CreateSceneObjectsParams, Cast<UInterchangePhysicalCameraFactoryNode>(FactoryNode), RootComponent, CameraComponent);
 				}
 				else if (FactoryNode->IsA<UInterchangeStandardCameraFactoryNode>())
 				{
 					UCameraComponent* CameraComponent = Cast<UCameraComponent>(ObjectToUpdate);
-					ApplyAllCameraCustomAttributes(CreateSceneObjectsParams, FactoryNode, RootComponent, CameraComponent);
+					ApplyAllCameraCustomAttributes(CreateSceneObjectsParams, Cast<UInterchangeStandardCameraFactoryNode>(FactoryNode), RootComponent, CameraComponent);
 				}
 				else
 				{
