Fix missing WITH_EDITORONLY_DATA surrounding

diff --git a/Engine/Plugins/Interchange/Runtime/Source/Import/Private/Mesh/InterchangeStaticMeshFactory.cpp b/Engine/Plugins/Interchange/Runtime/Source/Import/Private/Mesh/InterchangeStaticMeshFactory.cpp
--- a/Engine/Plugins/Interchange/Runtime/Source/Import/Private/Mesh/InterchangeStaticMeshFactory.cpp
+++ b/Engine/Plugins/Interchange/Runtime/Source/Import/Private/Mesh/InterchangeStaticMeshFactory.cpp
@@ -432,7 +432,11 @@ UInterchangeFactoryBase::FImportAssetResult UInterchangeStaticMeshFactory::Impor
 			{
 				if (bReimport)
 				{
+#if WITH_EDITORONLY_DATA
 					StaticMesh->GetBodySetup()->AggGeom.EmptyImportedElements();
+#else
+					StaticMesh->GetBodySetup()->AggGeom.EmptyElements();
+#endif
 				}
 				bImportedCustomCollision |= ImportBoxCollision(Arguments, StaticMesh, LodDataNode);
 				bImportedCustomCollision |= ImportCapsuleCollision(Arguments, StaticMesh, LodDataNode);
