[Interchange] Automatic T0 usage detection.

#rnx

diff --git a/Engine/Plugins/Interchange/Runtime/Source/Import/Private/Mesh/InterchangeSkeletalMeshFactory.cpp b/Engine/Plugins/Interchange/Runtime/Source/Import/Private/Mesh/InterchangeSkeletalMeshFactory.cpp
--- a/Engine/Plugins/Interchange/Runtime/Source/Import/Private/Mesh/InterchangeSkeletalMeshFactory.cpp
+++ b/Engine/Plugins/Interchange/Runtime/Source/Import/Private/Mesh/InterchangeSkeletalMeshFactory.cpp
@@ -1222,7 +1222,25 @@ UInterchangeFactoryBase::FImportAssetResult UInterchangeSkeletalMeshFactory::Beg
 		FImportAssetObjectLODData& ImportAssetObjectLODData = ImportAssetObjectData.LodDatas.AddDefaulted_GetRef();
 		ImportAssetObjectLODData.LodIndex = CurrentLodIndex;
 
-		SkeletonNode->GetCustomUseTimeZeroForBindPose(ImportAssetObjectLODData.bUseTimeZeroAsBindPose);
+		bool bUseTimeZeroAsBindPose = false;
+		SkeletonNode->GetCustomUseTimeZeroForBindPose(bUseTimeZeroAsBindPose);
+		if (!bUseTimeZeroAsBindPose)
+		{
+			bool bHasBoneWithoutBindPose = false;
+			UE::Interchange::Private::FSkeletonHelper::RecursiveBoneHasBindPose(Arguments.NodeContainer, RootJointNodeId, bHasBoneWithoutBindPose);
+			if (bHasBoneWithoutBindPose)
+			{
+				UInterchangeResultDisplay_Generic* Message = AddMessage<UInterchangeResultDisplay_Generic>();
+				Message->Text = FText::Format(NSLOCTEXT("InterchangeSkeletalMeshFactory", "BeginImportAsset_GameThread_NotAllJointsHaveBindPose", "Not all joints have BindPoses in the skeleton {0}, will use T0 as BindPose instead.")
+					, FText::FromString(SkeletonReference->GetName()));
+
+				UE_LOG(LogInterchangeImport, Display, TEXT("Not all joints have BindPoses in the skeleton %s, will use T0 as BindPose instead."), *SkeletonReference->GetName());
+
+				bUseTimeZeroAsBindPose = true;
+			}
+		}
+
+		ImportAssetObjectLODData.bUseTimeZeroAsBindPose = bUseTimeZeroAsBindPose;
 		
 		//Do not alter the skeletal mesh reference skeleton when importing geometry only
 		FReferenceSkeleton RefSkeleton;
diff --git a/Engine/Plugins/Interchange/Runtime/Source/Import/Private/Mesh/InterchangeSkeletonHelper.cpp b/Engine/Plugins/Interchange/Runtime/Source/Import/Private/Mesh/InterchangeSkeletonHelper.cpp
--- a/Engine/Plugins/Interchange/Runtime/Source/Import/Private/Mesh/InterchangeSkeletonHelper.cpp
+++ b/Engine/Plugins/Interchange/Runtime/Source/Import/Private/Mesh/InterchangeSkeletonHelper.cpp
@@ -178,6 +178,42 @@ namespace UE::Interchange::Private
 		}
 	}
 
+	void FSkeletonHelper::RecursiveBoneHasBindPose(const UInterchangeBaseNodeContainer* NodeContainer, const FString& JointNodeId, bool& bHasBoneWithoutBindPose)
+	{
+		if (bHasBoneWithoutBindPose)
+		{
+			return;
+		}
+
+		const UInterchangeSceneNode* JointNode = Cast<UInterchangeSceneNode>(NodeContainer->GetNode(JointNodeId));
+		if (!JointNode)
+		{
+			UE_LOG(LogInterchangeImport, Warning, TEXT("Invalid Skeleton Joint"));
+			return;
+		}
+
+		bool bHasBindPose;
+		if (!JointNode->GetCustomHasBindPose(bHasBindPose))
+		{
+			//if not set, then its presumed to have bind pose
+			bHasBindPose = true;
+		}
+
+		if (!bHasBindPose)
+		{
+			bHasBoneWithoutBindPose = true;
+		}
+
+		if (!bHasBoneWithoutBindPose)
+		{
+			const TArray<FString> ChildrenIds = NodeContainer->GetNodeChildrenUids(JointNodeId);
+			for (int32 ChildIndex = 0; ChildIndex < ChildrenIds.Num(); ++ChildIndex)
+			{
+				RecursiveBoneHasBindPose(NodeContainer, ChildrenIds[ChildIndex], bHasBoneWithoutBindPose);
+			}
+		}
+	}
+
 	void FSkeletonHelper::RecursiveAddBones(const UInterchangeBaseNodeContainer* NodeContainer
 		, const FString& JointNodeId
 		, TArray <FJointInfo>& JointInfos
diff --git a/Engine/Plugins/Interchange/Runtime/Source/Import/Public/InterchangeSkeletonHelper.h b/Engine/Plugins/Interchange/Runtime/Source/Import/Public/InterchangeSkeletonHelper.h
--- a/Engine/Plugins/Interchange/Runtime/Source/Import/Public/InterchangeSkeletonHelper.h
+++ b/Engine/Plugins/Interchange/Runtime/Source/Import/Public/InterchangeSkeletonHelper.h
@@ -39,6 +39,9 @@ namespace UE::Interchange::Private
 
 		static bool IsCompatibleSkeleton(const USkeleton* Skeleton, const FString RootJoinUid, const UInterchangeBaseNodeContainer* BaseNodeContainer, bool bConvertStaticToSkeletalActive, bool bCheckForIdenticalSkeleton = false);
 		static void RecursiveAddSkeletonMetaDataValues(UInterchangeBaseNodeContainer* NodeContainer, UInterchangeBaseNode* DestinationNode, const FString& JointUid);
+
+		static void RecursiveBoneHasBindPose(const UInterchangeBaseNodeContainer* NodeContainer, const FString& JointNodeId, bool& bHasBoneWithoutBindPose);
+
 	private:
 		static void RecursiveAddBones(const UInterchangeBaseNodeContainer* NodeContainer
 			, const FString& JointNodeId
diff --git a/Engine/Plugins/Interchange/Runtime/Source/Nodes/Private/InterchangeSceneNode.cpp b/Engine/Plugins/Interchange/Runtime/Source/Nodes/Private/InterchangeSceneNode.cpp
--- a/Engine/Plugins/Interchange/Runtime/Source/Nodes/Private/InterchangeSceneNode.cpp
+++ b/Engine/Plugins/Interchange/Runtime/Source/Nodes/Private/InterchangeSceneNode.cpp
@@ -439,3 +439,13 @@ void UInterchangeSceneNode::SetGlobalBindPoseReferenceForMeshUIDs(const TMap<FSt
 		MeshToGlobalBindPoseReferences.SetKeyValue(Entry.Key, Entry.Value);
 	}
 }
+
+bool UInterchangeSceneNode::SetCustomHasBindPose(const bool& AttributeValue)
+{
+	IMPLEMENT_NODE_ATTRIBUTE_SETTER_NODELEGATE(HasBindPose, bool);
+}
+
+bool UInterchangeSceneNode::GetCustomHasBindPose(bool& AttributeValue) const
+{
+	IMPLEMENT_NODE_ATTRIBUTE_GETTER(HasBindPose, bool);
+}
\ No newline at end of file
diff --git a/Engine/Plugins/Interchange/Runtime/Source/Nodes/Public/InterchangeSceneNode.h b/Engine/Plugins/Interchange/Runtime/Source/Nodes/Public/InterchangeSceneNode.h
--- a/Engine/Plugins/Interchange/Runtime/Source/Nodes/Public/InterchangeSceneNode.h
+++ b/Engine/Plugins/Interchange/Runtime/Source/Nodes/Public/InterchangeSceneNode.h
@@ -231,6 +231,13 @@ public:
 	bool GetGlobalBindPoseReferenceForMeshUID(const FString& MeshUID, FMatrix& GlobalBindPoseReference) const;
 
 
+	/** Sets if Joint has Bind Pose. Automatic T0 usage will be configured in case if the Skeleton contanis at least 1 Joint without BindPose. */
+	UFUNCTION(BlueprintCallable, Category = "Interchange | Node | Joint")
+	bool SetCustomHasBindPose(const bool& bHasBindPose);
+
+	/** Gets if the joint has BindPose (if the setter was used, otherwise returns with false and T0 evaluation presumes bHasBindPose==true). */
+	UFUNCTION(BlueprintCallable, Category = "Interchange | Node | Joint")
+	bool GetCustomHasBindPose(bool &bHasBindPose) const;
 
 private:
 
@@ -259,6 +266,9 @@ private:
 	//A scene node can reference an asset. Asset can be Mesh, Light, camera...
 	const UE::Interchange::FAttributeKey Macro_CustomAssetInstanceUidKey = UE::Interchange::FAttributeKey(TEXT("AssetInstanceUid"));
 
+	//Tracks if Scene Node Has Bind Pose.
+	const UE::Interchange::FAttributeKey Macro_CustomHasBindPoseKey = UE::Interchange::FAttributeKey(TEXT("HasBindPose"));
+
 	//A scene node can represent many special types
 	UE::Interchange::TArrayAttributeHelper<FString> NodeSpecializeTypes;
 
diff --git a/Engine/Plugins/Interchange/Runtime/Source/Parsers/Fbx/Private/FbxScene.cpp b/Engine/Plugins/Interchange/Runtime/Source/Parsers/Fbx/Private/FbxScene.cpp
--- a/Engine/Plugins/Interchange/Runtime/Source/Parsers/Fbx/Private/FbxScene.cpp
+++ b/Engine/Plugins/Interchange/Runtime/Source/Parsers/Fbx/Private/FbxScene.cpp
@@ -220,7 +220,11 @@ namespace UE
 					FbxAMatrix GlobalBindPoseJointMatrix = SDKScene->GetAnimationEvaluator()->GetNodeGlobalTransform(Node, 0);
 					TMap<FString, FMatrix> MeshIdToGlobalBindPoseReferenceMap;
 
-					FFbxMesh::GetGlobalJointBindPoseTransform(&Parser, SDKScene, Node, GlobalBindPoseJointMatrix, MeshIdToGlobalBindPoseReferenceMap, bBadBindPoseMessageDisplay); //if false it would use the previously calculated
+					if (!FFbxMesh::GetGlobalJointBindPoseTransform(&Parser, SDKScene, Node, GlobalBindPoseJointMatrix, MeshIdToGlobalBindPoseReferenceMap, bBadBindPoseMessageDisplay))
+					{
+						//Bind Pose not present for the Joint, we should use T0 for the entire skeleton:
+						UnrealNode->SetCustomHasBindPose(false);
+					}
 
 					FTransform GlobalBindPoseJointTransform = GetConvertedTransform(GlobalBindPoseJointMatrix);
 					UnrealNode->SetGlobalBindPoseReferenceForMeshUIDs(MeshIdToGlobalBindPoseReferenceMap);
