Interchange: Scene import LevelInstance

#rnx

diff --git a/Engine/Plugins/Interchange/Editor/Source/InterchangeEditor/Private/AssetDefinition_InterchangeSceneImportAsset.cpp b/Engine/Plugins/Interchange/Editor/Source/InterchangeEditor/Private/AssetDefinition_InterchangeSceneImportAsset.cpp
--- a/Engine/Plugins/Interchange/Editor/Source/InterchangeEditor/Private/AssetDefinition_InterchangeSceneImportAsset.cpp
+++ b/Engine/Plugins/Interchange/Editor/Source/InterchangeEditor/Private/AssetDefinition_InterchangeSceneImportAsset.cpp
@@ -8,6 +8,7 @@
 #include "InterchangeSceneImportAsset.h"
 
 #include "ContentBrowserMenuContexts.h"
+#include "LevelInstance/LevelInstanceActor.h"
 #include "Misc/App.h"
 #include "Misc/DelayedAutoRegister.h"
 #include "ThumbnailRendering/SceneThumbnailInfo.h"
@@ -31,6 +32,23 @@ UThumbnailInfo* UAssetDefinition_InterchangeSceneImportAsset::LoadThumbnailInfo(
 
 namespace MenuExtension_InterchangeSceneImportAsset
 {
+	ULevel* GetOriginalLevel(UInterchangeSceneImportAsset* Asset)
+	{
+		TArray<FSoftObjectPath> SceneSoftObjectPaths;
+		Asset->GetSceneSoftObjectPaths(SceneSoftObjectPaths);
+		for (const FSoftObjectPath& SoftObjectPath : SceneSoftObjectPaths)
+		{
+			if (AActor* Actor = Cast<AActor>(SoftObjectPath.TryLoad()))
+			{
+				if (!Actor->IsA<ALevelInstance>())
+				{
+					return Actor->GetLevel();
+				}
+			}
+		}
+		return nullptr;
+	}
+
 	void ExecuteReimportOneAsset(UInterchangeSceneImportAsset* Asset, const FString& FilePath)
 	{
 		using namespace UE::Interchange;
@@ -45,6 +63,7 @@ namespace MenuExtension_InterchangeSceneImportAsset
 			ImportAssetParameters.bIsAutomated = GIsAutomationTesting || FApp::IsUnattended() || IsRunningCommandlet() || GIsRunningUnattendedScript;
 			ImportAssetParameters.ReimportAsset = Asset;
 			ImportAssetParameters.ReimportSourceIndex = INDEX_NONE;
+			ImportAssetParameters.ImportLevel = GetOriginalLevel(Asset);
 
 			TTuple<FAssetImportResultRef, FSceneImportResultRef> ImportResult = InterchangeManager.ImportSceneAsync(FString(), SourceData, ImportAssetParameters);
 
diff --git a/Engine/Plugins/Interchange/Editor/Source/Utilities/Private/InterchangeEditorUtilities.cpp b/Engine/Plugins/Interchange/Editor/Source/Utilities/Private/InterchangeEditorUtilities.cpp
new file mode 100644
--- /dev/null
+++ b/Engine/Plugins/Interchange/Editor/Source/Utilities/Private/InterchangeEditorUtilities.cpp
@@ -0,0 +1,15 @@
+// Copyright Epic Games, Inc. All Rights Reserved. 
+#include "InterchangeEditorUtilities.h"
+
+#include "FileHelpers.h"
+#include "UObject/Package.h"
+
+#include UE_INLINE_GENERATED_CPP_BY_NAME(InterchangeEditorUtilities)
+
+bool UInterchangeEditorUtilities::SaveAsset(UObject* Asset)
+{
+	TArray<UPackage*> PackagesToSave;
+	PackagesToSave.Add(Asset->GetPackage());
+	FEditorFileUtils::EPromptReturnCode ReturnCode = FEditorFileUtils::PromptForCheckoutAndSave(PackagesToSave, false /*bCheckDirty*/, false /*bPromptToSave*/);
+	return (ReturnCode == FEditorFileUtils::PR_Success);
+}
diff --git a/Engine/Plugins/Interchange/Editor/Source/Utilities/Private/InterchangeEditorUtilitiesModule.cpp b/Engine/Plugins/Interchange/Editor/Source/Utilities/Private/InterchangeEditorUtilitiesModule.cpp
--- a/Engine/Plugins/Interchange/Editor/Source/Utilities/Private/InterchangeEditorUtilitiesModule.cpp
+++ b/Engine/Plugins/Interchange/Editor/Source/Utilities/Private/InterchangeEditorUtilitiesModule.cpp
@@ -2,6 +2,9 @@
 
 #include "InterchangeEditorUtilitiesModule.h"
 
+#include "InterchangeEditorUtilities.h"
+#include "InterchangeManager.h"
+
 
 #define LOCTEXT_NAMESPACE "InterchangeEditorUtilities"
 
@@ -17,6 +20,7 @@ IMPLEMENT_MODULE(FInterchangeEditorUtilitiesModule, InterchangeEditorUtilities)
 
 void FInterchangeEditorUtilitiesModule::StartupModule()
 {
+	UInterchangeManager::GetInterchangeManager().SetEditorUtilities(UInterchangeEditorUtilities::StaticClass());
 }
 
 void FInterchangeEditorUtilitiesModule::ShutdownModule()
diff --git a/Engine/Plugins/Interchange/Editor/Source/Utilities/Public/InterchangeEditorUtilities.h b/Engine/Plugins/Interchange/Editor/Source/Utilities/Public/InterchangeEditorUtilities.h
new file mode 100644
--- /dev/null
+++ b/Engine/Plugins/Interchange/Editor/Source/Utilities/Public/InterchangeEditorUtilities.h
@@ -0,0 +1,19 @@
+// Copyright Epic Games, Inc. All Rights Reserved.
+
+#pragma once
+
+#include "InterchangeEditorUtilitiesBase.h"
+
+#include "InterchangeEditorUtilities.generated.h"
+
+UCLASS()
+class INTERCHANGEEDITORUTILITIES_API UInterchangeEditorUtilities : public UInterchangeEditorUtilitiesBase
+{
+	GENERATED_BODY()
+
+public:
+
+protected:
+
+	virtual bool SaveAsset(UObject* Asset) override;
+};
diff --git a/Engine/Plugins/Interchange/Runtime/Source/FactoryNodes/Private/InterchangeLevelFactoryNode.cpp b/Engine/Plugins/Interchange/Runtime/Source/FactoryNodes/Private/InterchangeLevelFactoryNode.cpp
new file mode 100644
--- /dev/null
+++ b/Engine/Plugins/Interchange/Runtime/Source/FactoryNodes/Private/InterchangeLevelFactoryNode.cpp
@@ -0,0 +1,72 @@
+// Copyright Epic Games, Inc. All Rights Reserved.
+
+#include "InterchangeLevelFactoryNode.h"
+
+#include "Engine/Level.h"
+#include "Engine/World.h"
+#include "Math/Box.h"
+
+UInterchangeLevelFactoryNode::UInterchangeLevelFactoryNode()
+{
+	CustomActorFactoryNodeUids.Initialize(Attributes, TEXT("ActorFactoryNodeUids"));
+}
+
+UClass* UInterchangeLevelFactoryNode::GetObjectClass() const
+{
+	return UWorld::StaticClass();
+}
+
+int32 UInterchangeLevelFactoryNode::GetCustomActorFactoryNodeUidCount() const
+{
+	return CustomActorFactoryNodeUids.GetCount();
+}
+
+void UInterchangeLevelFactoryNode::GetCustomActorFactoryNodeUids(TArray<FString>& OutActorFactoryNodeUids) const
+{
+	CustomActorFactoryNodeUids.GetItems(OutActorFactoryNodeUids);
+}
+
+void UInterchangeLevelFactoryNode::GetCustomActorFactoryNodeUid(const int32 Index, FString& OutActorFactoryNodeUid) const
+{
+	CustomActorFactoryNodeUids.GetItem(Index, OutActorFactoryNodeUid);
+}
+
+bool UInterchangeLevelFactoryNode::AddCustomActorFactoryNodeUid(const FString& ActorFactoryNodeUid)
+{
+	return CustomActorFactoryNodeUids.AddItem(ActorFactoryNodeUid);
+}
+
+bool UInterchangeLevelFactoryNode::RemoveCustomActorFactoryNodeUid(const FString& ActorFactoryNodeUid)
+{
+	return CustomActorFactoryNodeUids.RemoveItem(ActorFactoryNodeUid);
+}
+
+bool UInterchangeLevelFactoryNode::GetCustomSceneImportAssetFactoryNodeUid(FString& AttributeValue) const
+{
+	IMPLEMENT_NODE_ATTRIBUTE_GETTER(SceneImportAssetFactoryNodeUid, FString);
+}
+
+bool UInterchangeLevelFactoryNode::SetCustomSceneImportAssetFactoryNodeUid(const FString& AttributeValue)
+{
+	IMPLEMENT_NODE_ATTRIBUTE_SETTER_NODELEGATE(SceneImportAssetFactoryNodeUid, FString);
+}
+
+bool UInterchangeLevelFactoryNode::GetCustomShouldCreateLevel(bool& AttributeValue) const
+{
+	IMPLEMENT_NODE_ATTRIBUTE_GETTER(ShouldCreateLevel, bool);
+}
+
+bool UInterchangeLevelFactoryNode::SetCustomShouldCreateLevel(const bool& AttributeValue)
+{
+	IMPLEMENT_NODE_ATTRIBUTE_SETTER_NODELEGATE(ShouldCreateLevel, bool);
+}
+
+bool UInterchangeLevelFactoryNode::GetCustomCreateWorldPartitionLevel(bool& AttributeValue) const
+{
+	IMPLEMENT_NODE_ATTRIBUTE_GETTER(CreateWorldPartitionLevel, bool);
+}
+
+bool UInterchangeLevelFactoryNode::SetCustomCreateWorldPartitionLevel(const bool& AttributeValue)
+{
+	IMPLEMENT_NODE_ATTRIBUTE_SETTER_NODELEGATE(CreateWorldPartitionLevel, bool);
+}
diff --git a/Engine/Plugins/Interchange/Runtime/Source/FactoryNodes/Private/InterchangeLevelInstanceActorFactoryNode.cpp b/Engine/Plugins/Interchange/Runtime/Source/FactoryNodes/Private/InterchangeLevelInstanceActorFactoryNode.cpp
new file mode 100644
--- /dev/null
+++ b/Engine/Plugins/Interchange/Runtime/Source/FactoryNodes/Private/InterchangeLevelInstanceActorFactoryNode.cpp
@@ -0,0 +1,13 @@
+// Copyright Epic Games, Inc. All Rights Reserved.
+
+#include "InterchangeLevelInstanceActorFactoryNode.h"
+
+bool UInterchangeLevelInstanceActorFactoryNode::GetCustomLevelReference(FString& AttributeValue) const
+{
+	IMPLEMENT_NODE_ATTRIBUTE_GETTER(LevelReference, FString);
+}
+
+bool UInterchangeLevelInstanceActorFactoryNode::SetCustomLevelReference(const FString& AttributeValue)
+{
+	IMPLEMENT_NODE_ATTRIBUTE_SETTER_NODELEGATE(LevelReference, FString);
+}
diff --git a/Engine/Plugins/Interchange/Runtime/Source/FactoryNodes/Public/InterchangeLevelFactoryNode.h b/Engine/Plugins/Interchange/Runtime/Source/FactoryNodes/Public/InterchangeLevelFactoryNode.h
new file mode 100644
--- /dev/null
+++ b/Engine/Plugins/Interchange/Runtime/Source/FactoryNodes/Public/InterchangeLevelFactoryNode.h
@@ -0,0 +1,110 @@
+// Copyright Epic Games, Inc. All Rights Reserved.
+
+#pragma once
+
+#include "CoreMinimal.h"
+#include "Nodes/InterchangeFactoryBaseNode.h"
+
+#include "InterchangeLevelFactoryNode.generated.h"
+
+UCLASS(BlueprintType)
+class INTERCHANGEFACTORYNODES_API UInterchangeLevelFactoryNode : public UInterchangeFactoryBaseNode
+{
+	GENERATED_BODY()
+
+	UInterchangeLevelFactoryNode();
+
+public:
+	static FStringView StaticAssetTypeName()
+	{
+		return TEXT("LevelFactory");
+	}
+
+	/**
+	 * Return the node type name of the class, we use this when reporting errors
+	 */
+	virtual FString GetTypeName() const override
+	{
+		const FString TypeName = TEXT("LevelFactoryNode");
+		return TypeName;
+	}
+
+	/** Get the class this node want to create */
+	UFUNCTION(BlueprintCallable, Category = "Interchange | Node | LevelFactory")
+	virtual class UClass* GetObjectClass() const override;
+
+	/**
+	 * This function allow to retrieve the number of track dependencies for this object.
+	 */
+	UFUNCTION(BlueprintCallable, Category = "Interchange | Node | LevelFactory")
+	int32 GetCustomActorFactoryNodeUidCount() const;
+
+	/**
+	 * This function allow to retrieve all actor factory node unique id for this object.
+	 */
+	UFUNCTION(BlueprintCallable, Category = "Interchange | Node | LevelFactory")
+	void GetCustomActorFactoryNodeUids(TArray<FString>& OutActorFactoryNodeUids) const;
+
+	/**
+	 * This function allow to retrieve one actor factory node unique id for this object.
+	 */
+	UFUNCTION(BlueprintCallable, Category = "Interchange | Node | LevelFactory")
+	void GetCustomActorFactoryNodeUid(const int32 Index, FString& OutActorFactoryNodeUid) const;
+
+	/**
+	 * Add one actor factory node unique id to this object.
+	 */
+	UFUNCTION(BlueprintCallable, Category = "Interchange | Node | LevelFactory")
+	bool AddCustomActorFactoryNodeUid(const FString& ActorFactoryNodeUid);
+
+	/**
+	 * Remove one actor factory node unique id from this object.
+	 */
+	UFUNCTION(BlueprintCallable, Category = "Interchange | Node | LevelFactory")
+	bool RemoveCustomActorFactoryNodeUid(const FString& ActorFactoryNodeUid);
+
+	/**
+	 * Get the actor factory node unique id that hold the re-import data.
+	 */
+	UFUNCTION(BlueprintCallable, Category = "Interchange | Node | LevelFactory")
+	bool GetCustomSceneImportAssetFactoryNodeUid(FString& AttributeValue) const;
+
+	/**
+	 * Set the actor factory node unique id that hold the re-import data.
+	 */
+	UFUNCTION(BlueprintCallable, Category = "Interchange | Node | LevelFactory")
+	bool SetCustomSceneImportAssetFactoryNodeUid(const FString& AttributeValue);
+
+	/**
+	 * Get actors bounding box.
+	 */
+	UFUNCTION(BlueprintCallable, Category = "Interchange | Node | LevelFactory")
+	bool GetCustomShouldCreateLevel(bool& AttributeValue) const;
+
+	/**
+	 * Set actors bounding box.
+	 */
+	UFUNCTION(BlueprintCallable, Category = "Interchange | Node | LevelFactory")
+	bool SetCustomShouldCreateLevel(const bool& AttributeValue);
+
+	/**
+	 * If true, created world partition level.
+	 */
+	UFUNCTION(BlueprintCallable, Category = "Interchange | Node | LevelFactory")
+	bool GetCustomCreateWorldPartitionLevel(bool& AttributeValue) const;
+
+	/**
+	 * If true, created world partition level.
+	 */
+	UFUNCTION(BlueprintCallable, Category = "Interchange | Node | LevelFactory")
+	bool SetCustomCreateWorldPartitionLevel(const bool& AttributeValue);
+
+private:
+	
+	const UE::Interchange::FAttributeKey Macro_CustomShouldCreateLevelKey = UE::Interchange::FAttributeKey(TEXT("ShouldCreateLevel"));
+	const UE::Interchange::FAttributeKey Macro_CustomSceneImportAssetFactoryNodeUidKey = UE::Interchange::FAttributeKey(TEXT("SceneImportAssetFactoryNodeUid"));
+	const UE::Interchange::FAttributeKey Macro_CustomCreateWorldPartitionLevelKey = UE::Interchange::FAttributeKey(TEXT("CreateWorldPartitionLevel"));
+	const UE::Interchange::FAttributeKey Macro_CustomParentLevelReferenceKey = UE::Interchange::FAttributeKey(TEXT("ParentLevelReference"));
+
+	UE::Interchange::TArrayAttributeHelper<FString> CustomActorFactoryNodeUids;
+};
diff --git a/Engine/Plugins/Interchange/Runtime/Source/FactoryNodes/Public/InterchangeLevelInstanceActorFactoryNode.h b/Engine/Plugins/Interchange/Runtime/Source/FactoryNodes/Public/InterchangeLevelInstanceActorFactoryNode.h
new file mode 100644
--- /dev/null
+++ b/Engine/Plugins/Interchange/Runtime/Source/FactoryNodes/Public/InterchangeLevelInstanceActorFactoryNode.h
@@ -0,0 +1,28 @@
+// Copyright Epic Games, Inc. All Rights Reserved.
+
+#pragma once
+
+#include "InterchangeActorFactoryNode.h"
+#include "UObject/Object.h"
+#include "UObject/ObjectMacros.h"
+
+
+#include "InterchangeLevelInstanceActorFactoryNode.generated.h"
+
+UCLASS(BlueprintType)
+class INTERCHANGEFACTORYNODES_API UInterchangeLevelInstanceActorFactoryNode : public UInterchangeActorFactoryNode
+{
+	GENERATED_BODY()
+public:
+
+	/** Get the Level this Level instance actor refer to. */
+	UFUNCTION(BlueprintCallable, Category = "Interchange | Node | LevelInstanceActor")
+	bool GetCustomLevelReference(FString& AttributeValue) const;
+	/** Set the Level this Level instance actor refer to. */
+	UFUNCTION(BlueprintCallable, Category = "Interchange | Node | LevelInstanceActor")
+	bool SetCustomLevelReference(const FString& AttributeValue);
+
+private:
+	const UE::Interchange::FAttributeKey Macro_CustomParentLevelKey = UE::Interchange::FAttributeKey(TEXT("ParentLevel"));
+	const UE::Interchange::FAttributeKey Macro_CustomLevelReferenceKey = UE::Interchange::FAttributeKey(TEXT("LevelReference"));
+};
\ No newline at end of file
diff --git a/Engine/Plugins/Interchange/Runtime/Source/Import/Private/InterchangeImportModule.cpp b/Engine/Plugins/Interchange/Runtime/Source/Import/Private/InterchangeImportModule.cpp
--- a/Engine/Plugins/Interchange/Runtime/Source/Import/Private/InterchangeImportModule.cpp
+++ b/Engine/Plugins/Interchange/Runtime/Source/Import/Private/InterchangeImportModule.cpp
@@ -23,6 +23,8 @@
 #include "Scene/InterchangeActorFactory.h"
 #include "Scene/InterchangeCameraActorFactory.h"
 #include "Scene/InterchangeSceneImportAssetFactory.h"
+#include "Scene/InterchangeLevelFactory.h"
+#include "Scene/InterchangeLevelInstanceActorFactory.h"
 #include "Scene/InterchangeLightActorFactory.h"
 #include "Scene/InterchangeSceneVariantSetsFactory.h"
 #include "Scene/InterchangeStaticMeshActorFactory.h"
@@ -105,6 +107,8 @@ void FInterchangeImportModule::StartupModule()
 		InterchangeManager.RegisterFactory(UInterchangeStaticMeshFactory::StaticClass());
 		InterchangeManager.RegisterFactory(UInterchangePhysicsAssetFactory::StaticClass());
 		InterchangeManager.RegisterFactory(UInterchangeActorFactory::StaticClass());
+		InterchangeManager.RegisterFactory(UInterchangeLevelFactory::StaticClass());
+		InterchangeManager.RegisterFactory(UInterchangeLevelInstanceActorFactory::StaticClass());
 		InterchangeManager.RegisterFactory(UInterchangeLevelSequenceFactory::StaticClass());
 		InterchangeManager.RegisterFactory(UInterchangeAnimSequenceFactory::StaticClass());
 		InterchangeManager.RegisterFactory(UInterchangeCineCameraActorFactory::StaticClass());
diff --git a/Engine/Plugins/Interchange/Runtime/Source/Import/Private/Scene/InterchangeLevelFactory.cpp b/Engine/Plugins/Interchange/Runtime/Source/Import/Private/Scene/InterchangeLevelFactory.cpp
new file mode 100644
--- /dev/null
+++ b/Engine/Plugins/Interchange/Runtime/Source/Import/Private/Scene/InterchangeLevelFactory.cpp
@@ -0,0 +1,174 @@
+// Copyright Epic Games, Inc. All Rights Reserved.
+
+#include "Scene/InterchangeLevelFactory.h"
+
+#include "Engine/World.h"
+#include "InterchangeAssetImportData.h"
+#include "InterchangeEditorUtilitiesBase.h"
+#include "InterchangeImportCommon.h"
+#include "InterchangeImportLog.h"
+#include "InterchangeLevelFactoryNode.h"
+#include "InterchangeLevelInstanceActorFactoryNode.h"
+#include "InterchangeManager.h"
+#include "InterchangeResult.h"
+#include "InterchangeSourceData.h"
+#include "InterchangeTranslatorBase.h"
+#include "InterchangeVariantSetNode.h"
+#include "LevelInstance/LevelInstanceActor.h"
+#include "Misc/App.h"
+#include "Nodes/InterchangeBaseNode.h"
+#include "Nodes/InterchangeBaseNodeContainer.h"
+#include "UObject/Class.h"
+#include "WorldPartition/LoaderAdapter/LoaderAdapterShape.h"
+#include "WorldPartition/WorldPartitionEditorLoaderAdapter.h"
+#include "WorldPartition/WorldPartition.h"
+
+#define LOCTEXT_NAMESPACE "InterchangeLevelFactory"
+
+namespace UE::Interchange::Private::InterchangeLevelFactory
+{
+	const UInterchangeLevelFactoryNode* GetFactoryNode(const UInterchangeFactoryBase::FImportAssetObjectParams& Arguments, UClass* TargetClass)
+	{
+		if (!Arguments.NodeContainer || !Arguments.AssetNode || !Arguments.AssetNode->GetObjectClass()->IsChildOf(TargetClass))
+		{
+			return nullptr;
+		}
+
+		return Cast<UInterchangeLevelFactoryNode>(Arguments.AssetNode);
+	}
+
+	UObject* FindOrCreateAsset(const UInterchangeFactoryBase::FImportAssetObjectParams& Arguments, UClass* TargetClass)
+	{
+
+		UObject* TargetAsset = nullptr;
+		if (Arguments.ReimportObject)
+		{
+			TargetAsset = Arguments.ReimportObject;
+		}
+		else
+		{
+			FSoftObjectPath ExistingWorldAssetPath;
+			if (Arguments.AssetNode->GetCustomReferenceObject(ExistingWorldAssetPath))
+			{
+				TargetAsset = ExistingWorldAssetPath.TryLoad();
+			}
+			if (!TargetAsset)
+			{
+				TargetAsset = StaticFindObject(nullptr, Arguments.Parent, *Arguments.AssetName);
+			}
+		}
+		
+		bool bCreateWorldPartition = false;
+		//Verify and ensure if we should create a level or not
+		if (UInterchangeLevelFactoryNode* LevelFactoryNode = Cast<UInterchangeLevelFactoryNode>(Arguments.AssetNode))
+		{
+			bool bShouldCreateLevel = false;
+			if (LevelFactoryNode->GetCustomShouldCreateLevel(bShouldCreateLevel))
+			{
+				ensure(bShouldCreateLevel ? TargetAsset == nullptr : TargetAsset != nullptr);
+			}
+			LevelFactoryNode->GetCustomCreateWorldPartitionLevel(bCreateWorldPartition);
+		}
+		// Create a new asset or return existing asset, if possible
+		if (!TargetAsset)
+		{
+			if(!ensure(IsInGameThread()))
+			{
+				return nullptr;
+			}
+			
+			// Create a new world.
+			const bool bAddToRoot = false;
+
+			constexpr bool bEnableWorldPartitionStreaming = false;
+			constexpr bool bInformEngineOfWorld = true;
+
+			// Those are the init values taken from the default in UWorld::CreateWorld + CreateWorldPartition.
+			UWorld::InitializationValues InitValues = UWorld::InitializationValues()
+				.ShouldSimulatePhysics(false)
+				.EnableTraceCollision(true)
+				.CreateNavigation(!FApp::IsGame())
+				.CreateAISystem(!FApp::IsGame())
+				.CreateWorldPartition(bCreateWorldPartition)
+				.EnableWorldPartitionStreaming(bEnableWorldPartitionStreaming);
+
+			UWorld* NewWorld = UWorld::CreateWorld(EWorldType::Editor, bInformEngineOfWorld, *Arguments.AssetName, Cast<UPackage>(Arguments.Parent), bAddToRoot, ERHIFeatureLevel::Num, &InitValues);
+			NewWorld->SetFlags(RF_Public | RF_Standalone);
+			TargetAsset = NewWorld;
+		}
+		else if (!TargetAsset->GetClass()->IsChildOf(TargetClass))
+		{
+			TargetAsset = nullptr;
+		}
+
+		return TargetAsset;
+	}
+}
+
+
+UClass* UInterchangeLevelFactory::GetFactoryClass() const
+{
+	return UWorld::StaticClass();
+}
+
+
+UInterchangeFactoryBase::FImportAssetResult UInterchangeLevelFactory::BeginImportAsset_GameThread(const FImportAssetObjectParams& Arguments)
+{
+	TRACE_CPUPROFILER_EVENT_SCOPE(UInterchangeLevelFactory::BeginImportAsset_GameThread);
+
+	using namespace UE::Interchange::Private::InterchangeLevelFactory;
+
+	UClass* TargetClass = GetFactoryClass();
+
+	if (GetFactoryNode(Arguments, TargetClass) == nullptr)
+	{
+		return {};
+	}
+
+	UWorld* WorldAsset = Cast<UWorld>(FindOrCreateAsset(Arguments, TargetClass));
+
+	if (!WorldAsset)
+	{
+		UE_LOG(LogInterchangeImport, Warning, TEXT("Could not create world asset %s"), *Arguments.AssetName);
+	}
+	//Getting the file Hash will cache it into the source data
+	else if (ensure(Arguments.SourceData))
+	{
+		Arguments.SourceData->GetFileContentHash();
+	}
+
+#if WITH_EDITOR
+	if (!FApp::IsGame() && WorldAsset)
+	{
+		WorldAsset->PreEditChange(nullptr);
+	}
+#endif //WITH_EDITOR
+
+	FImportAssetResult Result;
+	Result.ImportedObject = WorldAsset;
+
+	return Result;
+}
+
+void UInterchangeLevelFactory::SetupObject_GameThread(const FSetupObjectParams& Arguments)
+{
+	TRACE_CPUPROFILER_EVENT_SCOPE(UInterchangeLevelFactory::SetupObject_GameThread);
+	Super::SetupObject_GameThread(Arguments);
+	if (!ensure(IsInGameThread()))
+	{
+		return;
+	}
+
+	UWorld* World = Cast<UWorld>(Arguments.ImportedObject);
+	UInterchangeLevelFactoryNode* FactoryNode = Cast<UInterchangeLevelFactoryNode>(Arguments.FactoryNode);
+	if (ensure(FactoryNode && World && Arguments.SourceData))
+	{
+		/** Apply all FactoryNode custom attributes to the level sequence asset */
+		FactoryNode->ApplyAllCustomAttributeToObject(World);
+
+		//The scene import data object will add them self to this world has a sub object
+	}
+}
+
+
+#undef LOCTEXT_NAMESPACE
\ No newline at end of file
diff --git a/Engine/Plugins/Interchange/Runtime/Source/Import/Private/Scene/InterchangeLevelInstanceActorFactory.cpp b/Engine/Plugins/Interchange/Runtime/Source/Import/Private/Scene/InterchangeLevelInstanceActorFactory.cpp
new file mode 100644
--- /dev/null
+++ b/Engine/Plugins/Interchange/Runtime/Source/Import/Private/Scene/InterchangeLevelInstanceActorFactory.cpp
@@ -0,0 +1,74 @@
+// Copyright Epic Games, Inc. All Rights Reserved. 
+#include "Scene/InterchangeLevelInstanceActorFactory.h"
+
+#include "Engine/World.h"
+#include "InterchangeEditorUtilitiesBase.h"
+#include "InterchangeImportLog.h"
+#include "InterchangeLevelFactoryNode.h"
+#include "InterchangeLevelInstanceActorFactoryNode.h"
+#include "InterchangeManager.h"
+#include "LevelInstance/LevelInstanceActor.h"
+#include "Nodes/InterchangeBaseNodeContainer.h"
+#include "Scene/InterchangeActorHelper.h"
+
+#include UE_INLINE_GENERATED_CPP_BY_NAME(InterchangeLevelInstanceActorFactory)
+
+UClass* UInterchangeLevelInstanceActorFactory::GetFactoryClass() const
+{
+	return ALevelInstance::StaticClass();
+}
+
+UObject* UInterchangeLevelInstanceActorFactory::ProcessActor(AActor& SpawnedActor, const UInterchangeActorFactoryNode& FactoryNode, const UInterchangeBaseNodeContainer& NodeContainer, const FImportSceneObjectsParams& Params)
+{
+	auto SuperResult = [&]()
+		{
+			return Super::ProcessActor(SpawnedActor, FactoryNode, NodeContainer, Params);
+		};
+
+	const UInterchangeLevelInstanceActorFactoryNode* LevelInstanceActorFactoryNode = Cast<UInterchangeLevelInstanceActorFactoryNode>(&FactoryNode);
+	if (!LevelInstanceActorFactoryNode)
+	{
+		UE_LOG(LogInterchangeImport, Warning, TEXT("UInterchangeLevelInstanceActorFactory::ProcessActor: Level instance actor factory node is null"));
+		return SuperResult();
+	}
+
+	ALevelInstance* LevelInstanceActor = Cast<ALevelInstance>(&SpawnedActor);
+	if (!LevelInstanceActor)
+	{
+		FString AssetName = LevelInstanceActorFactoryNode->GetDisplayLabel();
+		UE_LOG(LogInterchangeImport, Warning, TEXT("UInterchangeLevelInstanceActorFactory::ProcessActor: Actor level instance was not created %s."), *AssetName);
+		return SuperResult();
+	}
+	//Get the world reference by this level instance actor
+	FString ReferenceLevelFactoryNodeUid;
+	if (!LevelInstanceActorFactoryNode->GetCustomLevelReference(ReferenceLevelFactoryNodeUid))
+	{
+		UE_LOG(LogInterchangeImport, Warning, TEXT("UInterchangeLevelInstanceActorFactory::ProcessActor: Actor level instance do not reference any level factory node."));
+		return SuperResult();
+	}
+
+	UInterchangeLevelFactoryNode* ReferenceLevelFactoryNode = Cast<UInterchangeLevelFactoryNode>(NodeContainer.GetFactoryNode(ReferenceLevelFactoryNodeUid));
+	if (!ReferenceLevelFactoryNode)
+	{
+		UE_LOG(LogInterchangeImport, Warning, TEXT("UInterchangeLevelInstanceActorFactory::ProcessActor: Level factory node is invalid."));
+		return SuperResult();
+	}
+	FSoftObjectPath ReferenceLevelSoftObjectPath;
+	if (!ReferenceLevelFactoryNode->GetCustomReferenceObject(ReferenceLevelSoftObjectPath))
+	{
+		UE_LOG(LogInterchangeImport, Warning, TEXT("UInterchangeLevelInstanceActorFactory::ProcessActor: Level factory node do not reference any UWorld soft object path."));
+		return SuperResult();
+	}
+
+	UWorld* ReferenceWorld = Cast<UWorld>(ReferenceLevelSoftObjectPath.TryLoad());
+	if (!ReferenceWorld)
+	{
+		UE_LOG(LogInterchangeImport, Warning, TEXT("UInterchangeLevelInstanceActorFactory::ProcessActor: UWorld soft object path point on an invalid UWorld."));
+		return SuperResult();
+	}
+#if WITH_EDITOR
+	LevelInstanceActor->SetActorLabel(ReferenceLevelFactoryNode->GetDisplayLabel());
+#endif
+
+	return SuperResult();
+}
diff --git a/Engine/Plugins/Interchange/Runtime/Source/Import/Private/Scene/InterchangeSkeletalMeshActorFactory.cpp b/Engine/Plugins/Interchange/Runtime/Source/Import/Private/Scene/InterchangeSkeletalMeshActorFactory.cpp
--- a/Engine/Plugins/Interchange/Runtime/Source/Import/Private/Scene/InterchangeSkeletalMeshActorFactory.cpp
+++ b/Engine/Plugins/Interchange/Runtime/Source/Import/Private/Scene/InterchangeSkeletalMeshActorFactory.cpp
@@ -9,13 +9,15 @@
 
 #include "Animation/SkeletalMeshActor.h"
 #include "Components/SkeletalMeshComponent.h"
+#include "Engine/SkeletalMesh.h"
+#include "Engine/SkinnedAssetAsyncCompileUtils.h"
 
 #include "Nodes/InterchangeBaseNodeContainer.h"
 
 #include UE_INLINE_GENERATED_CPP_BY_NAME(InterchangeSkeletalMeshActorFactory)
 
 
-UObject* UInterchangeSkeletalMeshActorFactory::ProcessActor(AActor& SpawnedActor, const UInterchangeActorFactoryNode& /*FactoryNode*/, const UInterchangeBaseNodeContainer& /*NodeContainer*/, const FImportSceneObjectsParams& /*Params*/)
+UObject* UInterchangeSkeletalMeshActorFactory::ProcessActor(AActor& SpawnedActor, const UInterchangeActorFactoryNode& /*FactoryNode*/, const UInterchangeBaseNodeContainer& /*NodeContainer*/, const FImportSceneObjectsParams& Params)
 {
 	ASkeletalMeshActor* SkeletalMeshActor = Cast<ASkeletalMeshActor>(&SpawnedActor);
 
@@ -26,9 +28,24 @@ UObject* UInterchangeSkeletalMeshActorFactory::ProcessActor(AActor& SpawnedActor
 
 	if (USkeletalMeshComponent* SkeletalMeshComponent = SkeletalMeshActor->GetSkeletalMeshComponent())
 	{
-		SkeletalMeshComponent->UnregisterComponent();
+		if (USkeletalMesh* SkeletalMesh = SkeletalMeshComponent->GetSkeletalMeshAsset())
+		{
+			if (Params.ImportAssets.Contains(SkeletalMesh))
+			{
+#if WITH_EDITOR
+				//If we are importing the skeletal mesh, we need to use a async build scope to unregister the component to avoid a deadlock
+				FSkinnedAssetAsyncBuildScope AsyncBuildScope(SkeletalMesh);
+#endif
+				SkeletalMeshComponent->UnregisterComponent();
+			}
+			else
+			{
+				SkeletalMeshComponent->UnregisterComponent();
+			}
+			
 
-		return SkeletalMeshComponent;
+			return SkeletalMeshComponent;
+		}
 	}
 
 	return nullptr;
diff --git a/Engine/Plugins/Interchange/Runtime/Source/Import/Public/Scene/InterchangeLevelFactory.h b/Engine/Plugins/Interchange/Runtime/Source/Import/Public/Scene/InterchangeLevelFactory.h
new file mode 100644
--- /dev/null
+++ b/Engine/Plugins/Interchange/Runtime/Source/Import/Public/Scene/InterchangeLevelFactory.h
@@ -0,0 +1,29 @@
+// Copyright Epic Games, Inc. All Rights Reserved.
+
+#pragma once
+
+#include "CoreMinimal.h"
+#include "InterchangeFactoryBase.h"
+#include "UObject/Object.h"
+#include "UObject/ObjectMacros.h"
+
+#include "InterchangeLevelFactory.generated.h"
+
+UCLASS(BlueprintType, Experimental)
+class INTERCHANGEIMPORT_API UInterchangeLevelFactory : public UInterchangeFactoryBase
+{
+	GENERATED_BODY()
+public:
+
+	//////////////////////////////////////////////////////////////////////////
+	// Interchange factory base interface begin
+
+	virtual UClass* GetFactoryClass() const override;
+	virtual FImportAssetResult BeginImportAsset_GameThread(const FImportAssetObjectParams& Arguments) override;
+	virtual void SetupObject_GameThread(const FSetupObjectParams& Arguments) override;
+
+	// Interchange factory base interface end
+	//////////////////////////////////////////////////////////////////////////
+};
+
+
diff --git a/Engine/Plugins/Interchange/Runtime/Source/Import/Public/Scene/InterchangeLevelInstanceActorFactory.h b/Engine/Plugins/Interchange/Runtime/Source/Import/Public/Scene/InterchangeLevelInstanceActorFactory.h
new file mode 100644
--- /dev/null
+++ b/Engine/Plugins/Interchange/Runtime/Source/Import/Public/Scene/InterchangeLevelInstanceActorFactory.h
@@ -0,0 +1,40 @@
+// Copyright Epic Games, Inc. All Rights Reserved.
+
+#pragma once
+
+#include "InterchangeFactoryBase.h"
+#include "Scene/InterchangeActorFactory.h"
+#include "UObject/Object.h"
+#include "UObject/ObjectMacros.h"
+
+#include "InterchangeLevelInstanceActorFactory.generated.h"
+
+class AActor;
+class UInterchangeActorFactoryNode;
+class UInterchangeBaseNodeContainer;
+
+UCLASS(BlueprintType)
+class INTERCHANGEIMPORT_API UInterchangeLevelInstanceActorFactory : public UInterchangeActorFactory
+{
+	GENERATED_BODY()
+public:
+
+	//////////////////////////////////////////////////////////////////////////
+	// Interchange factory base interface begin
+
+	virtual UClass* GetFactoryClass() const override;
+
+	// Interchange factory base interface end
+	//////////////////////////////////////////////////////////////////////////
+
+protected:
+	//////////////////////////////////////////////////////////////////////////
+	// Interchange actor factory interface begin
+
+	virtual UObject* ProcessActor(class AActor& SpawnedActor, const UInterchangeActorFactoryNode& FactoryNode, const UInterchangeBaseNodeContainer& NodeContainer, const FImportSceneObjectsParams& Params) override;
+
+	// Interchange actor factory interface end
+	//////////////////////////////////////////////////////////////////////////
+};
+
+
diff --git a/Engine/Plugins/Interchange/Runtime/Source/Pipelines/Private/InterchangeGenericScenesPipeline.cpp b/Engine/Plugins/Interchange/Runtime/Source/Pipelines/Private/InterchangeGenericScenesPipeline.cpp
--- a/Engine/Plugins/Interchange/Runtime/Source/Pipelines/Private/InterchangeGenericScenesPipeline.cpp
+++ b/Engine/Plugins/Interchange/Runtime/Source/Pipelines/Private/InterchangeGenericScenesPipeline.cpp
@@ -2,37 +2,45 @@
 
 #include "InterchangeGenericScenesPipeline.h"
 
+#include "Nodes/InterchangeUserDefinedAttribute.h"
 #include "InterchangeActorFactoryNode.h"
 #include "InterchangeAssetImportData.h"
 #include "InterchangeCameraNode.h"
 #include "InterchangeCameraFactoryNode.h"
 #include "InterchangeCommonPipelineDataFactoryNode.h"
-#include "InterchangeSceneImportAsset.h"
-#include "InterchangeSceneImportAssetFactoryNode.h"
-#include "InterchangeSkeletalMeshLodDataNode.h"
+#include "InterchangeLevelInstanceActorFactoryNode.h"
 #include "InterchangeLightNode.h"
 #include "InterchangeLightFactoryNode.h"
 #include "InterchangeMeshActorFactoryNode.h"
 #include "InterchangeMeshNode.h"
 #include "InterchangeDecalActorFactoryNode.h"
 #include "InterchangeDecalNode.h"
+#include "InterchangeEditorUtilitiesBase.h"
+#include "InterchangeLevelFactoryNode.h"
+#include "InterchangeManager.h"
 #include "InterchangePipelineLog.h"
 #include "InterchangePipelineMeshesUtilities.h"
+#include "InterchangeSceneImportAsset.h"
+#include "InterchangeSceneImportAssetFactoryNode.h"
 #include "InterchangeSceneNode.h"
 #include "InterchangeSceneVariantSetsFactoryNode.h"
-#include "InterchangeVariantSetNode.h"
-#include "Nodes/InterchangeUserDefinedAttribute.h"
-#include "InterchangeSkeletonFactoryNode.h"
 #include "InterchangeSkeletalMeshFactoryNode.h"
+#include "InterchangeSkeletalMeshLodDataNode.h"
+#include "InterchangeSkeletonFactoryNode.h"
+#include "InterchangeVariantSetNode.h"
+
 
 #include "Animation/SkeletalMeshActor.h"
 #include "CineCameraActor.h"
 #include "Engine/DirectionalLight.h"
+#include "Engine/Level.h"
 #include "Engine/PointLight.h"
 #include "Engine/RectLight.h"
 #include "Engine/SpotLight.h"
 #include "Engine/StaticMeshActor.h"
 #include "Engine/World.h"
+#include "WorldPartition/WorldPartition.h"
+#include "LevelInstance/LevelInstanceActor.h"
 #include "Misc/PackageName.h"
 
 #if WITH_EDITOR
@@ -269,19 +277,96 @@ void UInterchangeGenericLevelPipeline::ExecutePipeline(UInterchangeBaseNodeConta
 	});
 
 #if WITH_EDITORONLY_DATA
+	const FString FilePath = FPaths::ConvertRelativePathToFull(InSourceDatas[0]->GetFilename());
+
+	UWorld* TargetWorld = GWorld->GetCurrentLevel()->GetWorld();
+	
+	UInterchangeLevelFactoryNode* ParentLevelFactoryNode = nullptr;
+	LevelInstanceActorFactoryNode = nullptr;
+
+	if (TargetWorld)
+	{
+		//create the parent level factory node
+		ParentLevelFactoryNode = NewObject<UInterchangeLevelFactoryNode>(BaseNodeContainer, NAME_None);
+		const FString ParentDisplayLabel = TEXT("ParentLevel_") + TargetWorld->GetName();
+		const FString ParentNodeUid = ParentDisplayLabel;
+		ParentLevelFactoryNode->InitializeNode(ParentNodeUid, ParentDisplayLabel, EInterchangeNodeContainerType::FactoryData);
+		ParentLevelFactoryNode->SetCustomShouldCreateLevel(false);
+		ParentLevelFactoryNode->SetCustomReferenceObject(TargetWorld);
+		BaseNodeContainer->IterateNodesOfType<UInterchangeFactoryBaseNode>(
+			[&ParentLevelFactoryNode](const FString& NodeUid, UInterchangeFactoryBaseNode* Node)
+			{
+				// Add all assets (not actors) before calling the ParentLevelFactoryNode factory
+				ParentLevelFactoryNode->AddFactoryDependencyUid(NodeUid);
+			}
+		);
+		BaseNodeContainer->AddNode(ParentLevelFactoryNode);
+
+		UWorld* ReimportReferenceWorld = nullptr;
+		if (ULevel* Level = Cast<ULevel>(ReimportLevel.TryLoad()))
+		{
+			ReimportReferenceWorld = Level->GetWorld();
+		}
+
+		const bool bReimportReferenceWorld = ReimportReferenceWorld && ReimportReferenceWorld == TargetWorld;
+
+		if (!bReimportReferenceWorld && SceneHierarchyType == EInterchangeSceneHierarchyType::CreateLevelInstanceActor)
+		{
+			ensure(!LevelFactoryNode);
+			const FString DisplayLabel = ReimportReferenceWorld ? ReimportReferenceWorld->GetName() : ("Level_") + FPaths::GetBaseFilename(FilePath);
+			const FString NodeUid = TEXT("Level_") + FilePath;
+			ensure(!BaseNodeContainer->IsNodeUidValid(NodeUid));
+			LevelFactoryNode = NewObject<UInterchangeLevelFactoryNode>(BaseNodeContainer, NAME_None);
+			LevelFactoryNode->InitializeNode(NodeUid, DisplayLabel, EInterchangeNodeContainerType::FactoryData);
+			LevelFactoryNode->SetCustomCreateWorldPartitionLevel(false);
+			LevelFactoryNode->SetCustomShouldCreateLevel(true);
+
+			if (ReimportReferenceWorld)
+			{
+				LevelFactoryNode->SetCustomShouldCreateLevel(false);
+				LevelFactoryNode->SetCustomReferenceObject(ReimportReferenceWorld);
+				//When we re-import we want the FinalizeObject_GameThread to be call on the level factory for this level node
+				LevelFactoryNode->SetForceNodeReimport();
+			}
+			//Look at the parent level factory node before the level factory node
+			LevelFactoryNode->AddFactoryDependencyUid(ParentLevelFactoryNode->GetUniqueID());
+
+			BaseNodeContainer->AddNode(LevelFactoryNode);
+
+			//Create a level instance actor
+			{
+				LevelInstanceActorFactoryNode = NewObject<UInterchangeLevelInstanceActorFactoryNode>(BaseNodeContainer, NAME_None);
+
+				if (ensure(LevelInstanceActorFactoryNode))
+				{
+					LevelInstanceActorFactoryNode->SetCustomActorClassName(ALevelInstance::StaticClass()->GetPathName());
+					//This actor has the same name has the level it reference
+					const FString ActorDisplayLabel = DisplayLabel;
+					FString ActorNodeUid = TEXT("LevelInstance_") + FilePath;
+					LevelInstanceActorFactoryNode->InitializeNode(ActorNodeUid, ActorDisplayLabel, EInterchangeNodeContainerType::FactoryData);
+					ActorNodeUid = BaseNodeContainer->AddNode(LevelInstanceActorFactoryNode);
+					//Set the level this actor is referring
+					LevelInstanceActorFactoryNode->SetCustomLevelReference(LevelFactoryNode->GetUniqueID());
+					//We ensure the actor will be created after the parent and reference world are create or ready
+					LevelInstanceActorFactoryNode->AddFactoryDependencyUid(LevelFactoryNode->GetUniqueID());
+					//Add the Level instance actor into the parent level (TargetWorld)
+					LevelInstanceActorFactoryNode->SetCustomLevelUid(ParentLevelFactoryNode->GetUniqueID());
+					//Add the actor to the world that we want to create the actor
+					ParentLevelFactoryNode->AddCustomActorFactoryNodeUid(ActorNodeUid);
+				}
+			}
+		}
+	}
+
 	// Add the SceneImportData factory node
 	{
 		ensure(!SceneImportFactoryNode);
-
-		const FString FilePath = FPaths::ConvertRelativePathToFull(InSourceDatas[0]->GetFilename());
 		const FString DisplayLabel = TEXT("SceneImport_") + FPaths::GetBaseFilename(FilePath);
 		const FString NodeUid = TEXT("SceneImport_") + FilePath;
 		const FString FactoryNodeUid = UInterchangeFactoryBaseNode::BuildFactoryNodeUid(NodeUid);
 		ensure(!BaseNodeContainer->IsNodeUidValid(FactoryNodeUid));
 		SceneImportFactoryNode = NewObject<UInterchangeSceneImportAssetFactoryNode>(BaseNodeContainer, NAME_None);
-
 		SceneImportFactoryNode->InitializeNode(FactoryNodeUid, DisplayLabel, EInterchangeNodeContainerType::FactoryData);
-
 		// Add dependency to all the factory nodes created so far
 		BaseNodeContainer->IterateNodesOfType<UInterchangeFactoryBaseNode>(
 			[this](const FString& NodeUid, UInterchangeFactoryBaseNode* Node)
@@ -293,6 +378,7 @@ void UInterchangeGenericLevelPipeline::ExecutePipeline(UInterchangeBaseNodeConta
 
 		BaseNodeContainer->AddNode(SceneImportFactoryNode);
 	}
+
 #endif
 
 	/* Find all scene node that are active joint. Non active joint should be convert to actor if they are in a static mesh hierarchy */
@@ -450,7 +536,20 @@ void UInterchangeGenericLevelPipeline::ExecuteSceneNodePreImport(const FTransfor
 
 	ActorFactoryNode->InitializeNode(FactoryNodeUid, SceneNode->GetDisplayLabel(), EInterchangeNodeContainerType::FactoryData);
 	const FString ActorFactoryNodeUid = BaseNodeContainer->AddNode(ActorFactoryNode);
-
+#if WITH_EDITORONLY_DATA
+	//The level must be create before any actor asset since all actors will be create in the specified level
+	if (LevelFactoryNode)
+	{
+		ActorFactoryNode->AddFactoryDependencyUid(LevelFactoryNode->GetUniqueID());
+		ActorFactoryNode->SetCustomLevelUid(LevelFactoryNode->GetUniqueID());
+		LevelFactoryNode->AddCustomActorFactoryNodeUid(ActorFactoryNodeUid);
+		//The level instance actor must be create after the actor on the reference level are created
+		if (LevelInstanceActorFactoryNode)
+		{
+			LevelInstanceActorFactoryNode->AddFactoryDependencyUid(ActorFactoryNode->GetUniqueID());
+		}
+	}
+#endif
 	// The translator is responsible to provide a unique name
 	ActorFactoryNode->SetAssetName(SceneNode->GetAssetName());
 
@@ -976,7 +1075,7 @@ void UInterchangeGenericLevelPipeline::ExecutePostImportPipeline(const UIntercha
 
 	if (UInterchangeSceneImportAsset* SceneImportAsset = Cast<UInterchangeSceneImportAsset>(CreatedAsset))
 	{
-		if (!bIsReimportContext || !bDeleteMissingActors)
+		if (!bIsReimportContext || !(bDeleteMissingActors || bDeleteMissingAssets))
 		{
 			SceneImportAsset->UpdateSceneObjects();
 			return;
@@ -1023,7 +1122,7 @@ void UInterchangeGenericLevelPipeline::ExecutePostImportPipeline(const UIntercha
 						{
 							ActorsToDelete.Add(Cast<AActor>(Object));
 						}
-						else
+						else if(!Object->IsA<UWorld>()) //Avoid deleting UWorld asset
 						{
 							AssetsToForceDelete.Add(Object);
 						}
@@ -1032,9 +1131,12 @@ void UInterchangeGenericLevelPipeline::ExecutePostImportPipeline(const UIntercha
 			}
 		}
 
-		Private::DeleteActors(ActorsToDelete);
+		if (bDeleteMissingActors)
+		{
+			Private::DeleteActors(ActorsToDelete);
+		}
 
-		if (!AssetsToForceDelete.IsEmpty())
+		if (bDeleteMissingAssets && !AssetsToForceDelete.IsEmpty())
 		{
 			Private::DeleteAssets(AssetsToForceDelete);
 		}
@@ -1052,6 +1154,101 @@ void UInterchangeGenericLevelPipeline::ExecutePostImportPipeline(const UIntercha
 			}
 		}
 	}
+
+	if (UWorld* World = Cast<UWorld>(CreatedAsset))
+	{
+		PostPipelineImportData.AddWorld(World);
+	}
+	else if (ALevelInstance* LevelInstanceActor = Cast<ALevelInstance>(CreatedAsset))
+	{
+		const UInterchangeLevelInstanceActorFactoryNode* LevelInstanceFactoryNode = Cast<UInterchangeLevelInstanceActorFactoryNode>(BaseNodeContainer->GetFactoryNode(NodeKey));
+		FString LevelFactoryNodeUid;
+		if (LevelInstanceFactoryNode->GetCustomLevelReference(LevelFactoryNodeUid))
+		{
+			if (const UInterchangeFactoryBaseNode* ReferenceLevelFactoryNode = BaseNodeContainer->GetFactoryNode(LevelFactoryNodeUid))
+			{
+				FSoftObjectPath ReferenceLevelPath;
+				if (ReferenceLevelFactoryNode->GetCustomReferenceObject(ReferenceLevelPath))
+				{
+					if (UWorld* ReferenceWorld = Cast<UWorld>(ReferenceLevelPath.TryLoad()))
+					{
+						PostPipelineImportData.AddLevelInstanceActor(LevelInstanceActor, ReferenceWorld);
+					}
+				}
+			}
+		}
+	}
 #endif
 }
 
+#if WITH_EDITORONLY_DATA
+
+void UInterchangeGenericLevelPipeline::FPostPipelineImportData::AddWorld(UWorld* World)
+{
+	if (!Worlds.Contains(World))
+	{
+		Worlds.Add(World, false);
+		TArray<ALevelInstance*> CompletedLevelInstances;
+		for (TPair<ALevelInstance*, UWorld*>& LevelInstanceAndReferenceWorld : ReferenceWorldPerLevelInstanceToUpdates)
+		{
+			if (UpdateLevelInstanceInternal(LevelInstanceAndReferenceWorld.Key, LevelInstanceAndReferenceWorld.Value))
+			{
+				CompletedLevelInstances.Add(LevelInstanceAndReferenceWorld.Key);
+			}
+		}
+
+		//Remove level instance completed
+		for (ALevelInstance* CompletedLevelInstance : CompletedLevelInstances)
+		{
+			ReferenceWorldPerLevelInstanceToUpdates.Remove(CompletedLevelInstance);
+		}
+	}
+}
+
+void UInterchangeGenericLevelPipeline::FPostPipelineImportData::AddLevelInstanceActor(ALevelInstance* LevelInstanceActor, UWorld* ReferenceWorld)
+{
+	if (!UpdateLevelInstanceInternal(LevelInstanceActor, ReferenceWorld))
+	{
+		ReferenceWorldPerLevelInstanceToUpdates.Add(LevelInstanceActor, ReferenceWorld);
+	}
+}
+
+bool UInterchangeGenericLevelPipeline::FPostPipelineImportData::UpdateLevelInstanceInternal(ALevelInstance* LevelInstanceActor, UWorld* ReferenceWorld)
+{
+	if (!Worlds.Contains(ReferenceWorld))
+	{
+		return false;
+	}
+	UWorld* ParentWorld = LevelInstanceActor->GetWorld();
+	if (!Worlds.Contains(ParentWorld))
+	{
+		return false;
+	}
+
+	bool& bWorldSaved = Worlds.FindChecked(ReferenceWorld);
+	
+	if (!bWorldSaved)
+	{
+		if (UInterchangeEditorUtilitiesBase* EditorUtilities = UInterchangeManager::GetInterchangeManager().GetEditorUtilities())
+		{
+			if (!EditorUtilities->SaveAsset(ReferenceWorld))
+			{
+				UE_LOG(LogInterchangePipeline, Warning, TEXT("UInterchangeGenericAssetsPipeline: Cannot save the level instance actor (%s) referenced world (%s)"), *LevelInstanceActor->GetName(), *ReferenceWorld->GetName());
+			}
+		}
+		bWorldSaved = true;
+	}
+
+	ParentWorld->PreEditChange(nullptr);
+
+	LevelInstanceActor->SetWorldAsset(ReferenceWorld);
+	LevelInstanceActor->UpdateLevelInstanceFromWorldAsset();
+	LevelInstanceActor->LoadLevelInstance();
+	
+
+	ParentWorld->PostEditChange();
+
+	return true;
+}
+
+#endif //WITH_EDITORONLY_DATA
\ No newline at end of file
diff --git a/Engine/Plugins/Interchange/Runtime/Source/Pipelines/Public/InterchangeGenericScenesPipeline.h b/Engine/Plugins/Interchange/Runtime/Source/Pipelines/Public/InterchangeGenericScenesPipeline.h
--- a/Engine/Plugins/Interchange/Runtime/Source/Pipelines/Public/InterchangeGenericScenesPipeline.h
+++ b/Engine/Plugins/Interchange/Runtime/Source/Pipelines/Public/InterchangeGenericScenesPipeline.h
@@ -11,10 +11,21 @@
 
 #include "InterchangeGenericScenesPipeline.generated.h"
 
+class ALevelInstance;
 class UInterchangeActorFactoryNode;
+class UInterchangeLevelFactoryNode;
+class UInterchangeLevelInstanceActorFactoryNode;
 class UInterchangeSceneNode;
 class UInterchangeSceneVariantSetsNode;
 class UInterchangeSceneImportAssetFactoryNode;
+class UWorld;
+
+UENUM(BlueprintType)
+enum class EInterchangeSceneHierarchyType : uint8
+{
+	CreateLevelActors UMETA(DisplayName = "Create level actors", ToolTip = "Create actors in the current editor world for all scene nodes in the source hierarchy."),
+	CreateLevelInstanceActor UMETA(DisplayName = "Create a level instance actor", ToolTip = "Create a world containing all scene nodes in the source hierarchy and reference the world has a level instance actor in the current editor world."),
+};
 
 UCLASS(BlueprintType, editinlinenew)
 class INTERCHANGEPIPELINES_API UInterchangeGenericLevelPipeline : public UInterchangePipelineBase
@@ -22,27 +33,31 @@ class INTERCHANGEPIPELINES_API UInterchangeGenericLevelPipeline : public UInterc
 	GENERATED_BODY()
 public:
 	/** The name of the pipeline that will be display in the import dialog. */
-	UPROPERTY(EditAnywhere, BlueprintReadWrite, Category = "Level Reimport Policy", meta = (StandAlonePipelineProperty = "True", PipelineInternalEditionData = "True"))
+	UPROPERTY(EditAnywhere, BlueprintReadWrite, Category = "Scene", meta = (StandAlonePipelineProperty = "True", PipelineInternalEditionData = "True"))
 	FString PipelineDisplayName;
 
 	/* Set the reimport strategy when reimporting into the level. */
-	UPROPERTY(EditAnywhere, BlueprintReadWrite, Category = "Level Reimport Policy", meta = (SubCategory = "Actors properties", AdjustPipelineAndRefreshDetailOnChange = "True"))
+	UPROPERTY(EditAnywhere, BlueprintReadWrite, Category = "Scene", meta = (SubCategory = "Actors properties", AdjustPipelineAndRefreshDetailOnChange = "True"))
 	EReimportStrategyFlags ReimportPropertyStrategy = EReimportStrategyFlags::ApplyNoProperties;
 
+	/* Choose how you want to import the hierarchy. */
+	UPROPERTY(EditAnywhere, BlueprintReadWrite, Category = "Scene")
+	EInterchangeSceneHierarchyType SceneHierarchyType = EInterchangeSceneHierarchyType::CreateLevelActors;
+
 	/* If enabled, deletes actors that were not part of the translation when reimporting into a level. */
-	UPROPERTY(EditAnywhere, BlueprintReadWrite, Category = "Level Reimport Policy", meta = (SubCategory = "Reimport Actors"))
+	UPROPERTY(EditAnywhere, BlueprintReadWrite, Category = "Scene", meta = (SubCategory = "Reimport Actors"))
 	bool bDeleteMissingActors = false;
 
 	/* If enabled, respawns actors that were deleted in the editor prior to a reimport. */
-	UPROPERTY(EditAnywhere, BlueprintReadWrite, Category = "Level Reimport Policy", meta = (SubCategory = "Reimport Actors"))
+	UPROPERTY(EditAnywhere, BlueprintReadWrite, Category = "Scene", meta = (SubCategory = "Reimport Actors"))
 	bool bForceReimportDeletedActors = false;
 
 	/* If enabled, recreates assets that were deleted in the editor prior to reimporting into a level. */
-	UPROPERTY(EditAnywhere, BlueprintReadWrite, Category = "Level Reimport Policy", meta = (SubCategory = "Reimport Assets"))
+	UPROPERTY(EditAnywhere, BlueprintReadWrite, Category = "Scene", meta = (SubCategory = "Reimport Assets"))
 	bool bForceReimportDeletedAssets = false;
 
 	/* If enabled, deletes assets that were not part of the translation when reimporting into a level. */
-	UPROPERTY(EditAnywhere, BlueprintReadWrite, Category = "Level Reimport Policy", meta = (SubCategory = "Reimport Assets"))
+	UPROPERTY(EditAnywhere, BlueprintReadWrite, Category = "Scene", meta = (SubCategory = "Reimport Assets"))
 	bool bDeleteMissingAssets = false;
 
 	/** BEGIN UInterchangePipelineBase overrides */
@@ -94,6 +109,22 @@ protected:
 	 * of a level will not work at runtime
 	 */
 	UInterchangeSceneImportAssetFactoryNode* SceneImportFactoryNode = nullptr;
+	UInterchangeLevelFactoryNode* LevelFactoryNode = nullptr;
+	UInterchangeLevelInstanceActorFactoryNode* LevelInstanceActorFactoryNode = nullptr;
+
+	struct FPostPipelineImportData
+	{
+		void AddWorld(UWorld* World);
+		void AddLevelInstanceActor(ALevelInstance* LevelInstanceActor, UWorld* ReferenceWorld);
+	private:
+		bool UpdateLevelInstanceInternal(ALevelInstance* LevelInstanceActor, UWorld* ReferenceWorld);
+
+		//The bool is to know if we have save this world already or not
+		TMap<UWorld*, bool> Worlds;
+		TMap<ALevelInstance*, UWorld*> ReferenceWorldPerLevelInstanceToUpdates;
+	};
+
+	FPostPipelineImportData PostPipelineImportData;
 #endif
 };
 
diff --git a/Engine/Source/Runtime/Interchange/Core/Private/InterchangePipelineBase.cpp b/Engine/Source/Runtime/Interchange/Core/Private/InterchangePipelineBase.cpp
--- a/Engine/Source/Runtime/Interchange/Core/Private/InterchangePipelineBase.cpp
+++ b/Engine/Source/Runtime/Interchange/Core/Private/InterchangePipelineBase.cpp
@@ -445,6 +445,11 @@ void GatherObjectAndWeakObjectPtrs(UClass* Class, UObject* Object/*Value*/, TMap
 			if (SubObject)
 			{
 				FString PropertyName = Property->GetName();
+				if (ObjectPtrs.Contains(PropertyName))
+				{
+					//Prevent circular fetch (ULevel and UWorld will create an infinite loop)
+					continue;
+				}
 				ObjectPtrs.Add(PropertyName, SubObject);
 
 				UClass* SubObjectPropertyClass = SubObject->GetClass();
diff --git a/Engine/Source/Runtime/Interchange/Core/Private/Nodes/InterchangeFactoryBaseNode.cpp b/Engine/Source/Runtime/Interchange/Core/Private/Nodes/InterchangeFactoryBaseNode.cpp
--- a/Engine/Source/Runtime/Interchange/Core/Private/Nodes/InterchangeFactoryBaseNode.cpp
+++ b/Engine/Source/Runtime/Interchange/Core/Private/Nodes/InterchangeFactoryBaseNode.cpp
@@ -212,6 +212,17 @@ bool UInterchangeFactoryBaseNode::SetCustomSubPath(const FString& AttributeValue
 	IMPLEMENT_NODE_ATTRIBUTE_SETTER_NODELEGATE(SubPath, FString)
 }
 
+bool UInterchangeFactoryBaseNode::GetCustomLevelUid(FString& AttributeValue) const
+{
+	IMPLEMENT_NODE_ATTRIBUTE_GETTER(LevelUid, FString);
+}
+
+bool UInterchangeFactoryBaseNode::SetCustomLevelUid(const FString& AttributeValue)
+{
+	IMPLEMENT_NODE_ATTRIBUTE_SETTER_NODELEGATE(LevelUid, FString);
+}
+
+
 int32 UInterchangeFactoryBaseNode::GetFactoryDependenciesCount() const
 {
 	return FactoryDependencies.GetCount();
diff --git a/Engine/Source/Runtime/Interchange/Core/Public/InterchangeFactoryBase.h b/Engine/Source/Runtime/Interchange/Core/Public/InterchangeFactoryBase.h
--- a/Engine/Source/Runtime/Interchange/Core/Public/InterchangeFactoryBase.h
+++ b/Engine/Source/Runtime/Interchange/Core/Public/InterchangeFactoryBase.h
@@ -220,6 +220,11 @@ public:
 		 * Factory base node associated with the reimported scene node.
 		 */
 		const UInterchangeFactoryBaseNode* ReimportFactoryNode = nullptr;
+
+		/**
+		 * Get all assets we are importing the actor could point on
+		 */
+		TArray<TObjectPtr<UObject>> ImportAssets;
 	};
 
 	/**
diff --git a/Engine/Source/Runtime/Interchange/Core/Public/InterchangePipelineBase.h b/Engine/Source/Runtime/Interchange/Core/Public/InterchangePipelineBase.h
--- a/Engine/Source/Runtime/Interchange/Core/Public/InterchangePipelineBase.h
+++ b/Engine/Source/Runtime/Interchange/Core/Public/InterchangePipelineBase.h
@@ -520,6 +520,12 @@ public:
 	UPROPERTY()
 	FString ContentImportPath;
 
+	/*
+	 * Extra data available for scene pipeline, when we do a re-import we need to know what is the level so we do not create a new one.
+	 */
+	UPROPERTY()
+	FSoftObjectPath ReimportLevel;
+
 protected:
 
 	INTERCHANGECORE_API virtual FString GetPipelineDisplayName() const;
diff --git a/Engine/Source/Runtime/Interchange/Core/Public/Nodes/InterchangeFactoryBaseNode.h b/Engine/Source/Runtime/Interchange/Core/Public/Nodes/InterchangeFactoryBaseNode.h
--- a/Engine/Source/Runtime/Interchange/Core/Public/Nodes/InterchangeFactoryBaseNode.h
+++ b/Engine/Source/Runtime/Interchange/Core/Public/Nodes/InterchangeFactoryBaseNode.h
@@ -223,6 +223,18 @@ public:
 	UFUNCTION(BlueprintCallable, Category = "Interchange | Node")
 	INTERCHANGECORE_API bool SetCustomSubPath(const FString& AttributeValue);
 
+	/**
+	 * If this node represent a scene asset (actor), return a specific level in which we will create this scene asset.
+	 */
+	UFUNCTION(BlueprintCallable, Category = "Interchange | Node | ActorFactory")
+	INTERCHANGECORE_API bool GetCustomLevelUid(FString& AttributeValue) const;
+
+	/**
+	 * If this node represent a scene asset (actor), you can set a specific level in which we will create this scene asset.
+	 */
+	UFUNCTION(BlueprintCallable, Category = "Interchange | Node | ActorFactory")
+	INTERCHANGECORE_API bool SetCustomLevelUid(const FString& AttributeValue);
+
 	/**
 	 * Retrieve the number of factory dependencies for this object.
 	 */
@@ -364,6 +376,7 @@ protected:
 private:
 
 	const UE::Interchange::FAttributeKey Macro_CustomSubPathKey = UE::Interchange::FAttributeKey(TEXT("SubPath"));
+	const UE::Interchange::FAttributeKey Macro_CustomLevelUidKey = UE::Interchange::FAttributeKey(TEXT("LevelUid"));
 	const UE::Interchange::FAttributeKey Macro_CustomReferenceObjectKey = UE::Interchange::FAttributeKey(TEXT("ReferenceObject"));
 };
 
diff --git a/Engine/Source/Runtime/Interchange/Engine/Private/InterchangeManager.cpp b/Engine/Source/Runtime/Interchange/Engine/Private/InterchangeManager.cpp
--- a/Engine/Source/Runtime/Interchange/Engine/Private/InterchangeManager.cpp
+++ b/Engine/Source/Runtime/Interchange/Engine/Private/InterchangeManager.cpp
@@ -18,6 +18,7 @@
 #include "Framework/Notifications/NotificationManager.h"
 #include "InterchangeAssetImportData.h"
 #include "InterchangeBlueprintPipelineBase.h"
+#include "InterchangeEditorUtilitiesBase.h"
 #include "InterchangeFactoryBase.h"
 #include "InterchangeEngineLogPrivate.h"
 #include "InterchangeProjectSettings.h"
@@ -977,6 +978,9 @@ UInterchangeManager& UInterchangeManager::GetInterchangeManager()
 
 		bIsCreatingSingleton = false;
 
+		//Make sure we have a valid editor utilities
+		InterchangeManager->SetEditorUtilities(UInterchangeEditorUtilitiesBase::StaticClass());
+
 		//We cancel any running task when we pre exit the engine
 		FCoreDelegates::OnEnginePreExit.AddLambda([]()
 		{
@@ -1006,6 +1010,9 @@ UInterchangeManager& UInterchangeManager::GetInterchangeManager()
 				InterchangeManager->QueuedPostImportTasksTickerHandle.Reset();
 			}
 
+			//Release the editor utilities resources
+			InterchangeManager->EditorUtilities.Reset();
+
 			//Release the InterchangeManager object
 			InterchangeManager.Reset();
 			InterchangeManagerScopeOfLifeEnded = true;
@@ -1906,6 +1913,7 @@ UInterchangeManager::ImportInternal(const FString& ContentPath, const UInterchan
 				Context = bImportScene ? EInterchangePipelineContext::SceneImport : EInterchangePipelineContext::AssetImport;
 			}
 			Pipeline->ContentImportPath = ContentBasePath;
+			Pipeline->ReimportLevel = bIsReimport && bImportScene ? TaskData.ImportLevel : nullptr;
 			FInterchangePipelineContextParams ContextParams;
 			ContextParams.ContextType = Context;
 			ContextParams.ReimportAsset = TaskData.ReimportObject;
@@ -2273,6 +2281,22 @@ UInterchangeManager::ImportInternal(const FString& ContentPath, const UInterchan
 	return TTuple<UE::Interchange::FAssetImportResultRef, UE::Interchange::FSceneImportResultRef>{ AsyncHelper->AssetImportResult, AsyncHelper->SceneImportResult };
 }
 
+void UInterchangeManager::SetEditorUtilities(UClass* EditorUtilitiesClass)
+{
+	//Create the object
+	if (UInterchangeEditorUtilitiesBase* NewEditorUtilities = NewObject<UInterchangeEditorUtilitiesBase>(GetTransientPackage(), EditorUtilitiesClass, NAME_None))
+	{
+		//If successful, assign the new editor utilities
+		EditorUtilities.Reset();
+		EditorUtilities = TStrongObjectPtr<UInterchangeEditorUtilitiesBase>(NewEditorUtilities);
+	}
+}
+
+UInterchangeEditorUtilitiesBase* UInterchangeManager::GetEditorUtilities() const
+{
+	return EditorUtilities.Get();
+}
+
 bool UInterchangeManager::IsObjectBeingImported(UObject* Object) const
 {
 	if (!ensure(IsInGameThread()))
diff --git a/Engine/Source/Runtime/Interchange/Engine/Private/InterchangeSceneImportAsset.cpp b/Engine/Source/Runtime/Interchange/Engine/Private/InterchangeSceneImportAsset.cpp
--- a/Engine/Source/Runtime/Interchange/Engine/Private/InterchangeSceneImportAsset.cpp
+++ b/Engine/Source/Runtime/Interchange/Engine/Private/InterchangeSceneImportAsset.cpp
@@ -205,36 +205,51 @@ UObject* UInterchangeSceneImportAsset::GetSceneObject(const FString& PackageName
 #if WITH_EDITORONLY_DATA
 	const FSoftObjectPath ObjectPath(FName(PackageName), FName(AssetName), SubPathString);
 
-	if (SceneObjects.Find(ObjectPath))
+	UObject* SceneObject = nullptr;
+	for (const TPair<FSoftObjectPath, FString>& SceneObjectPair : SceneObjects)
 	{
-		if (UObject* SceneObject = ObjectPath.TryLoad())
+		const FSoftObjectPath& SceneObjectPath = SceneObjectPair.Key;
+		if (SceneObjectPath.GetLongPackageFName() == ObjectPath.GetLongPackageFName()
+			&& SceneObjectPath.GetAssetFName() == ObjectPath.GetAssetFName())
 		{
-			if (IsValid(SceneObject))
+			//World partition actor have a guid that we need to remove before comparison
+			FName SceneObjectSubPathFName = FActorSpawnUtils::GetBaseName(*SceneObjectPath.GetSubPathString());
+			if (SceneObjectSubPathFName == *ObjectPath.GetSubPathString())
 			{
-				if (SceneObject->IsA<AActor>())
-				{
-					return SceneObject;
-				}
+				SceneObject = SceneObjectPath.TryLoad();
+				break;
+			}
+		}
+	}
+	
+	if (SceneObject)
+	{
+		if (IsValid(SceneObject))
+		{
+			if (SceneObject->IsA<AActor>())
+			{
+				return SceneObject;
+			}
 
-				// Most likely an asset, check whether SceneObject has actually already been imported
-				TArray<UObject*> SubObjects;
-				GetObjectsWithOuter(SceneObject, SubObjects);
-				for (UObject* SubObject : SubObjects)
+			// Most likely an asset, check whether SceneObject has actually already been imported
+			TArray<UObject*> SubObjects;
+			GetObjectsWithOuter(SceneObject, SubObjects);
+			for (UObject* SubObject : SubObjects)
+			{
+				if (SubObject && SubObject->IsA<UInterchangeAssetImportData>())
 				{
-					if (SubObject && SubObject->IsA<UInterchangeAssetImportData>())
-					{
-						return SceneObject;
-					}
+					return SceneObject;
 				}
-
-				return nullptr;
 			}
 
-			// SceneObject is still in memory but invalid. Move it to TransientPackage
-			// Call UObject::Rename because for actors AActor::Rename will unnecessarily unregister and re-register components
-			SceneObject->UObject::Rename(nullptr, GetTransientPackage(), REN_DontCreateRedirectors);
+			return nullptr;
 		}
+
+		// SceneObject is still in memory but invalid. Move it to TransientPackage
+		// Call UObject::Rename because for actors AActor::Rename will unnecessarily unregister and re-register components
+		SceneObject->UObject::Rename(nullptr, GetTransientPackage(), REN_DontCreateRedirectors);
 	}
+
 #endif
 
 	return nullptr;
@@ -250,9 +265,20 @@ const UInterchangeFactoryBaseNode* UInterchangeSceneImportAsset::GetFactoryNode(
 
 	const FSoftObjectPath ObjectPath(FName(PackageName), FName(AssetName), SubPathString);
 
-	if (const FString* UniqueIDPtr = SceneObjects.Find(ObjectPath))
+	UObject* SceneObject = nullptr;
+	for (const TPair<FSoftObjectPath, FString>& SceneObjectPair : SceneObjects)
 	{
-		return Cast<UInterchangeFactoryBaseNode>(AssetImportData->GetStoredNode(*UniqueIDPtr));
+		const FSoftObjectPath& SceneObjectPath = SceneObjectPair.Key;
+		if (SceneObjectPath.GetLongPackageFName() == ObjectPath.GetLongPackageFName()
+			&& SceneObjectPath.GetAssetFName() == ObjectPath.GetAssetFName())
+		{
+			//World partition actor have a guid that we need to remove before comparison
+			FName SceneObjectSubPathFName = FActorSpawnUtils::GetBaseName(*SceneObjectPath.GetSubPathString());
+			if (SceneObjectSubPathFName == *ObjectPath.GetSubPathString())
+			{
+				return Cast<UInterchangeFactoryBaseNode>(AssetImportData->GetStoredNode(SceneObjectPair.Value));
+			}
+		}
 	}
 #endif
 
diff --git a/Engine/Source/Runtime/Interchange/Engine/Private/Tasks/InterchangeTaskCreateSceneObjects.cpp b/Engine/Source/Runtime/Interchange/Engine/Private/Tasks/InterchangeTaskCreateSceneObjects.cpp
--- a/Engine/Source/Runtime/Interchange/Engine/Private/Tasks/InterchangeTaskCreateSceneObjects.cpp
+++ b/Engine/Source/Runtime/Interchange/Engine/Private/Tasks/InterchangeTaskCreateSceneObjects.cpp
@@ -51,12 +51,18 @@ void UE::Interchange::FTaskCreateSceneObjects_GameThread::Execute()
 		return;
 	}
 
+	TArray<TObjectPtr<UObject>> ImportObjects;
+	AsyncHelper->IterateImportedAssets(SourceIndex, [&ImportObjects](const TArray<UE::Interchange::FImportAsyncHelper::FImportedObjectInfo>& ImportedObjectInfos)
+		{
+			for (const UE::Interchange::FImportAsyncHelper::FImportedObjectInfo& Info : ImportedObjectInfos)
+			{
+				ImportObjects.Add(Info.ImportedObject);
+			}
+		});
+
 	UObject* ReimportObject = AsyncHelper->TaskData.ReimportObject;
 	ULevel* ImportLevel = AsyncHelper->TaskData.ImportLevel ? AsyncHelper->TaskData.ImportLevel : GWorld->GetCurrentLevel();
 	UWorld* ImportWorld = ImportLevel->GetWorld();
-	const FString WorldPath = ImportWorld->GetOutermost()->GetPathName();
-	const FString WorldName = ImportWorld->GetName();
-	const FString NodePrefix = ImportLevel->GetName() + TEXT(".");
 
 	for (UInterchangeFactoryBaseNode* FactoryNode : FactoryNodes)
 	{
@@ -73,12 +79,39 @@ void UE::Interchange::FTaskCreateSceneObjects_GameThread::Execute()
 		FString SceneNodeName = FactoryNode->GetAssetName();
 		UInterchangeManager::GetInterchangeManager().SanitizeNameInline(SceneNodeName, ESanitizeNameTypeFlags::ObjectName | ESanitizeNameTypeFlags::ObjectPath | ESanitizeNameTypeFlags::LongPackage);
 
+		ULevel* FactoryImportLevel = ImportLevel;
+		UWorld* FactoryImportWorld = ImportWorld;
+
+		FString FactoryNodeLevelUid;
+		if (FactoryNode->GetCustomLevelUid(FactoryNodeLevelUid))
+		{
+			if (UInterchangeFactoryBaseNode* LevelFactoryNode = AsyncHelper->BaseNodeContainers[SourceIndex]->GetFactoryNode(FactoryNodeLevelUid))
+			{
+				FSoftObjectPath LevelAssetPath;
+				if (LevelFactoryNode->GetCustomReferenceObject(LevelAssetPath))
+				{
+					if (UWorld* FactoryNodeWorld = Cast<UWorld>(LevelAssetPath.TryLoad()))
+					{
+						if (FactoryNodeWorld->PersistentLevel)
+						{
+							FactoryImportWorld = FactoryNodeWorld;
+							FactoryImportLevel = FactoryNodeWorld->PersistentLevel;
+						}
+					}
+				}
+			}
+		}
+		const FString FactoryNodeWorldPath = FactoryImportWorld->GetOutermost()->GetPathName();
+		const FString FactoryNodeWorldName = FactoryImportWorld->GetName();
+		const FString FactoryNodeNodePrefix = FactoryImportLevel->GetName() + TEXT(".");
+
 		UInterchangeFactoryBase::FImportSceneObjectsParams CreateSceneObjectsParams;
 		CreateSceneObjectsParams.ObjectName = SceneNodeName;
 		CreateSceneObjectsParams.FactoryNode = FactoryNode;
-		CreateSceneObjectsParams.Level = ImportLevel;
-		CreateSceneObjectsParams.ReimportObject = FFactoryCommon::GetObjectToReimport(Factory, ReimportObject, *FactoryNode, WorldPath, WorldName, NodePrefix + SceneNodeName);
-		CreateSceneObjectsParams.ReimportFactoryNode = FFactoryCommon::GetFactoryNode(ReimportObject, WorldPath, WorldName, NodePrefix + SceneNodeName);
+		CreateSceneObjectsParams.Level = FactoryImportLevel;
+		CreateSceneObjectsParams.ImportAssets = ImportObjects;
+		CreateSceneObjectsParams.ReimportObject = FFactoryCommon::GetObjectToReimport(Factory, ReimportObject, *FactoryNode, FactoryNodeWorldPath, FactoryNodeWorldName, FactoryNodeNodePrefix + SceneNodeName);
+		CreateSceneObjectsParams.ReimportFactoryNode = FFactoryCommon::GetFactoryNode(ReimportObject, FactoryNodeWorldPath, FactoryNodeWorldName, FactoryNodeNodePrefix + SceneNodeName);
 
 		if (AsyncHelper->BaseNodeContainers.IsValidIndex(SourceIndex))
 		{
diff --git a/Engine/Source/Runtime/Interchange/Engine/Private/Tasks/InterchangeTaskImportObject.cpp b/Engine/Source/Runtime/Interchange/Engine/Private/Tasks/InterchangeTaskImportObject.cpp
--- a/Engine/Source/Runtime/Interchange/Engine/Private/Tasks/InterchangeTaskImportObject.cpp
+++ b/Engine/Source/Runtime/Interchange/Engine/Private/Tasks/InterchangeTaskImportObject.cpp
@@ -4,6 +4,7 @@
 #include "AssetCompilingManager.h"
 #include "Async/TaskGraphInterfaces.h"
 #include "CoreMinimal.h"
+#include "Engine/World.h"
 #include "InterchangeAssetImportData.h"
 #include "InterchangeEngineLogPrivate.h"
 #include "InterchangeFactoryBase.h"
@@ -330,11 +331,25 @@ void UE::Interchange::FTaskImportObject_GameThread::Execute()
 			}
 		};
 
+	//Check if the node reference a map
+	const bool bRefObjectIsMap = [this]()
+		{
+			FSoftObjectPath Reference;
+			if (FactoryNode->GetCustomReferenceObject(Reference))
+			{
+				if (UObject* RefObj = Reference.TryLoad())
+				{
+					return RefObj->IsA<UWorld>();
+				}
+			}
+			return false;
+		}();
+
 	//If the factory node is disable see if there is an existing UObject for it
 	if (!FactoryNode->IsEnabled())
 	{
 		Private::InternalGetPackageName(*AsyncHelper, SourceIndex, PackageBasePath, FactoryNode, PackageName, AssetName);
-		if (!FPackageUtils::IsMapPackageAsset(PackageName))
+		if (bRefObjectIsMap || !FPackageUtils::IsMapPackageAsset(PackageName))
 		{
 			Pkg = FindPackage(nullptr, *PackageName);
 			if (Pkg)
@@ -390,7 +405,8 @@ void UE::Interchange::FTaskImportObject_GameThread::Execute()
 
 		Private::InternalGetPackageName(*AsyncHelper, SourceIndex, PackageBasePath, FactoryNode, PackageName, AssetName);
 		// We can not create assets that share the name of a map file in the same location
-		if (FPackageUtils::IsMapPackageAsset(PackageName))
+		//Except if we reference a map
+		if (!bRefObjectIsMap && FPackageUtils::IsMapPackageAsset(PackageName))
 		{
 			//Skip this asset
 			UInterchangeResultError_Generic* Message = Factory->AddMessage<UInterchangeResultError_Generic>();
diff --git a/Engine/Source/Runtime/Interchange/Engine/Public/InterchangeEditorUtilitiesBase.h b/Engine/Source/Runtime/Interchange/Engine/Public/InterchangeEditorUtilitiesBase.h
new file mode 100644
--- /dev/null
+++ b/Engine/Source/Runtime/Interchange/Engine/Public/InterchangeEditorUtilitiesBase.h
@@ -0,0 +1,20 @@
+// Copyright Epic Games, Inc. All Rights Reserved.
+
+#pragma once
+
+#include "UObject/Object.h"
+#include "UObject/ObjectMacros.h"
+
+#include "InterchangeEditorUtilitiesBase.generated.h"
+
+UCLASS(MinimalAPI)
+class UInterchangeEditorUtilitiesBase : public UObject
+{
+	GENERATED_BODY()
+
+public:
+	INTERCHANGEENGINE_API virtual bool SaveAsset(UObject* Asset)
+	{
+		return false;
+	}
+};
diff --git a/Engine/Source/Runtime/Interchange/Engine/Public/InterchangeManager.h b/Engine/Source/Runtime/Interchange/Engine/Public/InterchangeManager.h
--- a/Engine/Source/Runtime/Interchange/Engine/Public/InterchangeManager.h
+++ b/Engine/Source/Runtime/Interchange/Engine/Public/InterchangeManager.h
@@ -32,6 +32,7 @@
 
 class FAsyncTaskNotification;
 class UInterchangeBlueprintPipelineBase;
+class UInterchangeEditorUtilitiesBase;
 class UInterchangeFactoryBaseNode;
 class UInterchangePipelineBase;
 class UInterchangePythonPipelineBase;
@@ -779,6 +780,16 @@ public:
 	*/
 	INTERCHANGEENGINE_API static TOptional<bool> GetReplaceExistingAlldialogAnswer();
 
+	/**
+	* Set the editor utilities, those are use for editor operation like saving an asset.
+	*/
+	INTERCHANGEENGINE_API void SetEditorUtilities(UClass* EditorUtilitiesClass);
+	
+	/**
+	* Get the editor utilities, those are use for editor operation like saving an asset.
+	*/
+	INTERCHANGEENGINE_API UInterchangeEditorUtilitiesBase* GetEditorUtilities() const;
+
 protected:
 
 	/** Return true if Interchange can show UI. */
@@ -878,6 +889,9 @@ private:
 	UPROPERTY()
 	TMap<TObjectPtr<const UClass>, TObjectPtr<UInterchangeAssetImportDataConverterBase> > RegisteredConverters;
 
+	//We support one editor utilities class
+	TStrongObjectPtr<UInterchangeEditorUtilitiesBase> EditorUtilities = nullptr;
+
 	//If interchange is currently importing, we have a timer to watch the cancel and we block GC.
 	FThreadSafeBool bIsActive = false;
 
diff --git a/Engine/Source/Runtime/Interchange/Engine/Public/InterchangeSceneImportAsset.h b/Engine/Source/Runtime/Interchange/Engine/Public/InterchangeSceneImportAsset.h
--- a/Engine/Source/Runtime/Interchange/Engine/Public/InterchangeSceneImportAsset.h
+++ b/Engine/Source/Runtime/Interchange/Engine/Public/InterchangeSceneImportAsset.h
@@ -23,7 +23,7 @@ class UInterchangeSceneImportAsset : public UObject, public IInterface_AssetUser
 
 public:
 #if WITH_EDITORONLY_DATA
-	/** Importing data and options used for this Datasmith scene */
+	/** Importing data and options used for this scene */
 	UPROPERTY(EditAnywhere, Instanced, Category=ImportSettings)
 	TObjectPtr<UInterchangeAssetImportData> AssetImportData;
 
