Interchange Scene: Cache the active joint unique ids list to optimize the actor creation

#rnx

diff --git a/Engine/Plugins/Interchange/Runtime/Source/Pipelines/Private/InterchangeGenericScenesPipeline.cpp b/Engine/Plugins/Interchange/Runtime/Source/Pipelines/Private/InterchangeGenericScenesPipeline.cpp
--- a/Engine/Plugins/Interchange/Runtime/Source/Pipelines/Private/InterchangeGenericScenesPipeline.cpp
+++ b/Engine/Plugins/Interchange/Runtime/Source/Pipelines/Private/InterchangeGenericScenesPipeline.cpp
@@ -57,32 +57,6 @@ if(AttributeType AttributeName; TranslatedNode->GetCustom##AttributeName(Attribu
 
 namespace UE::Interchange::Private
 {
-	TArray<FString> GetAllActiveJoints(UInterchangeBaseNodeContainer* BaseNodeContainer)
-	{
-		TArray<FString> AllActiveJoints;
-		BaseNodeContainer->IterateNodesOfType<UInterchangeSkeletonFactoryNode>([&BaseNodeContainer, &AllActiveJoints](const FString& NodeUid, UInterchangeSkeletonFactoryNode* Node)
-			{
-				FString RootNodeUid;
-				if (Node->GetCustomRootJointUid(RootNodeUid))
-				{
-					AllActiveJoints.Add(RootNodeUid);
-					BaseNodeContainer->IterateNodeChildren(RootNodeUid, [&AllActiveJoints](const UInterchangeBaseNode* Node)
-						{
-							if (const UInterchangeSceneNode* SceneNode = Cast<UInterchangeSceneNode>(Node))
-							{
-								TArray<FString> SpecializeTypes;
-								SceneNode->GetSpecializedTypes(SpecializeTypes);
-								if (SpecializeTypes.Contains(UE::Interchange::FSceneNodeStaticData::GetJointSpecializeTypeString()))
-								{
-									AllActiveJoints.Add(Node->GetUniqueID());
-								}
-							}
-						});
-				}
-			});
-		return AllActiveJoints;
-	}
-
 	//Either a (TransformSpecialized || !JointSpecialized || RootJoint) can be a parent (only those get FactoryNodes) :
 	FString FindFactoryParentSceneNodeUid(UInterchangeBaseNodeContainer* BaseNodeContainer, TArray<FString>& ActiveSkeletonUids, const UInterchangeSceneNode* SceneNode)
 	{
@@ -382,7 +356,7 @@ void UInterchangeGenericLevelPipeline::ExecutePipeline(UInterchangeBaseNodeConta
 #endif
 
 	/* Find all scene node that are active joint. Non active joint should be convert to actor if they are in a static mesh hierarchy */
-	TArray<FString> ActiveSkeletonUids = UE::Interchange::Private::GetAllActiveJoints(BaseNodeContainer);
+	CacheActiveJointUids();
 
 	for (const UInterchangeSceneNode* SceneNode : SceneNodes)
 	{
@@ -397,7 +371,7 @@ void UInterchangeGenericLevelPipeline::ExecutePipeline(UInterchangeBaseNodeConta
 					bool bSkipNode = true;
 					if (SpecializeTypes.Contains(UE::Interchange::FSceneNodeStaticData::GetJointSpecializeTypeString()))
 					{
-						if(!ActiveSkeletonUids.Contains(SceneNode->GetUniqueID()))
+						if(!Cached_ActiveJointUids.Contains(SceneNode->GetUniqueID()))
 						{ 
 							bSkipNode = false;
 						}
@@ -556,8 +530,7 @@ void UInterchangeGenericLevelPipeline::ExecuteSceneNodePreImport(const FTransfor
 	if (!SceneNode->GetParentUid().IsEmpty())
 	{
 		/* Find all scene node that are active joint. Non active joint should be convert to actor if they are in a static mesh hierarchy */
-		TArray<FString> ActiveSkeletonUids = UE::Interchange::Private::GetAllActiveJoints(BaseNodeContainer);
-		FString ParentNodeUid = UE::Interchange::Private::FindFactoryParentSceneNodeUid(BaseNodeContainer, ActiveSkeletonUids, SceneNode);
+		FString ParentNodeUid = UE::Interchange::Private::FindFactoryParentSceneNodeUid(BaseNodeContainer, Cached_ActiveJointUids, SceneNode);
 		if (ParentNodeUid != UInterchangeBaseNode::InvalidNodeUid())
 		{
 			FString ParentFactoryNodeUid = UInterchangeFactoryBaseNode::BuildFactoryNodeUid(ParentNodeUid);
@@ -1251,4 +1224,29 @@ bool UInterchangeGenericLevelPipeline::FPostPipelineImportData::UpdateLevelInsta
 	return true;
 }
 
-#endif //WITH_EDITORONLY_DATA
\ No newline at end of file
+#endif //WITH_EDITORONLY_DATA
+
+void UInterchangeGenericLevelPipeline::CacheActiveJointUids()
+{
+	Cached_ActiveJointUids.Reset();
+	BaseNodeContainer->IterateNodesOfType<UInterchangeSkeletonFactoryNode>([this](const FString& NodeUid, UInterchangeSkeletonFactoryNode* Node)
+		{
+			FString RootNodeUid;
+			if (Node->GetCustomRootJointUid(RootNodeUid))
+			{
+				Cached_ActiveJointUids.Add(RootNodeUid);
+				BaseNodeContainer->IterateNodeChildren(RootNodeUid, [this](const UInterchangeBaseNode* Node)
+					{
+						if (const UInterchangeSceneNode* SceneNode = Cast<UInterchangeSceneNode>(Node))
+						{
+							TArray<FString> SpecializeTypes;
+							SceneNode->GetSpecializedTypes(SpecializeTypes);
+							if (SpecializeTypes.Contains(UE::Interchange::FSceneNodeStaticData::GetJointSpecializeTypeString()))
+							{
+								Cached_ActiveJointUids.Add(Node->GetUniqueID());
+							}
+						}
+					});
+			}
+		});
+}
\ No newline at end of file
diff --git a/Engine/Plugins/Interchange/Runtime/Source/Pipelines/Public/InterchangeGenericScenesPipeline.h b/Engine/Plugins/Interchange/Runtime/Source/Pipelines/Public/InterchangeGenericScenesPipeline.h
--- a/Engine/Plugins/Interchange/Runtime/Source/Pipelines/Public/InterchangeGenericScenesPipeline.h
+++ b/Engine/Plugins/Interchange/Runtime/Source/Pipelines/Public/InterchangeGenericScenesPipeline.h
@@ -126,6 +126,9 @@ protected:
 
 	FPostPipelineImportData PostPipelineImportData;
 #endif
+
+	void CacheActiveJointUids();
+	TArray<FString> Cached_ActiveJointUids;
 };
 
 
