Fix reference for level instance world

#rnx

diff --git a/Engine/Plugins/Interchange/Runtime/Source/Pipelines/Private/InterchangeGenericScenesPipeline.cpp b/Engine/Plugins/Interchange/Runtime/Source/Pipelines/Private/InterchangeGenericScenesPipeline.cpp
--- a/Engine/Plugins/Interchange/Runtime/Source/Pipelines/Private/InterchangeGenericScenesPipeline.cpp
+++ b/Engine/Plugins/Interchange/Runtime/Source/Pipelines/Private/InterchangeGenericScenesPipeline.cpp
@@ -1186,9 +1186,9 @@ void UInterchangeGenericLevelPipeline::FPostPipelineImportData::AddLevelInstance
 	}
 }
 
-bool UInterchangeGenericLevelPipeline::FPostPipelineImportData::UpdateLevelInstanceInternal(ALevelInstance* LevelInstanceActor, UWorld* ReferenceWorld)
+bool UInterchangeGenericLevelPipeline::FPostPipelineImportData::UpdateLevelInstanceInternal(ALevelInstance* LevelInstanceActor, UWorld* ReferencedWorld)
 {
-	if (!Worlds.Contains(ReferenceWorld))
+	if (!Worlds.Contains(ReferencedWorld))
 	{
 		return false;
 	}
@@ -1198,15 +1198,15 @@ bool UInterchangeGenericLevelPipeline::FPostPipelineImportData::UpdateLevelInsta
 		return false;
 	}
 
-	bool& bWorldSaved = Worlds.FindChecked(ReferenceWorld);
+	bool& bWorldSaved = Worlds.FindChecked(ReferencedWorld);
 	
 	if (!bWorldSaved)
 	{
 		if (UInterchangeEditorUtilitiesBase* EditorUtilities = UInterchangeManager::GetInterchangeManager().GetEditorUtilities())
 		{
-			if (!EditorUtilities->SaveAsset(ReferenceWorld))
+			if (!EditorUtilities->SaveAsset(ReferencedWorld))
 			{
-				UE_LOG(LogInterchangePipeline, Warning, TEXT("UInterchangeGenericAssetsPipeline: Cannot save the level instance actor (%s) referenced world (%s)"), *LevelInstanceActor->GetName(), *ReferenceWorld->GetName());
+				UE_LOG(LogInterchangePipeline, Warning, TEXT("UInterchangeGenericAssetsPipeline: Cannot save the level instance actor (%s) referenced world (%s)"), *LevelInstanceActor->GetName(), *ReferencedWorld->GetName());
 			}
 		}
 		bWorldSaved = true;
@@ -1214,10 +1214,19 @@ bool UInterchangeGenericLevelPipeline::FPostPipelineImportData::UpdateLevelInsta
 
 	ParentWorld->PreEditChange(nullptr);
 
-	LevelInstanceActor->SetWorldAsset(ReferenceWorld);
+	LevelInstanceActor->SetWorldAsset(ReferencedWorld);
 	LevelInstanceActor->UpdateLevelInstanceFromWorldAsset();
 	LevelInstanceActor->LoadLevelInstance();
 	
+	//A level instance referenced world cannot be RF_Standalone
+	ReferencedWorld->ClearFlags(RF_Standalone);
+
+	//Reference world must be cleanup since they are not the main world.
+	//This remove all the world managers and prevent GC issue when unloading the main world referencing this world.
+	if (ReferencedWorld->bIsWorldInitialized)
+	{
+		ReferencedWorld->CleanupWorld();
+	}
 
 	ParentWorld->PostEditChange();
 
