Fix infinite loop callstack crash
#rnx

diff --git a/Engine/Plugins/Interchange/Runtime/Source/Pipelines/Private/InterchangeGenericScenesPipeline.cpp b/Engine/Plugins/Interchange/Runtime/Source/Pipelines/Private/InterchangeGenericScenesPipeline.cpp
--- a/Engine/Plugins/Interchange/Runtime/Source/Pipelines/Private/InterchangeGenericScenesPipeline.cpp
+++ b/Engine/Plugins/Interchange/Runtime/Source/Pipelines/Private/InterchangeGenericScenesPipeline.cpp
@@ -267,13 +267,6 @@ void UInterchangeGenericLevelPipeline::ExecutePipeline(UInterchangeBaseNodeConta
 		ParentLevelFactoryNode->InitializeNode(ParentNodeUid, ParentDisplayLabel, EInterchangeNodeContainerType::FactoryData);
 		ParentLevelFactoryNode->SetCustomShouldCreateLevel(false);
 		ParentLevelFactoryNode->SetCustomReferenceObject(TargetWorld);
-		BaseNodeContainer->IterateNodesOfType<UInterchangeFactoryBaseNode>(
-			[&ParentLevelFactoryNode](const FString& NodeUid, UInterchangeFactoryBaseNode* Node)
-			{
-				// Add all assets (not actors) before calling the ParentLevelFactoryNode factory
-				ParentLevelFactoryNode->AddFactoryDependencyUid(NodeUid);
-			}
-		);
 		BaseNodeContainer->AddNode(ParentLevelFactoryNode);
 
 		UWorld* ReimportReferenceWorld = nullptr;
