Interchange Telemetry
- All the InterchangeResults in the Interchange now use localized texts.

diff --git a/Engine/Plugins/Interchange/Runtime/Source/Import/Private/Gltf/InterchangeGltfTranslator.cpp b/Engine/Plugins/Interchange/Runtime/Source/Import/Private/Gltf/InterchangeGltfTranslator.cpp
--- a/Engine/Plugins/Interchange/Runtime/Source/Import/Private/Gltf/InterchangeGltfTranslator.cpp
+++ b/Engine/Plugins/Interchange/Runtime/Source/Import/Private/Gltf/InterchangeGltfTranslator.cpp
@@ -47,7 +47,7 @@
 
 #include UE_INLINE_GENERATED_CPP_BY_NAME(InterchangeGltfTranslator)
 
-#define LOCTEXT_NAMESPACE "InterchangeGltfTranslator"
+#define LOCTEXT_NAMESPACE "InterchangeGLTFTranslator"
 
 static const TArray<FString> ImporterSupportedExtensions = {
 	/* Lights */
@@ -149,7 +149,7 @@ namespace UE::Interchange::Gltf::Private
 	};
 	void SendAnalytics(const TranslationResult& TranslationResult,
 		const GLTF::FAsset& Asset = GLTF::FAsset(),
-		const FString& GLTFReaderLogMessage = "")
+		const FText& GLTFReaderLogMessage = FText::GetEmpty())
 	{
 		if (FEngineAnalytics::IsAvailable())
 		{
@@ -461,7 +461,7 @@ bool UInterchangeGLTFTranslator::Translate( UInterchangeBaseNodeContainer& NodeC
 		{
 			UInterchangeResultError_Generic* ErrorResult = AddMessage< UInterchangeResultError_Generic >();
 			ErrorResult->SourceAssetName = FileName;
-			ErrorResult->Text = FText::Format(LOCTEXT("GLTF::FFileReader::ReadFile Failed.", "LogMessage: {0}"), FText::FromString(LogMessage.Value));
+			ErrorResult->Text = LogMessage.Value;
 
 			SendAnalytics(TranslationResult::GLTFREADER_FAILED, GltfAsset, LogMessage.Value);
 			return false;
@@ -807,21 +807,21 @@ bool UInterchangeGLTFTranslator::Translate( UInterchangeBaseNodeContainer& NodeC
 		case GLTF::EMessageSeverity::Error :
 			{
 				UInterchangeResultError_Generic* ErrorResult = AddMessage< UInterchangeResultError_Generic >();
-				ErrorResult->Text = FText::FromString(LogMessage.Get<1>());
+				ErrorResult->Text = LogMessage.Get<1>();
 				Result = ErrorResult;
 			}
 			break;
 		case GLTF::EMessageSeverity::Warning:
 			{
 				UInterchangeResultWarning_Generic* WarningResult = AddMessage< UInterchangeResultWarning_Generic >();
-				WarningResult->Text = FText::FromString(LogMessage.Get<1>());
+				WarningResult->Text = LogMessage.Get<1>();
 				Result = WarningResult;
 			}
 			break;
 		case GLTF::EMessageSeverity::Display:
 			{
 				UInterchangeResultDisplay_Generic* DisplayResult = AddMessage< UInterchangeResultDisplay_Generic >();
-				DisplayResult->Text = FText::FromString(LogMessage.Get<1>());
+				DisplayResult->Text = LogMessage.Get<1>();
 				Result = DisplayResult;
 			}
 			break;
@@ -1239,7 +1239,7 @@ TFuture< TOptional< UE::Interchange::FMeshPayloadData > > UInterchangeGLTFTransl
 				{
 					UInterchangeResultError_Generic* ErrorResult = AddMessage<UInterchangeResultError_Generic>();
 					ErrorResult->SourceAssetName = SourceData ? SourceData->GetFilename() : FString();
-					ErrorResult->Text = NSLOCTEXT("UInterchangeGLTFTranslator", "GetMeshPayloadData_ValidateMeshDescriptionFail", "Invalid mesh data (NAN) was found and changed to zero. This may affect the mesh rendering.");
+					ErrorResult->Text = LOCTEXT("GetMeshPayloadData_ValidateMeshDescriptionFail", "Invalid mesh data (NAN) was found and changed to zero. This may affect the mesh rendering.");
 				}
 
 				Result.Emplace(MeshPayLoadData);
diff --git a/Engine/Plugins/Interchange/Runtime/Source/Import/Private/Mesh/InterchangeOBJTranslator.cpp b/Engine/Plugins/Interchange/Runtime/Source/Import/Private/Mesh/InterchangeOBJTranslator.cpp
--- a/Engine/Plugins/Interchange/Runtime/Source/Import/Private/Mesh/InterchangeOBJTranslator.cpp
+++ b/Engine/Plugins/Interchange/Runtime/Source/Import/Private/Mesh/InterchangeOBJTranslator.cpp
@@ -30,6 +30,8 @@
 
 #include UE_INLINE_GENERATED_CPP_BY_NAME(InterchangeOBJTranslator)
 
+#define LOCTEXT_NAMESPACE "InterchangeOBJTranslator"
+
 static bool GInterchangeEnableOBJImport = true;
 static FAutoConsoleVariableRef CCvarInterchangeEnableOBJImport(
 	TEXT("Interchange.FeatureFlags.Import.OBJ"),
@@ -1313,7 +1315,7 @@ bool UInterchangeOBJTranslator::Translate(UInterchangeBaseNodeContainer& BaseNod
 		{
 			UInterchangeResultError_Generic* ErrorResult = AddMessage<UInterchangeResultError_Generic>();
 			ErrorResult->SourceAssetName = FPaths::GetBaseFilename(Filename);
-			ErrorResult->Text = NSLOCTEXT("InterchangeOBJTranslator", "EmptyFileError", "The OBJ file appears to be empty.");
+			ErrorResult->Text = LOCTEXT("EmptyFileError", "The OBJ file appears to be empty.");
 
 			return false;
 		}
@@ -1592,7 +1594,7 @@ TFuture<TOptional<UE::Interchange::FMeshPayloadData>> UInterchangeOBJTranslator:
 			{
 				UInterchangeResultError_Generic* ErrorResult = AddMessage<UInterchangeResultError_Generic>();
 				ErrorResult->SourceAssetName = SourceData ? SourceData->GetFilename() : FString();
-				ErrorResult->Text = NSLOCTEXT("UInterchangeOBJTranslator", "GetMeshPayloadData_ValidateMeshDescriptionFail", "Invalid mesh data (NAN) was found and fix to zero. Mesh render can be bad.");
+				ErrorResult->Text = LOCTEXT("GetMeshPayloadData_ValidateMeshDescriptionFail", "Invalid mesh data (NAN) was found and fix to zero. Mesh render can be bad.");
 			}
 
 			return TOptional<FMeshPayloadData>(Payload);
@@ -1619,4 +1621,4 @@ TOptional<UE::Interchange::FImportImage> UInterchangeOBJTranslator::GetTexturePa
 	return TextureTranslator->GetTexturePayloadData(PayLoadKey, AlternateTexturePath);
 }
 
-
+#undef LOCTEXT_NAMESPACE
\ No newline at end of file
diff --git a/Engine/Plugins/Interchange/Runtime/Source/Parsers/GLTFCore/Private/GLTF/ExtensionsHandler.cpp b/Engine/Plugins/Interchange/Runtime/Source/Parsers/GLTFCore/Private/GLTF/ExtensionsHandler.cpp
--- a/Engine/Plugins/Interchange/Runtime/Source/Parsers/GLTFCore/Private/GLTF/ExtensionsHandler.cpp
+++ b/Engine/Plugins/Interchange/Runtime/Source/Parsers/GLTFCore/Private/GLTF/ExtensionsHandler.cpp
@@ -13,6 +13,9 @@
 #include "draco/compression/decode.h"
 #endif
 
+
+#define LOCTEXT_NAMESPACE "InterchangeGLTFExtensionHandler"
+
 namespace GLTF
 {
 #if USE_DRACO_LIBRARY
@@ -400,7 +403,7 @@ namespace GLTF
 				default:
 					if (!ensure(false))
 					{
-						Messages.Emplace(RuntimeWarningSeverity(), FString::Printf(TEXT("Material.Extension not supported: %s"), ToString(Extension)));
+						Messages.Emplace(RuntimeWarningSeverity(), FText::Format(LOCTEXT("UnsupportedMaterialExtension", "Material.Extension not supported: {0}"), FText::FromString(ToString(Extension))));
 					}
 					break;
 			}
@@ -513,7 +516,7 @@ namespace GLTF
 									if (!DracoHelpers::AcquireIndicesFromDracoMesh(Mesh.get(), Accessor.ComponentType, Accessor.BufferView.Buffer.Data))
 									{
 										Accessor.BufferView = FBufferView();
-										Messages.Emplace(EMessageSeverity::Warning, FString::Printf(TEXT("Failed to acquire Indices from Draco Mesh, for PrimitiveIdx: %s, in Mesh: %s"), *FString::FromInt(PrimitiveIndex), *MeshUniqueId));
+										Messages.Emplace(EMessageSeverity::Warning, FText::Format(LOCTEXT("DracoMeshIndexAcquisitionFailed", "Failed to acquire Indices from Draco Mesh, for PrimitiveIdx: {0}, in Mesh: {1}"), PrimitiveIndex, FText::FromString(MeshUniqueId)));
 									}
 								}
 							}
@@ -540,7 +543,7 @@ namespace GLTF
 												{
 													//Clear out the BufferView so it is not used:
 													Accessor.BufferView = FBufferView();
-													Messages.Emplace(EMessageSeverity::Warning, FString::Printf(TEXT("Failed to acquire %s Attributes from Draco Mesh, for PrimitiveIdx: %s, in Mesh: %s"), *ToString(MeshAttributeType), *FString::FromInt(PrimitiveIndex), *MeshUniqueId));
+													Messages.Emplace(EMessageSeverity::Warning, FText::Format(LOCTEXT("DracoMeshAttributeAcquisitionFailed", "Failed to acquire {0} Attributes from Draco Mesh, for PrimitiveIdx: {1}, in Mesh: {2}"), FText::FromString(ToString(MeshAttributeType)), PrimitiveIndex, FText::FromString(MeshUniqueId)));
 												}
 											}
 										}
@@ -561,7 +564,7 @@ namespace GLTF
 				default:
 					if (!ensure(false))
 					{
-						Messages.Emplace(RuntimeWarningSeverity(), FString::Printf(TEXT("Primitive.Extension not supported: %s"), ToCStr(ToString(Extension))));
+						Messages.Emplace(RuntimeWarningSeverity(), FText::Format(LOCTEXT("UnsupportedPrimitiveExtension", "Primitive.Extension not supported: {0}"), FText::FromString(ToString(Extension))));
 					}
 				break;
 			}
@@ -642,7 +645,7 @@ namespace GLTF
 		{
 			if (ExtensionsSupported.Find(StrValuePair.Get<0>()) == INDEX_NONE)
 			{
-				Messages.Emplace(RuntimeWarningSeverity(), FString::Printf(TEXT("Extension is not supported: %s"), *StrValuePair.Get<0>()));
+				Messages.Emplace(RuntimeWarningSeverity(), FText::Format(LOCTEXT("UnsupportedExtension", "Extension is not supported: {0}"), FText::FromString(StrValuePair.Get<0>())));
 			}
 		}
 	}
@@ -675,7 +678,9 @@ namespace GLTF
 		else if (Type == TEXT("directional"))
 			Light.Type = FLight::EType::Directional;
 		else
-			Messages.Emplace(RuntimeWarningSeverity(), FString::Printf(TEXT("Light has no type specified: %s"), *Light.Name));
+			Messages.Emplace(RuntimeWarningSeverity(), FText::Format(LOCTEXT("UnspecifiedLightType", "Light has no type specified: %s"), FText::FromString(Light.Name)));
 	}
 
 }  // namespace GLTF
+
+#undef LOCTEXT_NAMESPACE
\ No newline at end of file
diff --git a/Engine/Plugins/Interchange/Runtime/Source/Parsers/GLTFCore/Private/GLTF/GLTFBinaryReader.cpp b/Engine/Plugins/Interchange/Runtime/Source/Parsers/GLTFCore/Private/GLTF/GLTFBinaryReader.cpp
--- a/Engine/Plugins/Interchange/Runtime/Source/Parsers/GLTFCore/Private/GLTF/GLTFBinaryReader.cpp
+++ b/Engine/Plugins/Interchange/Runtime/Source/Parsers/GLTFCore/Private/GLTF/GLTFBinaryReader.cpp
@@ -3,6 +3,7 @@
 #include "GLTFBinaryReader.h"
 
 #include "CoreGlobals.h"
+#include "Internationalization/Internationalization.h"
 
 namespace GLTF
 {
@@ -155,7 +156,7 @@ namespace GLTF
 
 		if (!IsHeaderValid(FileReader))
 		{
-			Messages.Emplace(EMessageSeverity::Error, TEXT("Invalid GLTF header!"));
+			Messages.Emplace(EMessageSeverity::Error, NSLOCTEXT("InterchangeGLTFBinaryReader", "InvalidHeader", "Invalid GLTF header!"));
 			return false;
 		}
 
diff --git a/Engine/Plugins/Interchange/Runtime/Source/Parsers/GLTFCore/Private/GLTF/GLTFReader.cpp b/Engine/Plugins/Interchange/Runtime/Source/Parsers/GLTFCore/Private/GLTF/GLTFReader.cpp
--- a/Engine/Plugins/Interchange/Runtime/Source/Parsers/GLTFCore/Private/GLTF/GLTFReader.cpp
+++ b/Engine/Plugins/Interchange/Runtime/Source/Parsers/GLTFCore/Private/GLTF/GLTFReader.cpp
@@ -21,6 +21,9 @@
 
 #include "Math/UnrealMathVectorConstants.h"
 
+
+#define LOCTEXT_NAMESPACE "InterchangeGLTFReader"
+
 namespace GLTF
 {
 	namespace
@@ -211,7 +214,7 @@ namespace GLTF
 				bool    bSuccess = DecodeDataURI(URI, MimeType, CurrentBufferOffset, DataSize);
 				if (!bSuccess || (MimeType != TEXT("application/octet-stream") && MimeType != TEXT("application/gltf-buffer")) || !ensure(DataSize == ByteLength))
 				{
-					Messages.Emplace(EMessageSeverity::Error, TEXT("Problem decoding buffer from data URI."));
+					Messages.Emplace(EMessageSeverity::Error, LOCTEXT("BufferDecodingProblem", "Problem decoding buffer from data URI."));
 				}
 				else
 				{
@@ -235,7 +238,7 @@ namespace GLTF
 					}
 					else
 					{
-						Messages.Emplace(EMessageSeverity::Error, TEXT("Buffer file size does not match."));
+						Messages.Emplace(EMessageSeverity::Error, LOCTEXT("NonMatchingBufferFileSize", "Buffer file size does not match."));
 					}
 
 					Reader->Close();
@@ -243,7 +246,7 @@ namespace GLTF
 				}
 				else
 				{
-					Messages.Emplace(EMessageSeverity::Error, FString::Printf(TEXT("Could not load file: '%s'"), *FullPath));
+					Messages.Emplace(EMessageSeverity::Error, FText::Format(LOCTEXT("LoadingFileFailed", "Could not load file: '{0}'"), FText::FromString(FullPath)));
 				}
 			}
 		}
@@ -253,11 +256,11 @@ namespace GLTF
 			const uint32 BinSize = Asset->BinData.Num();
 			if (BinSize == 0)
 			{
-				Messages.Emplace(EMessageSeverity::Error, TEXT("Buffer from BIN chunk is missing or empty."));
+				Messages.Emplace(EMessageSeverity::Error, LOCTEXT("BINChunkMissing","Buffer from BIN chunk is missing or empty."));
 			}
 			else if (BinSize < ByteLength)
 			{
-				Messages.Emplace(EMessageSeverity::Error, TEXT("Buffer from BIN chunk is too small."));
+				Messages.Emplace(EMessageSeverity::Error, LOCTEXT("BINChunkTooSmall","Buffer from BIN chunk is too small."));
 			}
 			else
 			{
@@ -410,7 +413,11 @@ namespace GLTF
 
 		if (!FPrimitive::SupportedModes.Contains(Mode))
 		{
-			Messages.Emplace(EMessageSeverity::Warning, FString::Printf(TEXT("Primitive Mode[%s] in Primitive[%i] (in Mesh[%s]) is currently not supported. Geometry won't be imported."), *FPrimitive::ToString(Mode), PrimitiveIndex, *Mesh.Name));
+			Messages.Emplace(EMessageSeverity::Warning, FText::Format(LOCTEXT("UnsupportedPrimitiveMode", 
+				"Primitive Mode[{0}] in Primitive[{1}] (in Mesh[{2}]) is currently not supported. Geometry won't be imported."), 
+				FText::FromString(FPrimitive::ToString(Mode)), 
+				PrimitiveIndex, 
+				FText::FromString(Mesh.Name)));
 		}
 
 		const int32                   MaterialIndex = GetIndex(Object, TEXT("material"));
@@ -479,7 +486,7 @@ namespace GLTF
 				if (NumberOfMorphTargets != Mesh.Primitives.Last().MorphTargets.Num())
 				{
 					//All primitives MUST have the same number of morph targets in the same order.
-					Messages.Emplace(EMessageSeverity::Error, TEXT("Number of Primitive.Targets is not consistent across the Mesh."));
+					Messages.Emplace(EMessageSeverity::Error, LOCTEXT("InconsistentNumMorphTargets", "Number of Primitive.Targets is not consistent across the Mesh."));
 				}
 			}
 
@@ -603,7 +610,7 @@ namespace GLTF
 		FString Name = GetString(Object, TEXT("name"));
 		if (!Found)
 		{
-			Messages.Emplace(EMessageSeverity::Warning, FString::Printf(TEXT("No camera node found for camera %d('%s')"), CameraIndex, *Name));
+			Messages.Emplace(EMessageSeverity::Warning, FText::Format(LOCTEXT("NoCameraNodeFound", "No camera node found for camera {0}('{1}')"), CameraIndex, FText::FromString(Name)));
 			return;
 		}
 
@@ -633,7 +640,7 @@ namespace GLTF
 		}
 		else
 		{
-			Messages.Emplace(EMessageSeverity::Error, TEXT("Invalid camera type: ") + Type);
+			Messages.Emplace(EMessageSeverity::Error, FText::Format(LOCTEXT("InvalidCameraType", "Invalid camera type: {0}"), FText::FromString(Type)));
 		}
 
 		ProcessExtras(Object, Camera.Extras);
@@ -751,7 +758,7 @@ namespace GLTF
 				Image.Format     = ImageFormatFromMimeType(MimeType);
 				if (!bSuccess || Image.Format == FImage::EFormat::Unknown)
 				{
-					Messages.Emplace(EMessageSeverity::Error, TEXT("Problem decoding image from data URI."));
+					Messages.Emplace(EMessageSeverity::Error, LOCTEXT("DecodingImageProblem", "Problem decoding image from data URI."));
 				}
 				else
 				{
@@ -780,7 +787,7 @@ namespace GLTF
 					}
 					else
 					{
-						Messages.Emplace(EMessageSeverity::Error, TEXT("Could not load image file."));
+						Messages.Emplace(EMessageSeverity::Error, LOCTEXT("ErrorLoadingImageFile", "Could not load image file."));
 					}
 				}
 			}
@@ -850,7 +857,7 @@ namespace GLTF
 		}
 		else
 		{
-			Messages.Emplace(EMessageSeverity::Warning, TEXT("Invalid texture source index: ") + FString::FromInt(SourceIndex));
+			Messages.Emplace(EMessageSeverity::Warning, FText::Format(LOCTEXT("InvalidTextureSourceIndex", "Invalid texture source index: {0}"), SourceIndex));
 		}
 	}
 
@@ -901,7 +908,7 @@ namespace GLTF
 		TUniquePtr<FArchive> JsonFileReader;
 		if (!FileReader)
 		{
-			Messages.Emplace(EMessageSeverity::Error, TEXT("Can't load file: ") + InFilePath);
+			Messages.Emplace(EMessageSeverity::Error, FText::Format(LOCTEXT("ErrorLoadingFile", "Can't load file: {0}"), FText::FromString(InFilePath)));
 			return;
 		}
 
@@ -926,7 +933,7 @@ namespace GLTF
 		}
 		else
 		{
-			Messages.Emplace(EMessageSeverity::Error, TEXT("Invalid extension."));
+			Messages.Emplace(EMessageSeverity::Error, LOCTEXT("InvalidFileExtention", "Invalid extension."));
 			return;
 		}
 		JsonFileReader.Reset(new FBufferReader(JsonBuffer.GetCharArray().GetData(), sizeof(FString::ElementType) * JsonBuffer.Len(), false));
@@ -936,7 +943,7 @@ namespace GLTF
 		if (!FJsonSerializer::Deserialize(JsonReader, JsonRoot))
 		{
 			JsonRoot.Reset();
-			Messages.Emplace(EMessageSeverity::Error, TEXT("Problem loading JSON."));
+			Messages.Emplace(EMessageSeverity::Error, LOCTEXT("JSONDeserializationError","Problem loading JSON."));
 			return;
 		}
 
@@ -947,7 +954,7 @@ namespace GLTF
 			const double MinVersion = AssetInfo->GetNumberField(TEXT("minVersion"));
 			if (MinVersion > 2.0)
 			{
-				Messages.Emplace(EMessageSeverity::Error, TEXT("This importer supports glTF version 2.0 (or compatible) assets."));
+				Messages.Emplace(EMessageSeverity::Error, LOCTEXT("UnsupportedGLTFAssetMinVersion", "This importer supports glTF version 2.0 (or compatible) assets."));
 				return;
 			}
 			OutAsset.Metadata.Version = MinVersion;
@@ -957,7 +964,7 @@ namespace GLTF
 			const double Version = AssetInfo->GetNumberField(TEXT("version"));
 			if (Version < 2.0)
 			{
-				Messages.Emplace(EMessageSeverity::Error, TEXT("This importer supports glTF asset version 2.0 or later."));
+				Messages.Emplace(EMessageSeverity::Error, LOCTEXT("UnsupportedGLTFAssetVersion", "This importer supports glTF asset version 2.0 or later."));
 				return;
 			}
 			OutAsset.Metadata.Version = Version;
@@ -988,7 +995,7 @@ namespace GLTF
 
 		if (OutAsset.ValidationCheck() != FAsset::Valid)
 		{
-			Messages.Emplace(EMessageSeverity::Warning, FString::Printf(TEXT("For GLTF Asset [%s] not all imported objects are valid."), *OutAsset.Name));
+			Messages.Emplace(EMessageSeverity::Warning, FText::Format(LOCTEXT("SomeImportedObjectsInvalid","For GLTF Asset [{0}] not all imported objects are valid."), FText::FromString(OutAsset.Name)));
 		}
 
 		JsonRoot.Reset();
@@ -1270,7 +1277,7 @@ namespace GLTF
 
 		if (UnUsedIndicesString.Len() > 0)
 		{
-			Messages.Emplace(EMessageSeverity::Warning, FString::Printf(TEXT("Skin objects unused. At indices: %s."), *UnUsedIndicesString));
+			Messages.Emplace(EMessageSeverity::Warning, FText::Format(LOCTEXT("UnusedSkinObjects","Skin objects unused. At indices: {0}."), FText::FromString(UnUsedIndicesString)));
 		}
 	}
 
@@ -1432,7 +1439,7 @@ namespace GLTF
 				OffendingJointsNamesString += OffendingJointName;
 			}
 			
-			Messages.Emplace(EMessageSeverity::Warning, FString::Printf(TEXT("The same Joint(s) are used in multiple Skins with multiple different InverseBindMatrix values, which is not supported. Ignoring InverseBindMatrices for the entire Import. Offending Joints' Names: %s."), *OffendingJointsNamesString));
+			Messages.Emplace(EMessageSeverity::Warning, FText::Format(LOCTEXT("MultipleSkinsUseSameJointProblem","The same Joint(s) are used in multiple Skins with multiple different InverseBindMatrix values, which is not supported. Ignoring InverseBindMatrices for the entire Import. Offending Joints' Names: {0}."), FText::FromString(OffendingJointsNamesString)));
 			
 			Asset->HasAbnormalInverseBindMatrices = true;
 
@@ -1562,6 +1569,8 @@ namespace GLTF
 					const GLTF::FNode SkinSkeletonNode = Asset->Nodes[Skin.Skeleton];
 					int32 SkinSkeletonDistance = GetRootDistance(SkinSkeletonNode);
 					
+
+					const FText MessageText = LOCTEXT("NonCommonRootNode", "Skeleton node is not a common root.");
 					for (TPair<int32, FRootJoints> Group : ParentToRootJointIndices)
 					{
 						if (Group.Value.Indices.Num() == 1 )
@@ -1574,7 +1583,7 @@ namespace GLTF
 							int32 RootNodeCandidateDistance = GetRootDistance(RootNodeCandidate);
 							if (SkinSkeletonDistance > RootNodeCandidateDistance)
 							{
-								Messages.Emplace(EMessageSeverity::Warning, TEXT("Skeleton node is not a common root."));
+								Messages.Emplace(EMessageSeverity::Warning, MessageText);
 								break;
 							}
 						}
@@ -1588,7 +1597,7 @@ namespace GLTF
 							int32 RootNodeCandidateDistance = GetRootDistance(RootNodeCandidate);
 							if (SkinSkeletonDistance > RootNodeCandidateDistance)
 							{
-								Messages.Emplace(EMessageSeverity::Warning, TEXT("Skeleton node is not a common root."));
+								Messages.Emplace(EMessageSeverity::Warning, MessageText);
 								break;
 							}
 						}
@@ -1638,3 +1647,5 @@ namespace GLTF
 	}
 
 }  // namespace GLTF
+
+#undef LOCTEXT_NAMESPACE
\ No newline at end of file
diff --git a/Engine/Plugins/Interchange/Runtime/Source/Parsers/GLTFCore/Private/GLTF/MaterialUtilities.cpp b/Engine/Plugins/Interchange/Runtime/Source/Parsers/GLTFCore/Private/GLTF/MaterialUtilities.cpp
--- a/Engine/Plugins/Interchange/Runtime/Source/Parsers/GLTFCore/Private/GLTF/MaterialUtilities.cpp
+++ b/Engine/Plugins/Interchange/Runtime/Source/Parsers/GLTFCore/Private/GLTF/MaterialUtilities.cpp
@@ -3,6 +3,8 @@
 #include "MaterialUtilities.h"
 #include "GLTFAsset.h"
 
+#define LOCTEXT_NAMESPACE "InterchangeGLTFMaterialUtilities"
+
 namespace GLTF
 {
 	// Returns scale factor if JSON has it, 1.0 by default.
@@ -45,7 +47,7 @@ namespace GLTF
 						const FString ExtensionName = SupportedTextureInfoExtensionsStringified[Index];
 						if (!ExtensionsObj.HasTypedField<EJson::Object>(ExtensionName))
 						{
-							OutMessages.Emplace(RuntimeWarningSeverity(), FString::Printf(TEXT("Extension is not supported: %s"), *ExtensionName));
+							OutMessages.Emplace(RuntimeWarningSeverity(), FText::Format(LOCTEXT("UnsupportedExtension", "Extension is not supported: {0}"), FText::FromString(ExtensionName)));
 							continue;
 						}
 
@@ -94,3 +96,5 @@ namespace GLTF
 	}
 
 }  // namespace GLTF
+
+#undef LOCTEXT_NAMESPACE
\ No newline at end of file
diff --git a/Engine/Plugins/Interchange/Runtime/Source/Parsers/GLTFCore/Private/GLTFMaterialFactory.cpp b/Engine/Plugins/Interchange/Runtime/Source/Parsers/GLTFCore/Private/GLTFMaterialFactory.cpp
--- a/Engine/Plugins/Interchange/Runtime/Source/Parsers/GLTFCore/Private/GLTFMaterialFactory.cpp
+++ b/Engine/Plugins/Interchange/Runtime/Source/Parsers/GLTFCore/Private/GLTFMaterialFactory.cpp
@@ -14,6 +14,8 @@
 #include "GLTFTexture.h"
 // IWYU pragma: end_keep
 
+#define LOCTEXT_NAMESPACE "InterchangeGLTFMaterialFactory"
+
 namespace GLTF
 {
 	namespace
@@ -35,7 +37,7 @@ namespace GLTF
 		{
 			if (!ensure(MaterialElementFactory))
 			{
-				Messages.Emplace(EMessageSeverity::Warning, FString::Printf(TEXT("MaterialElementFactory is unexpected nullptr.")));
+				Messages.Emplace(EMessageSeverity::Warning, LOCTEXT("MaterialElementFactoryNull", "MaterialElementFactory is unexpected nullptr."));
 			}
 		}
 
@@ -226,7 +228,7 @@ namespace GLTF
 			default:
 				if (!ensure(false))
 				{
-					Messages.Emplace(EMessageSeverity::Warning, FString::Printf(TEXT("AlphaMode Opaque not supported for Opacity")));
+					Messages.Emplace(EMessageSeverity::Warning, LOCTEXT("InvalidAlphaMode", "AlphaMode Opaque not supported for Opacity"));
 				}
 				break;
 		}
@@ -834,3 +836,5 @@ namespace GLTF
 	}
 
 }  // namespace GLTF
+
+#undef LOCTEXT_NAMESPACE
\ No newline at end of file
diff --git a/Engine/Plugins/Interchange/Runtime/Source/Parsers/GLTFCore/Private/GLTFMeshFactory.cpp b/Engine/Plugins/Interchange/Runtime/Source/Parsers/GLTFCore/Private/GLTFMeshFactory.cpp
--- a/Engine/Plugins/Interchange/Runtime/Source/Parsers/GLTFCore/Private/GLTFMeshFactory.cpp
+++ b/Engine/Plugins/Interchange/Runtime/Source/Parsers/GLTFCore/Private/GLTFMeshFactory.cpp
@@ -7,6 +7,8 @@
 #include "StaticMeshAttributes.h"
 #include "Util/ColorConstants.h"
 
+#define LOCTEXT_NAMESPACE "InterchangeGLTFMeshFactory"
+
 namespace GLTF
 {
 	class FMeshFactoryImpl : public GLTF::FBaseLogger
@@ -337,7 +339,7 @@ namespace GLTF
 
 			if (!Primitive.IsValid())
 			{
-				Messages.Emplace(EMessageSeverity::Warning, TEXT("Mesh has an invalid primitive: ") + Mesh.Name);
+				Messages.Emplace(EMessageSeverity::Warning, FText::Format(LOCTEXT("InvalidMeshPrimitive", "Mesh has an invalid primitive : {0}"), FText::FromString(Mesh.Name)));
 				continue;
 			}
 
@@ -417,13 +419,12 @@ namespace GLTF
 
 			if (bHasDegenerateTriangles)
 			{
-				Messages.Emplace(EMessageSeverity::Warning,
-					FString::Printf(TEXT("Mesh %s has primitive with degenerate triangles: %d"), *Mesh.Name, Index));
+				Messages.Emplace(EMessageSeverity::Warning, FText::Format(LOCTEXT("PrimitiveFoundDegenerateTriangles", "Mesh {0} has primitive with degenerate triangles: {1}"), FText::FromString(Mesh.Name), Index));
 			}
 		}
 		if (bMeshUsesEmptyMaterial)
 		{
-			Messages.Emplace(EMessageSeverity::Warning, TEXT("Mesh has primitives with no materials assigned: ") + Mesh.Name);
+			Messages.Emplace(EMessageSeverity::Warning, FText::Format(LOCTEXT("PrimitiveNoMaterialsAssigned", "Mesh has primitives with no materials assigned: {0}") , FText::FromString(Mesh.Name)));
 		}
 	}
 
@@ -789,3 +790,5 @@ namespace GLTF
 		return Impl->PositionIndexToVertexIdPerPrim;
 	}
 } //namespace GLTF
+
+#undef LOCTEXT_NAMESPACE
\ No newline at end of file
diff --git a/Engine/Plugins/Interchange/Runtime/Source/Parsers/GLTFCore/Public/GLTFLogger.h b/Engine/Plugins/Interchange/Runtime/Source/Parsers/GLTFCore/Public/GLTFLogger.h
--- a/Engine/Plugins/Interchange/Runtime/Source/Parsers/GLTFCore/Public/GLTFLogger.h
+++ b/Engine/Plugins/Interchange/Runtime/Source/Parsers/GLTFCore/Public/GLTFLogger.h
@@ -3,7 +3,7 @@
 #pragma once
 
 #include "Templates/Tuple.h"
-#include "Containers/UnrealString.h"
+#include "Internationalization/Text.h"
 
 namespace GLTF
 {
@@ -22,7 +22,7 @@ namespace GLTF
 	*/
 	EMessageSeverity RuntimeWarningSeverity();
 
-	using FLogMessage = TTuple<EMessageSeverity, FString>;
+	using FLogMessage = TTuple<EMessageSeverity, FText>;
 
 	class GLTFCORE_API FBaseLogger
 	{
diff --git a/Engine/Source/Programs/InterchangeWorker/Private/InterchangeWorkerImpl.cpp b/Engine/Source/Programs/InterchangeWorker/Private/InterchangeWorkerImpl.cpp
--- a/Engine/Source/Programs/InterchangeWorker/Private/InterchangeWorkerImpl.cpp
+++ b/Engine/Source/Programs/InterchangeWorker/Private/InterchangeWorkerImpl.cpp
@@ -51,7 +51,7 @@ bool FInterchangeWorkerImpl::Run(const FString& WorkerVersionError)
 		FErrorCommand ErrorCmd;
 		
 		UInterchangeResultError_Generic* Message = FbxParser.AddMessage<UInterchangeResultError_Generic>();
-		Message->Text = FText::FromString(WorkerVersionError);
+		Message->Text = FText::Format(NSLOCTEXT("InterchangeWorker", "WorkerVersionError", "Version Error: {0}"), FText::FromString(WorkerVersionError));
 
 		ErrorCmd.ErrorMessage = Message->ToJson();
 		CommandIO.SendCommand(ErrorCmd, Config::SendCommandTimeout_s);
diff --git a/Engine/Source/Runtime/Interchange/Engine/Private/Tasks/InterchangeTaskParsing.cpp b/Engine/Source/Runtime/Interchange/Engine/Private/Tasks/InterchangeTaskParsing.cpp
--- a/Engine/Source/Runtime/Interchange/Engine/Private/Tasks/InterchangeTaskParsing.cpp
+++ b/Engine/Source/Runtime/Interchange/Engine/Private/Tasks/InterchangeTaskParsing.cpp
@@ -403,21 +403,20 @@ void UE::Interchange::FTaskParsing::Execute()
 		}
 		for (const TPair<TWeakObjectPtr<UInterchangeTranslatorBase>, FString>& MessagePerTranslator : TranslatorMessageMap)
 		{
+
+			FText ResultText = FText::Format(NSLOCTEXT("InterchangeTaskParsingDoTask", "RenamedAssetsMessageHeader", "Renamed Assets: {0}")
+				, FText::FromString(MessagePerTranslator.Value));
+
 			if (GIsAutomationTesting)
 			{
-				
 				UInterchangeResultDisplay_Generic* DisplayResult = NewObject<UInterchangeResultDisplay_Generic>(GetTransientPackage(), UInterchangeResultDisplay_Generic::StaticClass());
-				FString Message = NSLOCTEXT("InterchangeTaskParsingDoTask", "RenamedAssetsMessageHeader", "Renamed Assets:").ToString();
-				Message += MessagePerTranslator.Value;
-				DisplayResult->Text = FText::FromString(Message);
+				DisplayResult->Text = ResultText;
 				MessagePerTranslator.Key->AddMessage(DisplayResult);
 			}
 			else
 			{
 				UInterchangeResultWarning_Generic* WarningResult = NewObject<UInterchangeResultWarning_Generic>(GetTransientPackage(), UInterchangeResultWarning_Generic::StaticClass());
-				FString Message = NSLOCTEXT("InterchangeTaskParsingDoTask", "RenamedAssetsMessageHeader", "Renamed Assets:").ToString();
-				Message += MessagePerTranslator.Value;
-				WarningResult->Text = FText::FromString(Message);
+				WarningResult->Text = ResultText;
 				MessagePerTranslator.Key->AddMessage(WarningResult);
 			}
 		}
