[Interchange] Animation Curve Tangent fixes.

#rnx

diff --git a/Engine/Plugins/Interchange/Runtime/Source/Import/Private/Animation/InterchangeLevelSequenceFactory.cpp b/Engine/Plugins/Interchange/Runtime/Source/Import/Private/Animation/InterchangeLevelSequenceFactory.cpp
--- a/Engine/Plugins/Interchange/Runtime/Source/Import/Private/Animation/InterchangeLevelSequenceFactory.cpp
+++ b/Engine/Plugins/Interchange/Runtime/Source/Import/Private/Animation/InterchangeLevelSequenceFactory.cpp
@@ -349,6 +349,111 @@ namespace UE::Interchange::Private
 		}
 	}
 
+	template <class T>
+	void ConvertRichCurveKeyToFloatValue(const FRichCurveKey& RichCurveKey, T& OutMovieSceneKey, double TangentRatio /*= 1.0*/, double SecondsPerFrame /*= 1.0*/)
+	{
+		OutMovieSceneKey.Value = RichCurveKey.Value;
+
+		OutMovieSceneKey.Tangent.TangentWeightMode = RichCurveKey.TangentWeightMode;
+		if (OutMovieSceneKey.Tangent.TangentWeightMode != RCTWM_WeightedNone &&
+			OutMovieSceneKey.Tangent.TangentWeightMode != RCTWM_WeightedArrive &&
+			OutMovieSceneKey.Tangent.TangentWeightMode != RCTWM_WeightedLeave &&
+			OutMovieSceneKey.Tangent.TangentWeightMode != RCTWM_WeightedBoth
+			)
+		{
+			OutMovieSceneKey.Tangent.TangentWeightMode = RCTWM_WeightedNone;
+		}
+
+		OutMovieSceneKey.Tangent.ArriveTangentWeight = RichCurveKey.ArriveTangentWeight;
+		OutMovieSceneKey.Tangent.LeaveTangentWeight = RichCurveKey.LeaveTangentWeight;
+
+		if (OutMovieSceneKey.Tangent.TangentWeightMode == RCTWM_WeightedNone)
+		{
+			OutMovieSceneKey.Tangent.ArriveTangent = static_cast<float>(RichCurveKey.ArriveTangent * TangentRatio);
+			OutMovieSceneKey.Tangent.LeaveTangent = static_cast<float>(RichCurveKey.LeaveTangent * TangentRatio);
+		}
+		else
+		{
+			OutMovieSceneKey.Tangent.ArriveTangent = static_cast<float>(RichCurveKey.ArriveTangent * SecondsPerFrame);
+			OutMovieSceneKey.Tangent.LeaveTangent = static_cast<float>(RichCurveKey.LeaveTangent * SecondsPerFrame);
+		}
+
+		OutMovieSceneKey.TangentMode = RichCurveKey.TangentMode;
+		OutMovieSceneKey.InterpMode = RichCurveKey.InterpMode;
+	}
+
+	template <class T>
+	void ProcessRichCurveKeys(const FFrameRate& TargetFrameRate, const FRichCurve& Curve, TArray<FFrameNumber>& OutFrameNumbers, TArray<T>& OutValues, FFrameNumber& MinFrameNumber, FFrameNumber& MaxFrameNumber)
+	{
+		const TArray<FRichCurveKey>& CurveKeys = Curve.GetConstRefOfKeys();
+		const int32 NumCurveKeys = CurveKeys.Num();
+
+		OutFrameNumbers.Reserve(NumCurveKeys);
+		OutValues.Reserve(NumCurveKeys);
+
+		for (int32 KeyIndex = 0; KeyIndex < CurveKeys.Num(); ++KeyIndex)
+		{
+			const FRichCurveKey* PrevKey = KeyIndex > 0 ? &CurveKeys[KeyIndex - 1] : nullptr;
+			const FRichCurveKey* NextKey = KeyIndex < (NumCurveKeys - 1) ? &CurveKeys[KeyIndex + 1] : nullptr;
+			const FRichCurveKey& RichCurveKey = CurveKeys[KeyIndex];
+
+			FFrameNumber& FrameNumber = OutFrameNumbers.Add_GetRef(TargetFrameRate.AsFrameNumber(RichCurveKey.Time));
+
+			if (FrameNumber < MinFrameNumber)
+			{
+				MinFrameNumber = FrameNumber;
+			}
+
+			if (FrameNumber > MaxFrameNumber)
+			{
+				MaxFrameNumber = FrameNumber;
+			}
+
+			const float SecondsDelta = [&]() -> float
+				{
+					if (PrevKey && NextKey)
+					{
+						return NextKey->Time - PrevKey->Time;
+					}
+					else if (PrevKey)
+					{
+						return RichCurveKey.Time - PrevKey->Time;
+					}
+					else if (NextKey)
+					{
+						return NextKey->Time - RichCurveKey.Time;
+					}
+
+					return 1.f;
+				}();
+
+			const int32 FrameNumberDelta = [&]() -> int32
+				{
+					if (PrevKey && NextKey)
+					{
+						return TargetFrameRate.AsFrameTime(NextKey->Time).RoundToFrame().Value - TargetFrameRate.AsFrameTime(PrevKey->Time).RoundToFrame().Value;
+					}
+					else if (PrevKey)
+					{
+						return TargetFrameRate.AsFrameTime(RichCurveKey.Time).RoundToFrame().Value - TargetFrameRate.AsFrameTime(PrevKey->Time).RoundToFrame().Value;
+					}
+					else if (NextKey)
+					{
+						return TargetFrameRate.AsFrameTime(NextKey->Time).RoundToFrame().Value - TargetFrameRate.AsFrameTime(RichCurveKey.Time).RoundToFrame().Value;
+					}
+
+					return 1;
+				}();
+
+			T& Value = OutValues.AddDefaulted_GetRef();
+
+			// Ratio between rich-curve and moviescene key(s) timing, if there are any surrounding keys (otherwise default to a ratio of 1:1)
+			const double KeyTimingRatio = (PrevKey || NextKey) ? SecondsDelta / FrameNumberDelta : 1.0;
+
+			ConvertRichCurveKeyToFloatValue(RichCurveKey, Value, KeyTimingRatio, TargetFrameRate.AsInterval());
+		}
+	}
+
 	void FLevelSequenceHelper::PopulateSubsequenceTrack(const UInterchangeAnimationTrackSetInstanceNode& InstanceNode)
 	{
 		FString TrackSetNodeUid;
@@ -503,41 +608,16 @@ namespace UE::Interchange::Private
 
 		auto CopyToChannel = [this](auto Channel, const FRichCurve& Curve)
 		{
-			const FFrameRate& FrameRate = this->MovieScene->GetTickResolution();
-
-			const TArray<FRichCurveKey>& CurveKeys = Curve.GetConstRefOfKeys();
-
 			TArray<FFrameNumber> FrameNumbers;
-			FrameNumbers.Reserve(CurveKeys.Num());
 
 			using FMovieSceneValue = typename std::remove_pointer_t<decltype(Channel)>::ChannelValueType;
-			TArray<FMovieSceneValue> MovieSceneValues;
-			MovieSceneValues.Reserve(CurveKeys.Num());
-
-			for(int32 KeyIndex = 0; KeyIndex < CurveKeys.Num(); ++KeyIndex)
-			{
-				const FRichCurveKey& CurveKey = CurveKeys[KeyIndex];
+			TArray<FMovieSceneValue> Values;
 
-				FFrameNumber& FrameNumber = FrameNumbers.Add_GetRef(FrameRate.AsFrameNumber(CurveKey.Time));
-
-				if(FrameNumber < this->MinFrameNumber)
-				{
-					this->MinFrameNumber = FrameNumber;
-				}
+			ProcessRichCurveKeys(this->MovieScene->GetTickResolution(), Curve, FrameNumbers, Values, this->MinFrameNumber, this->MaxFrameNumber);
 
-				if(FrameNumber > this->MaxFrameNumber)
-				{
-					this->MaxFrameNumber = FrameNumber;
-				}
-
-				FMovieSceneValue& SceneValue = MovieSceneValues.AddDefaulted_GetRef();
-				SceneValue.InterpMode = CurveKey.InterpMode;
-				SceneValue.Value = CurveKey.Value;
-			}
-
-			if(!MovieSceneValues.IsEmpty())
+			if(!Values.IsEmpty())
 			{
-				Channel->Set(FrameNumbers, MovieSceneValues);
+				Channel->Set(FrameNumbers, Values);
 			}
 			else
 			{
@@ -673,41 +753,14 @@ namespace UE::Interchange::Private
 	{
 		auto CopyToChannel = [this](FMovieSceneDoubleChannel* Channel, const FRichCurve& Curve)
 		{
-			const FFrameRate& FrameRate = this->MovieScene->GetTickResolution();
-
-			const TArray<FRichCurveKey>& CurveKeys = Curve.GetConstRefOfKeys();
-			
 			TArray<FFrameNumber> FrameNumbers;
-			FrameNumbers.Reserve(CurveKeys.Num());
-			
-			TArray<FMovieSceneDoubleValue> MovieSceneDoubleValues;
-			MovieSceneDoubleValues.Reserve(CurveKeys.Num());
-
-			for (int32 KeyIndex = 0; KeyIndex < CurveKeys.Num(); ++KeyIndex)
-			{
-				const FRichCurveKey& CurveKey = CurveKeys[KeyIndex];
-
-				FFrameNumber& FrameNumber = FrameNumbers.Add_GetRef(FrameRate.AsFrameNumber(CurveKey.Time));
+			TArray<FMovieSceneDoubleValue> Values;
 
-				if (FrameNumber < this->MinFrameNumber)
-				{
-					this->MinFrameNumber = FrameNumber;
-				}
-
-				if (FrameNumber > this->MaxFrameNumber)
-				{
-					this->MaxFrameNumber = FrameNumber;
-				}
-
-				FMovieSceneDoubleValue& SceneValue = MovieSceneDoubleValues.AddDefaulted_GetRef();
-
-				SceneValue.InterpMode = CurveKey.InterpMode;
-				SceneValue.Value = CurveKey.Value;
-			}
+			ProcessRichCurveKeys(this->MovieScene->GetTickResolution(), Curve, FrameNumbers, Values, this->MinFrameNumber, this->MaxFrameNumber);
 
-			if (!MovieSceneDoubleValues.IsEmpty())
+			if (!Values.IsEmpty())
 			{
-				Channel->Set(FrameNumbers, MovieSceneDoubleValues);
+				Channel->Set(FrameNumbers, Values);
 			}
 			else
 			{
