Interchange:
- New Collision UI
- Avoid creating new editor-collisions if there's already one of the requested type.

diff --git a/Engine/Plugins/Interchange/Editor/Source/InterchangeEditor/Private/InterchangeFbxAssetImportDataConverter.cpp b/Engine/Plugins/Interchange/Editor/Source/InterchangeEditor/Private/InterchangeFbxAssetImportDataConverter.cpp
--- a/Engine/Plugins/Interchange/Editor/Source/InterchangeEditor/Private/InterchangeFbxAssetImportDataConverter.cpp
+++ b/Engine/Plugins/Interchange/Editor/Source/InterchangeEditor/Private/InterchangeFbxAssetImportDataConverter.cpp
@@ -123,7 +123,7 @@ namespace UE::Interchange::Private
 
 	void FillFbxStaticMeshImportData(const UInterchangeGenericAssetsPipeline* GenericAssetPipeline, UFbxStaticMeshImportData* DestinationStaticMeshImportData)
 	{
-		DestinationStaticMeshImportData->bAutoGenerateCollision = GenericAssetPipeline->MeshPipeline->Collision != EInterchangeMeshCollision::None;
+		DestinationStaticMeshImportData->bAutoGenerateCollision = GenericAssetPipeline->MeshPipeline->bCollision && GenericAssetPipeline->MeshPipeline->Collision != EInterchangeMeshCollision::None && GenericAssetPipeline->MeshPipeline->Collision != EInterchangeMeshCollision::Custom;
 		DestinationStaticMeshImportData->bBuildNanite = GenericAssetPipeline->MeshPipeline->bBuildNanite;
 		DestinationStaticMeshImportData->bBuildReversedIndexBuffer = GenericAssetPipeline->MeshPipeline->bBuildReversedIndexBuffer;
 		DestinationStaticMeshImportData->bCombineMeshes = GenericAssetPipeline->MeshPipeline->bCombineStaticMeshes;
@@ -311,6 +311,7 @@ namespace UE::Interchange::Private
 			FillInterchangeGenericAssetsPipelineFromFbxMeshImportData(GenericAssetPipeline, StaticMeshImportData);
 		}
 
+		GenericAssetPipeline->MeshPipeline->bCollision = StaticMeshImportData->bAutoGenerateCollision;
 		GenericAssetPipeline->MeshPipeline->Collision = StaticMeshImportData->bAutoGenerateCollision ? EInterchangeMeshCollision::Convex18DOP : EInterchangeMeshCollision::None;
 		GenericAssetPipeline->MeshPipeline->bBuildNanite = StaticMeshImportData->bBuildNanite;
 		GenericAssetPipeline->MeshPipeline->bBuildReversedIndexBuffer = StaticMeshImportData->bBuildReversedIndexBuffer;
diff --git a/Engine/Plugins/Interchange/Runtime/Source/FactoryNodes/Private/InterchangeStaticMeshLodDataNode.cpp b/Engine/Plugins/Interchange/Runtime/Source/FactoryNodes/Private/InterchangeStaticMeshLodDataNode.cpp
--- a/Engine/Plugins/Interchange/Runtime/Source/FactoryNodes/Private/InterchangeStaticMeshLodDataNode.cpp
+++ b/Engine/Plugins/Interchange/Runtime/Source/FactoryNodes/Private/InterchangeStaticMeshLodDataNode.cpp
@@ -322,24 +322,24 @@ bool UInterchangeStaticMeshLodDataNode::SetOneConvexHullPerUCX(bool AttributeVal
 	IMPLEMENT_NODE_ATTRIBUTE_SETTER_NODELEGATE(OneConvexHullPerUCX, bool)
 }
 
-bool UInterchangeStaticMeshLodDataNode::GetImportCollision(EInterchangeMeshCollision& AttributeValue) const
+bool UInterchangeStaticMeshLodDataNode::GetImportCollisionType(EInterchangeMeshCollision& AttributeValue) const
 {
-	IMPLEMENT_NODE_ATTRIBUTE_GETTER(ImportCollision, EInterchangeMeshCollision);
+	IMPLEMENT_NODE_ATTRIBUTE_GETTER(ImportCollisionType, EInterchangeMeshCollision);
 }
 
-bool UInterchangeStaticMeshLodDataNode::SetImportCollision(EInterchangeMeshCollision AttributeValue)
+bool UInterchangeStaticMeshLodDataNode::SetImportCollisionType(EInterchangeMeshCollision AttributeValue)
 {
-	IMPLEMENT_NODE_ATTRIBUTE_SETTER_NODELEGATE(ImportCollision, EInterchangeMeshCollision)
+	IMPLEMENT_NODE_ATTRIBUTE_SETTER_NODELEGATE(ImportCollisionType, EInterchangeMeshCollision)
 }
 
 bool UInterchangeStaticMeshLodDataNode::GetImportCollision(bool& AttributeValue) const
 {
-	return false;
+	IMPLEMENT_NODE_ATTRIBUTE_GETTER(ImportCollision, bool);
 }
 
 bool UInterchangeStaticMeshLodDataNode::SetImportCollision(bool AttributeValue)
 {
-	return false;
+	IMPLEMENT_NODE_ATTRIBUTE_SETTER_NODELEGATE(ImportCollision, bool);
 }
 
 bool UInterchangeStaticMeshLodDataNode::IsEditorOnlyDataDefined()
diff --git a/Engine/Plugins/Interchange/Runtime/Source/FactoryNodes/Public/InterchangeStaticMeshLodDataNode.h b/Engine/Plugins/Interchange/Runtime/Source/FactoryNodes/Public/InterchangeStaticMeshLodDataNode.h
--- a/Engine/Plugins/Interchange/Runtime/Source/FactoryNodes/Public/InterchangeStaticMeshLodDataNode.h
+++ b/Engine/Plugins/Interchange/Runtime/Source/FactoryNodes/Public/InterchangeStaticMeshLodDataNode.h
@@ -129,15 +129,15 @@ public:
 	bool SetOneConvexHullPerUCX(bool AttributeValue);
 
 	UFUNCTION(BlueprintCallable, Category = "Interchange | Node | StaticMeshLodData")
-	bool GetImportCollision(EInterchangeMeshCollision& AttributeValue) const;
+	bool GetImportCollisionType(EInterchangeMeshCollision& AttributeValue) const;
 
 	UFUNCTION(BlueprintCallable, Category = "Interchange | Node | StaticMeshLodData")
-	bool SetImportCollision(EInterchangeMeshCollision AttributeValue);
+	bool SetImportCollisionType(EInterchangeMeshCollision AttributeValue);
 
-	UE_DEPRECATED(5.5, "GetImportCollision using a bool has been deprecated, please use the function with the enum instead.")
+	UFUNCTION(BlueprintCallable, Category = "Interchange | Node | StaticMeshLodData")
 	bool GetImportCollision(bool& AttributeValue) const;
 
-	UE_DEPRECATED(5.5, "SetImportCollision using a bool has been deprecated, please use the function with the enum instead.")
+	UFUNCTION(BlueprintCallable, Category = "Interchange | Node | StaticMeshLodData")
 	bool SetImportCollision(bool AttributeValue);
 
 
@@ -150,7 +150,8 @@ private:
 	UE::Interchange::TArrayAttributeHelper<FString> CapsuleCollisionMeshUids;
 	UE::Interchange::TArrayAttributeHelper<FString> SphereCollisionMeshUids;
 	UE::Interchange::TArrayAttributeHelper<FString> ConvexCollisionMeshUids;
-
-	const UE::Interchange::FAttributeKey Macro_CustomOneConvexHullPerUCXKey = UE::Interchange::FAttributeKey(TEXT("__OneConvexHullPerUCX__Key"));
-	const UE::Interchange::FAttributeKey Macro_CustomImportCollisionKey = UE::Interchange::FAttributeKey(TEXT("__ImportCollision__Key"));
+		
+	IMPLEMENT_NODE_ATTRIBUTE_KEY(OneConvexHullPerUCX)
+	IMPLEMENT_NODE_ATTRIBUTE_KEY(ImportCollision)
+	IMPLEMENT_NODE_ATTRIBUTE_KEY(ImportCollisionType)
 };
diff --git a/Engine/Plugins/Interchange/Runtime/Source/Import/Private/Mesh/InterchangeStaticMeshFactory.cpp b/Engine/Plugins/Interchange/Runtime/Source/Import/Private/Mesh/InterchangeStaticMeshFactory.cpp
--- a/Engine/Plugins/Interchange/Runtime/Source/Import/Private/Mesh/InterchangeStaticMeshFactory.cpp
+++ b/Engine/Plugins/Interchange/Runtime/Source/Import/Private/Mesh/InterchangeStaticMeshFactory.cpp
@@ -302,6 +302,7 @@ UInterchangeFactoryBase::FImportAssetResult UInterchangeStaticMeshFactory::Impor
 	TArray<FMeshDescription>& LodMeshDescriptions = ImportAssetObjectData.LodMeshDescriptions;
 	LodMeshDescriptions.SetNum(LodCount);
 
+	bool bImportCollision = false;
 	EInterchangeMeshCollision Collision = EInterchangeMeshCollision::None;
 	bool bImportedCustomCollision = false;
 	int32 CurrentLodIndex = 0;
@@ -427,8 +428,9 @@ UInterchangeFactoryBase::FImportAssetResult UInterchangeStaticMeshFactory::Impor
 		// Import collision geometry
 		if (CurrentLodIndex == 0)
 		{
-			LodDataNode->GetImportCollision(Collision);
-			if(Collision != EInterchangeMeshCollision::None)
+			LodDataNode->GetImportCollision(bImportCollision);
+			LodDataNode->GetImportCollisionType(Collision);
+			if(bImportCollision && Collision != EInterchangeMeshCollision::None)
 			{
 				if (bReimport)
 				{
@@ -471,6 +473,7 @@ UInterchangeFactoryBase::FImportAssetResult UInterchangeStaticMeshFactory::Impor
 	}
 #endif // WITH_EDITOR
 
+	ImportAssetObjectData.bImportCollision = bImportCollision;
 	ImportAssetObjectData.Collision = Collision;
 	ImportAssetObjectData.bImportedCustomCollision = bImportedCustomCollision;
 
@@ -674,44 +677,79 @@ UInterchangeFactoryBase::FImportAssetResult UInterchangeStaticMeshFactory::EndIm
 #endif // WITH_EDITOR
 
 #if WITH_EDITOR
-	if(!ImportAssetObjectData.bImportedCustomCollision)
+	if(ImportAssetObjectData.bImportCollision)
 	{
-		constexpr bool bUpdateRendering = false;
-		switch(ImportAssetObjectData.Collision)
+		if(!ImportAssetObjectData.bImportedCustomCollision)
 		{
-		case EInterchangeMeshCollision::Box:
-			GenerateBoxAsSimpleCollision(StaticMesh, bUpdateRendering);
-			break;
-		case EInterchangeMeshCollision::Sphere:
-			GenerateSphereAsSimpleCollision(StaticMesh, bUpdateRendering);
-			break;
-		case EInterchangeMeshCollision::Capsule:
-			GenerateSphylAsSimpleCollision(StaticMesh, bUpdateRendering);
-			break;
-		case EInterchangeMeshCollision::Convex10DOP_X:
-			GenerateKDopAsSimpleCollision(StaticMesh, TArray<FVector>(KDopDir10X, sizeof(KDopDir10X) / sizeof(FVector)), bUpdateRendering);
-			break;
-		case EInterchangeMeshCollision::Convex10DOP_Y:
-			GenerateKDopAsSimpleCollision(StaticMesh, TArray<FVector>(KDopDir10Y, sizeof(KDopDir10Y) / sizeof(FVector)), bUpdateRendering);
-			break;
-		case EInterchangeMeshCollision::Convex10DOP_Z:
-			GenerateKDopAsSimpleCollision(StaticMesh, TArray<FVector>(KDopDir10Z, sizeof(KDopDir10Z) / sizeof(FVector)), bUpdateRendering);
-			break;
-		case EInterchangeMeshCollision::Convex18DOP:
-			GenerateKDopAsSimpleCollision(StaticMesh, TArray<FVector>(KDopDir18, sizeof(KDopDir18) / sizeof(FVector)), bUpdateRendering);
-			break;
-		case EInterchangeMeshCollision::Convex26DOP:
-			GenerateKDopAsSimpleCollision(StaticMesh, TArray<FVector>(KDopDir26, sizeof(KDopDir26) / sizeof(FVector)), bUpdateRendering);
-			break;
-		default:
-			break;
+			// Don't generate collisions if the mesh already has one of the requested type, otherwise it will continue to create collisions,
+			// it can happen in the case of an import, and then importing the same file without deleting the asset in the content browser (different from a reimport)
+			bool bHasBoxCollision = !StaticMesh->GetBodySetup()->AggGeom.BoxElems.IsEmpty();
+			bool bHasSphereCollision = !StaticMesh->GetBodySetup()->AggGeom.SphereElems.IsEmpty();
+			bool bHasCapsuleCollision = !StaticMesh->GetBodySetup()->AggGeom.SphylElems.IsEmpty();
+			bool bHasConvexCollision = !StaticMesh->GetBodySetup()->AggGeom.ConvexElems.IsEmpty();
+
+			constexpr bool bUpdateRendering = false;
+			switch(ImportAssetObjectData.Collision)
+			{
+			case EInterchangeMeshCollision::Box:
+				if(!bHasBoxCollision)
+				{
+					GenerateBoxAsSimpleCollision(StaticMesh, bUpdateRendering);
+				}
+				break;
+			case EInterchangeMeshCollision::Sphere:
+				if(!bHasSphereCollision)
+				{
+					GenerateSphereAsSimpleCollision(StaticMesh, bUpdateRendering);
+				}
+				break;
+			case EInterchangeMeshCollision::Capsule:
+				if(!bHasCapsuleCollision)
+				{
+					GenerateSphylAsSimpleCollision(StaticMesh, bUpdateRendering);
+				}
+				break;
+			case EInterchangeMeshCollision::Convex10DOP_X:
+				if(!bHasConvexCollision)
+				{
+					GenerateKDopAsSimpleCollision(StaticMesh, TArray<FVector>(KDopDir10X, sizeof(KDopDir10X) / sizeof(FVector)), bUpdateRendering);
+				}
+				break;
+			case EInterchangeMeshCollision::Convex10DOP_Y:
+				if(!bHasConvexCollision)
+				{
+					GenerateKDopAsSimpleCollision(StaticMesh, TArray<FVector>(KDopDir10Y, sizeof(KDopDir10Y) / sizeof(FVector)), bUpdateRendering);
+
+				}
+				break;
+			case EInterchangeMeshCollision::Convex10DOP_Z:
+				if(!bHasConvexCollision)
+				{
+					GenerateKDopAsSimpleCollision(StaticMesh, TArray<FVector>(KDopDir10Z, sizeof(KDopDir10Z) / sizeof(FVector)), bUpdateRendering);
+				}
+				break;
+			case EInterchangeMeshCollision::Convex18DOP:
+				if(!bHasConvexCollision)
+				{
+					GenerateKDopAsSimpleCollision(StaticMesh, TArray<FVector>(KDopDir18, sizeof(KDopDir18) / sizeof(FVector)), bUpdateRendering);
+				}
+				break;
+			case EInterchangeMeshCollision::Convex26DOP:
+				if(!bHasConvexCollision)
+				{
+					GenerateKDopAsSimpleCollision(StaticMesh, TArray<FVector>(KDopDir26, sizeof(KDopDir26) / sizeof(FVector)), bUpdateRendering);
+				}
+				break;
+			default:
+				break;
+			}
 		}
-	}
 #endif
 #if WITH_EDITORONLY_DATA
-	else
-	{
-		StaticMesh->bCustomizedCollision = true;
+		else
+		{
+			StaticMesh->bCustomizedCollision = true;
+		}
 	}
 #endif // WITH_EDITORONLY_DATA
 #if WITH_EDITOR
diff --git a/Engine/Plugins/Interchange/Runtime/Source/Import/Public/Mesh/InterchangeStaticMeshFactory.h b/Engine/Plugins/Interchange/Runtime/Source/Import/Public/Mesh/InterchangeStaticMeshFactory.h
--- a/Engine/Plugins/Interchange/Runtime/Source/Import/Public/Mesh/InterchangeStaticMeshFactory.h
+++ b/Engine/Plugins/Interchange/Runtime/Source/Import/Public/Mesh/InterchangeStaticMeshFactory.h
@@ -79,6 +79,7 @@ private:
 	{
 		bool bIsAppGame = false;
 		bool bImportedCustomCollision = false;
+		bool bImportCollision = false;
 		EInterchangeMeshCollision Collision = EInterchangeMeshCollision::None;
 		TArray<FMeshDescription> LodMeshDescriptions;
 	};
diff --git a/Engine/Plugins/Interchange/Runtime/Source/Nodes/Public/InterchangeMeshDefinitions.h b/Engine/Plugins/Interchange/Runtime/Source/Nodes/Public/InterchangeMeshDefinitions.h
--- a/Engine/Plugins/Interchange/Runtime/Source/Nodes/Public/InterchangeMeshDefinitions.h
+++ b/Engine/Plugins/Interchange/Runtime/Source/Nodes/Public/InterchangeMeshDefinitions.h
@@ -34,7 +34,8 @@ enum class EInterchangeMeshCollision : uint8
 	/** Generates a new axis-aligned box collision mesh with all edges and corners beveled (26 total sides) encompassing the static mesh*/
 	Convex26DOP UMETA(DisplayName = "26DOP Simplified"),
 	
-	/** Generates only collisions retrieved from the file*/
+	/** Generates only collisions retrieved from the file. This option is mostly useful in case of a reimport,
+	  * to only import collisions by removing the previous imported ones and preserving the ones that were editor-created.*/
 	Custom UMETA(DisplayName = "From File"),
 
 	None = 0xFF
diff --git a/Engine/Plugins/Interchange/Runtime/Source/Pipelines/Private/InterchangeGenericMeshPipeline.cpp b/Engine/Plugins/Interchange/Runtime/Source/Pipelines/Private/InterchangeGenericMeshPipeline.cpp
--- a/Engine/Plugins/Interchange/Runtime/Source/Pipelines/Private/InterchangeGenericMeshPipeline.cpp
+++ b/Engine/Plugins/Interchange/Runtime/Source/Pipelines/Private/InterchangeGenericMeshPipeline.cpp
@@ -106,6 +106,7 @@ void UInterchangeGenericMeshPipeline::AdjustSettingsForContext(const FInterchang
 			bCombineStaticMeshes = true;
 			bBuildNanite = false;
 			LodGroup = NAME_None;
+			bCollision = false;
 			Collision = EInterchangeMeshCollision::None;
 			bImportCollisionAccordingToMeshName = false;
 			bGenerateLightmapUVs = false;
@@ -128,6 +129,7 @@ void UInterchangeGenericMeshPipeline::AdjustSettingsForContext(const FInterchang
 					bCombineStaticMeshes = true;
 					LodGroup = NAME_None;
 					bSupportFaceRemap = false;
+					bCollision = false;
 					Collision = EInterchangeMeshCollision::None;
 					bImportCollisionAccordingToMeshName = false;
 					bGenerateLightmapUVs = false;
@@ -353,10 +355,9 @@ void UInterchangeGenericMeshPipeline::GetMeshesInformationFromTranslatedData(con
 void UInterchangeGenericMeshPipeline::PostLoad()
 {
 	Super::PostLoad();
-
 	if(!bImportCollision_DEPRECATED)
 	{
-		Collision = EInterchangeMeshCollision::None;
+		bCollision = bImportCollision_DEPRECATED;
 	}
 }
 
diff --git a/Engine/Plugins/Interchange/Runtime/Source/Pipelines/Private/InterchangeGenericStaticMeshPipeline.cpp b/Engine/Plugins/Interchange/Runtime/Source/Pipelines/Private/InterchangeGenericStaticMeshPipeline.cpp
--- a/Engine/Plugins/Interchange/Runtime/Source/Pipelines/Private/InterchangeGenericStaticMeshPipeline.cpp
+++ b/Engine/Plugins/Interchange/Runtime/Source/Pipelines/Private/InterchangeGenericStaticMeshPipeline.cpp
@@ -270,7 +270,7 @@ void UInterchangeGenericMeshPipeline::ExecutePreImportPipelineStaticMesh()
 
 					if (MeshUidsPerLodIndex.Num() > 0)
 					{
-						if (Collision != EInterchangeMeshCollision::None && bImportCollisionAccordingToMeshName)
+						if (bCollision && Collision != EInterchangeMeshCollision::None && bImportCollisionAccordingToMeshName)
 						{
 							if (const TArray<FString>* CorrespondingCollisionMeshes = MeshToCollisionMeshMap.Find(MeshUid))
 							{
@@ -314,7 +314,7 @@ void UInterchangeGenericMeshPipeline::ExecutePreImportPipelineStaticMesh()
 
 					if (MeshUidsPerLodIndex.Num() > 0)
 					{
-						if (Collision != EInterchangeMeshCollision::None && bImportCollisionAccordingToMeshName)
+						if (bCollision && Collision != EInterchangeMeshCollision::None && bImportCollisionAccordingToMeshName)
 						{
 							if (const TArray<FString>* CorrespondingCollisionMeshes = MeshToCollisionMeshMap.Find(MeshUid))
 							{
@@ -500,7 +500,8 @@ UInterchangeStaticMeshLodDataNode* UInterchangeGenericMeshPipeline::CreateStatic
 
 	StaticMeshLodDataNode->InitializeNode(NodeUID, DisplayLabel, EInterchangeNodeContainerType::FactoryData);
 	StaticMeshLodDataNode->SetOneConvexHullPerUCX(bOneConvexHullPerUCX);
-	StaticMeshLodDataNode->SetImportCollision(Collision);
+	StaticMeshLodDataNode->SetImportCollision(bCollision);
+	StaticMeshLodDataNode->SetImportCollisionType(Collision);
 	BaseNodeContainer->AddNode(StaticMeshLodDataNode);
 	return StaticMeshLodDataNode;
 }
diff --git a/Engine/Plugins/Interchange/Runtime/Source/Pipelines/Public/InterchangeGenericMeshPipeline.h b/Engine/Plugins/Interchange/Runtime/Source/Pipelines/Public/InterchangeGenericMeshPipeline.h
--- a/Engine/Plugins/Interchange/Runtime/Source/Pipelines/Public/InterchangeGenericMeshPipeline.h
+++ b/Engine/Plugins/Interchange/Runtime/Source/Pipelines/Public/InterchangeGenericMeshPipeline.h
@@ -70,8 +70,8 @@ public:
 	/** If enabled, custom collision will be imported. If enabled and there is no custom collision, a generic collision will be automatically generated.
 	 * If disabled, no collision will be created or imported.
 	 */
-	UPROPERTY(EditAnywhere, BlueprintReadWrite, Category = "Static Meshes", meta = (SubCategory = "Collision"))
-	EInterchangeMeshCollision Collision = EInterchangeMeshCollision::Convex18DOP;
+	UPROPERTY(EditAnywhere, BlueprintReadWrite, Category = "Static Meshes", meta = (DisplayName = "Collision", SubCategory = "Collision"))
+	bool bCollision = true;
 
 	/**
 	 * If enabled, meshes with certain prefixes will be imported as collision primitives for the mesh with the corresponding unprefixed name.
@@ -82,13 +82,17 @@ public:
 	 * USP_ Sphere collision
 	 * UCX_ Convex collision
 	 */
-	UPROPERTY(EditAnywhere, BlueprintReadWrite, Category = "Static Meshes", meta = (SubCategory = "Collision", editcondition = "Collision != EInterchangeMeshCollision::None"))
+	UPROPERTY(EditAnywhere, BlueprintReadWrite, Category = "Static Meshes", meta = (SubCategory = "Collision", editcondition = "bCollision == true"))
 	bool bImportCollisionAccordingToMeshName = true;
 
 	/** If enabled, each UCX collision mesh will be imported as a single convex hull. If disabled, a UCX mesh will be decomposed into its separate pieces and a convex hull generated for each. */
-	UPROPERTY(EditAnywhere, BlueprintReadWrite, Category = "Static Meshes", meta = (SubCategory = "Collision", editcondition = "Collision != EInterchangeMeshCollision::None && bImportCollisionAccordingToMeshName"))
+	UPROPERTY(EditAnywhere, BlueprintReadWrite, Category = "Static Meshes", meta = (SubCategory = "Collision", editcondition = "bCollision == true && bImportCollisionAccordingToMeshName"))
 	bool bOneConvexHullPerUCX = true;
 
+	/** Type used to generate collision when no custom collision provided. */
+	UPROPERTY(EditAnywhere, BlueprintReadWrite, Category = "Static Meshes", meta = (DisplayName = "Default Collision Type", SubCategory = "Collision", editcondition = "bCollision == true"))
+	EInterchangeMeshCollision Collision = EInterchangeMeshCollision::Convex18DOP;
+
 	//////	Static Meshes Build settings Properties //////
 
 	/** If enabled, imported meshes will be rendered by Nanite at runtime. Make sure your meshes and materials meet the requirements for Nanite. */
@@ -379,7 +383,7 @@ private:
 
 public:
 
-	virtual void PostLoad() override;
+		virtual void PostLoad() override;
 
 protected:
 	UInterchangeBaseNodeContainer* BaseNodeContainer = nullptr;
