Fix Clashing name exposed to Python

diff --git a/Engine/Plugins/Interchange/Runtime/Source/Pipelines/Public/InterchangeGenericMeshPipeline.h b/Engine/Plugins/Interchange/Runtime/Source/Pipelines/Public/InterchangeGenericMeshPipeline.h
--- a/Engine/Plugins/Interchange/Runtime/Source/Pipelines/Public/InterchangeGenericMeshPipeline.h
+++ b/Engine/Plugins/Interchange/Runtime/Source/Pipelines/Public/InterchangeGenericMeshPipeline.h
@@ -90,7 +90,7 @@ public:
 	bool bOneConvexHullPerUCX = true;
 
 	/** Type used to generate collision when no custom collision provided. */
-	UPROPERTY(EditAnywhere, BlueprintReadWrite, Category = "Static Meshes", meta = (DisplayName = "Default Collision Type", SubCategory = "Collision", editcondition = "bCollision == true"))
+	UPROPERTY(EditAnywhere, BlueprintReadWrite, Category = "Static Meshes", meta = (ScriptName = "DefaultCollisionType", DisplayName = "Default Collision Type", SubCategory = "Collision", editcondition = "bCollision == true"))
 	EInterchangeMeshCollision Collision = EInterchangeMeshCollision::Convex18DOP;
 
 	//////	Static Meshes Build settings Properties //////
