Interchange: Fix a garbage collection crash when using translator that have translator settings. When importing the translator settings are now a UPROPERTY referenced by the translator which will prevent the translator settings to be garbage collect during the import process.

#rnx

diff --git a/Engine/Plugins/Experimental/Enterprise/DatasmithInterchange/Source/DatasmithInterchange/Private/InterchangeDatasmithTranslator.cpp b/Engine/Plugins/Experimental/Enterprise/DatasmithInterchange/Source/DatasmithInterchange/Private/InterchangeDatasmithTranslator.cpp
--- a/Engine/Plugins/Experimental/Enterprise/DatasmithInterchange/Source/DatasmithInterchange/Private/InterchangeDatasmithTranslator.cpp
+++ b/Engine/Plugins/Experimental/Enterprise/DatasmithInterchange/Source/DatasmithInterchange/Private/InterchangeDatasmithTranslator.cpp
@@ -956,6 +956,7 @@ UInterchangeTranslatorSettings* UInterchangeDatasmithTranslator::GetSettings() c
 
 		CachedSettings = DuplicateObject<UInterchangeDatasmithTranslatorSettings>(UInterchangeDatasmithTranslatorSettings::StaticClass()->GetDefaultObject<UInterchangeDatasmithTranslatorSettings>(), GetTransientPackage());
 		CachedSettings->SetFlags(RF_Standalone);
+		CachedSettings->ClearFlags(RF_ArchetypeObject);
 		CachedSettings->ClearInternalFlags(EInternalObjectFlags::Async);
 
 		// Only the first one is considered
diff --git a/Engine/Plugins/Experimental/Enterprise/DatasmithInterchange/Source/DatasmithInterchange/Public/InterchangeDatasmithTranslator.h b/Engine/Plugins/Experimental/Enterprise/DatasmithInterchange/Source/DatasmithInterchange/Public/InterchangeDatasmithTranslator.h
--- a/Engine/Plugins/Experimental/Enterprise/DatasmithInterchange/Source/DatasmithInterchange/Public/InterchangeDatasmithTranslator.h
+++ b/Engine/Plugins/Experimental/Enterprise/DatasmithInterchange/Source/DatasmithInterchange/Public/InterchangeDatasmithTranslator.h
@@ -154,6 +154,8 @@ private:
 
 	mutable uint64 StartTime = 0;
 	mutable FString FileName;
+
+	UPROPERTY(Transient, DuplicateTransient)
 	mutable TObjectPtr<UInterchangeDatasmithTranslatorSettings> CachedSettings = nullptr;
 
 	mutable TMap<FString, UE::DatasmithInterchange::AnimUtils::FAnimationPayloadDesc> AnimationPayLoadMapping;
diff --git a/Engine/Plugins/Interchange/Runtime/Source/Import/Private/Fbx/InterchangeFbxTranslator.cpp b/Engine/Plugins/Interchange/Runtime/Source/Import/Private/Fbx/InterchangeFbxTranslator.cpp
--- a/Engine/Plugins/Interchange/Runtime/Source/Import/Private/Fbx/InterchangeFbxTranslator.cpp
+++ b/Engine/Plugins/Interchange/Runtime/Source/Import/Private/Fbx/InterchangeFbxTranslator.cpp
@@ -289,6 +289,8 @@ UInterchangeTranslatorSettings* UInterchangeFbxTranslator::GetSettings() const
 	{
 		CacheFbxTranslatorSettings = DuplicateObject<UInterchangeFbxTranslatorSettings>(UInterchangeFbxTranslatorSettings::StaticClass()->GetDefaultObject<UInterchangeFbxTranslatorSettings>(), GetTransientPackage());
 		CacheFbxTranslatorSettings->LoadSettings();
+		//Since we duplicate the CDO we want to remove the archetype flag
+		CacheFbxTranslatorSettings->ClearFlags(RF_ArchetypeObject);
 		CacheFbxTranslatorSettings->SetFlags(RF_Standalone);
 		CacheFbxTranslatorSettings->ClearInternalFlags(EInternalObjectFlags::Async);
 	}
diff --git a/Engine/Plugins/Interchange/Runtime/Source/Import/Private/Usd/InterchangeUsdTranslator.cpp b/Engine/Plugins/Interchange/Runtime/Source/Import/Private/Usd/InterchangeUsdTranslator.cpp
--- a/Engine/Plugins/Interchange/Runtime/Source/Import/Private/Usd/InterchangeUsdTranslator.cpp
+++ b/Engine/Plugins/Interchange/Runtime/Source/Import/Private/Usd/InterchangeUsdTranslator.cpp
@@ -129,8 +129,6 @@ namespace UE::InterchangeUsdTranslator::Private
 		// ReleaseSource() clears this member, once translation is complete.
 		UE::FUsdStage UsdStage;
 
-		TObjectPtr<UInterchangeUsdTranslatorSettings> TranslatorSettings = nullptr;
-
 #if USE_USD_SDK
 		// On UInterchangeUSDTranslator::Translate we set this up based on our TranslatorSettings, and then
 		// we can reuse it (otherwise we have to keep converting the FNames into Tokens all the time)
@@ -676,7 +674,7 @@ namespace UE::InterchangeUsdTranslator::Private
 
 	void AddMaterialInstanceNode(
 		const UE::FUsdPrim& Prim,
-		UInterchangeUSDTranslatorImpl& TranslatorImpl,
+		UInterchangeUsdTranslatorSettings* TranslatorSettings,
 		UInterchangeBaseNodeContainer& NodeContainer
 	)
 	{
@@ -696,8 +694,8 @@ namespace UE::InterchangeUsdTranslator::Private
 		NodeContainer.AddNode(MaterialNode);
 
 		UsdToUnreal::FUsdPreviewSurfaceMaterialData MaterialData;
-		FString RenderContext = TranslatorImpl.TranslatorSettings ? TranslatorImpl.TranslatorSettings->RenderContext.ToString() : FString();
-		const bool bSuccess = UsdToUnreal::ConvertMaterial(Prim, MaterialData, TranslatorImpl.TranslatorSettings ? *RenderContext : nullptr);
+		FString RenderContext = TranslatorSettings ? TranslatorSettings->RenderContext.ToString() : FString();
+		const bool bSuccess = UsdToUnreal::ConvertMaterial(Prim, MaterialData, TranslatorSettings ? *RenderContext : nullptr);
 
 		// Set all the parameter values to the interchange node
 		bool bHasUDIMTexture = false;
@@ -1565,6 +1563,7 @@ namespace UE::InterchangeUsdTranslator::Private
 		const UE::FUsdPrim& Prim,
 		UInterchangeUSDTranslatorImpl& TranslatorImpl,
 		UInterchangeBaseNodeContainer& NodeContainer,
+		UInterchangeUsdTranslatorSettings* TranslatorSettings,
 		FTraversalInfo Info
 	)
 	{
@@ -1587,7 +1586,7 @@ namespace UE::InterchangeUsdTranslator::Private
 		if (Prim.IsA(TEXT("Material")))
 		{
 			Prefix = &MaterialPrefix;
-			AddMaterialInstanceNode(Prim, TranslatorImpl, NodeContainer);
+			AddMaterialInstanceNode(Prim, TranslatorSettings, NodeContainer);
 		}
 		else if (Prim.IsA(TEXT("Mesh")))
 		{
@@ -1663,7 +1662,7 @@ namespace UE::InterchangeUsdTranslator::Private
 		// Recurse into child prims
 		for (const FUsdPrim& ChildPrim : Prim.GetChildren())
 		{
-			Traverse(ChildPrim, TranslatorImpl, NodeContainer, Info);
+			Traverse(ChildPrim, TranslatorImpl, NodeContainer, TranslatorSettings, Info);
 		}
 	}
 
@@ -2281,7 +2280,7 @@ bool UInterchangeUSDTranslator::Translate(UInterchangeBaseNodeContainer& NodeCon
 	FTraversalInfo Info;
 	for (const FUsdPrim& Prim : ImplPtr->UsdStage.GetPseudoRoot().GetChildren())
 	{
-		Traverse(Prim, *ImplPtr, NodeContainer, Info);
+		Traverse(Prim, *ImplPtr, NodeContainer, TranslatorSettings, Info);
 	}
 
 	return true;
@@ -2301,10 +2300,10 @@ void UInterchangeUSDTranslator::ReleaseSource()
 	ImplPtr->UsdStage = UE::FUsdStage{};
 	ImplPtr->CurrentTrackSet = nullptr;
 
-	if (ImplPtr->TranslatorSettings)
+	if (TranslatorSettings)
 	{
-		ImplPtr->TranslatorSettings->ClearFlags(RF_Standalone);
-		ImplPtr->TranslatorSettings = nullptr;
+		TranslatorSettings->ClearFlags(RF_Standalone);
+		TranslatorSettings = nullptr;
 	}
 }
 
@@ -2318,18 +2317,18 @@ UInterchangeTranslatorSettings* UInterchangeUSDTranslator::GetSettings() const
 		return nullptr;
 	}
 
-	TObjectPtr<UInterchangeUsdTranslatorSettings>& Settings = ImplPtr->TranslatorSettings;
-	if (!Settings)
+	if (!TranslatorSettings)
 	{
-		Settings = DuplicateObject<UInterchangeUsdTranslatorSettings>(
+		TranslatorSettings = DuplicateObject<UInterchangeUsdTranslatorSettings>(
 			UInterchangeUsdTranslatorSettings::StaticClass()->GetDefaultObject<UInterchangeUsdTranslatorSettings>(),
 			GetTransientPackage()
 		);
-		Settings->LoadSettings();
-		Settings->SetFlags(RF_Standalone);
-		Settings->ClearInternalFlags(EInternalObjectFlags::Async);
+		TranslatorSettings->LoadSettings();
+		TranslatorSettings->ClearFlags(RF_ArchetypeObject);
+		TranslatorSettings->SetFlags(RF_Standalone);
+		TranslatorSettings->ClearInternalFlags(EInternalObjectFlags::Async);
 	}
-	return Settings;
+	return TranslatorSettings;
 }
 
 void UInterchangeUSDTranslator::SetSettings(const UInterchangeTranslatorSettings* InterchangeTranslatorSettings)
@@ -2342,19 +2341,17 @@ void UInterchangeUSDTranslator::SetSettings(const UInterchangeTranslatorSettings
 		return;
 	}
 
-	TObjectPtr<UInterchangeUsdTranslatorSettings>& Settings = ImplPtr->TranslatorSettings;
-
-	if (Settings)
+	if (TranslatorSettings)
 	{
-		Settings->ClearFlags(RF_Standalone);
-		Settings->ClearInternalFlags(EInternalObjectFlags::Async);
-		Settings = nullptr;
+		TranslatorSettings->ClearFlags(RF_Standalone);
+		TranslatorSettings->ClearInternalFlags(EInternalObjectFlags::Async);
+		TranslatorSettings = nullptr;
 	}
 	if (const UInterchangeUsdTranslatorSettings* USDTranslatorSettings = Cast<UInterchangeUsdTranslatorSettings>(InterchangeTranslatorSettings))
 	{
-		Settings = DuplicateObject<UInterchangeUsdTranslatorSettings>(USDTranslatorSettings, GetTransientPackage());
-		Settings->ClearInternalFlags(EInternalObjectFlags::Async);
-		Settings->SetFlags(RF_Standalone);
+		TranslatorSettings = DuplicateObject<UInterchangeUsdTranslatorSettings>(USDTranslatorSettings, GetTransientPackage());
+		TranslatorSettings->ClearInternalFlags(EInternalObjectFlags::Async);
+		TranslatorSettings->SetFlags(RF_Standalone);
 	}
 }
 
diff --git a/Engine/Plugins/Interchange/Runtime/Source/Import/Public/Fbx/InterchangeFbxTranslator.h b/Engine/Plugins/Interchange/Runtime/Source/Import/Public/Fbx/InterchangeFbxTranslator.h
--- a/Engine/Plugins/Interchange/Runtime/Source/Import/Public/Fbx/InterchangeFbxTranslator.h
+++ b/Engine/Plugins/Interchange/Runtime/Source/Import/Public/Fbx/InterchangeFbxTranslator.h
@@ -120,6 +120,7 @@ private:
 	//never translate a source.
 	mutable TUniquePtr<UE::Interchange::FInterchangeDispatcher> Dispatcher;
 
+	UPROPERTY(Transient, DuplicateTransient)
 	mutable TObjectPtr<UInterchangeFbxTranslatorSettings> CacheFbxTranslatorSettings = nullptr;
 
 	//If true this translator will use the dispatcher (InterchangeWorker program) to translate and return payloads.
diff --git a/Engine/Plugins/Interchange/Runtime/Source/Import/Public/Usd/InterchangeUsdTranslator.h b/Engine/Plugins/Interchange/Runtime/Source/Import/Public/Usd/InterchangeUsdTranslator.h
--- a/Engine/Plugins/Interchange/Runtime/Source/Import/Public/Usd/InterchangeUsdTranslator.h
+++ b/Engine/Plugins/Interchange/Runtime/Source/Import/Public/Usd/InterchangeUsdTranslator.h
@@ -117,4 +117,7 @@ public:
 
 private:
 	TUniquePtr<UE::InterchangeUsdTranslator::Private::UInterchangeUSDTranslatorImpl> Impl;
+
+	UPROPERTY(Transient, DuplicateTransient)
+	mutable TObjectPtr<UInterchangeUsdTranslatorSettings> TranslatorSettings = nullptr;
 };
diff --git a/Engine/Source/Runtime/Interchange/Engine/Private/InterchangeAssetImportData.cpp b/Engine/Source/Runtime/Interchange/Engine/Private/InterchangeAssetImportData.cpp
--- a/Engine/Source/Runtime/Interchange/Engine/Private/InterchangeAssetImportData.cpp
+++ b/Engine/Source/Runtime/Interchange/Engine/Private/InterchangeAssetImportData.cpp
@@ -172,6 +172,7 @@ const UInterchangeTranslatorSettings* UInterchangeAssetImportData::GetTranslator
 void UInterchangeAssetImportData::SetTranslatorSettings(UInterchangeTranslatorSettings* TranslatorSettings) const
 {
 	TransientTranslatorSettings = TranslatorSettings;
+	TransientTranslatorSettings->SetFlags(RF_Standalone);
 
 	//Serialize cache
 	CachedTranslatorSettings = {};
@@ -315,6 +316,7 @@ void UInterchangeAssetImportData::ProcessTranslatorCache() const
 			UClass* ToCreateClass = ClassPerName.FindChecked(ClassFullName);
 
 			TransientTranslatorSettings = Cast<UInterchangeTranslatorSettings>(DeSerializeTranslatorSettings(CachedTranslatorSettings.Value, ToCreateClass));
+			TransientTranslatorSettings->SetFlags(RF_Standalone);
 		}
 	}
 }
diff --git a/Engine/Source/Runtime/Interchange/Engine/Private/InterchangeManager.cpp b/Engine/Source/Runtime/Interchange/Engine/Private/InterchangeManager.cpp
--- a/Engine/Source/Runtime/Interchange/Engine/Private/InterchangeManager.cpp
+++ b/Engine/Source/Runtime/Interchange/Engine/Private/InterchangeManager.cpp
@@ -688,6 +688,7 @@ void UE::Interchange::FImportAsyncHelper::CleanUp()
 			}
 		);
 		Container->ClearInternalFlags(EInternalObjectFlags::Async);
+		Container->ClearFlags(RF_Standalone);
 	}
 	BaseNodeContainers.Empty();
 
@@ -696,7 +697,7 @@ void UE::Interchange::FImportAsyncHelper::CleanUp()
 		if (SourceData)
 		{
 			SourceData->RemoveFromRoot();
-			SourceData->MarkAsGarbage();
+			SourceData->ClearFlags(RF_Standalone);
 		}
 	}
 	SourceDatas.Empty();
@@ -707,7 +708,8 @@ void UE::Interchange::FImportAsyncHelper::CleanUp()
 		{
 			Translator->ImportFinish();
 			Translator->RemoveFromRoot();
-			Translator->MarkAsGarbage();
+			Translator->ClearFlags(RF_Standalone);
+			Translator->ClearInternalFlags(EInternalObjectFlags::Async);
 		}
 	}
 	Translators.Empty();
@@ -717,7 +719,8 @@ void UE::Interchange::FImportAsyncHelper::CleanUp()
 		if(Pipeline)
 		{
 			Pipeline->RemoveFromRoot();
-			Pipeline->MarkAsGarbage();
+			Pipeline->ClearFlags(RF_Standalone);
+			Pipeline->ClearInternalFlags(EInternalObjectFlags::Async);
 		}
 	}
 	Pipelines.Empty();
@@ -727,7 +730,8 @@ void UE::Interchange::FImportAsyncHelper::CleanUp()
 		if (FactoryKeyAndValue.Value)
 		{
 			FactoryKeyAndValue.Value->RemoveFromRoot();
-			FactoryKeyAndValue.Value->MarkAsGarbage();
+			FactoryKeyAndValue.Value->ClearFlags(RF_Standalone);
+			FactoryKeyAndValue.Value->ClearInternalFlags(EInternalObjectFlags::Async);
 		}
 	}
 	CreatedFactories.Empty();
