[Interchange] FrameRate fix.

#rnx

diff --git a/Engine/Plugins/Interchange/Runtime/Source/Import/Private/Animation/InterchangeAnimSequenceFactory.cpp b/Engine/Plugins/Interchange/Runtime/Source/Import/Private/Animation/InterchangeAnimSequenceFactory.cpp
--- a/Engine/Plugins/Interchange/Runtime/Source/Import/Private/Animation/InterchangeAnimSequenceFactory.cpp
+++ b/Engine/Plugins/Interchange/Runtime/Source/Import/Private/Animation/InterchangeAnimSequenceFactory.cpp
@@ -458,7 +458,7 @@ namespace UE::Interchange::Private
 			const double SequenceLength = FMath::Max<double>(BoneTrackData.MergedRangeEnd - BoneTrackData.MergedRangeStart, MINIMUM_ANIMATION_LENGTH);
 			int32 FrameCount = FMath::RoundToInt32(SequenceLength * SampleRate);
 			int32 BakeKeyCount = FrameCount + 1;
-			const FFrameRate ResampleFrameRate(SampleRate, 1);
+			const FFrameRate ResampleFrameRate = UE::Interchange::Animation::ConvertSampleRatetoFrameRate(SampleRate);
 			Controller.SetFrameRate(ResampleFrameRate, bShouldTransact);
 			Controller.SetNumberOfFrames(FrameCount, bShouldTransact);
 
