[BUGFIX] MipGenSettings is not the same value.

diff --git a/Engine/Plugins/Interchange/Runtime/Source/Import/Private/Texture/InterchangeImageWrapperTranslator.cpp b/Engine/Plugins/Interchange/Runtime/Source/Import/Private/Texture/InterchangeImageWrapperTranslator.cpp
--- a/Engine/Plugins/Interchange/Runtime/Source/Import/Private/Texture/InterchangeImageWrapperTranslator.cpp
+++ b/Engine/Plugins/Interchange/Runtime/Source/Import/Private/Texture/InterchangeImageWrapperTranslator.cpp
@@ -222,7 +222,7 @@ TOptional<UE::Interchange::FImportImage> UInterchangeImageWrapperTranslator::Get
 				if (ImageFormat == EImageFormat::TIFF)
 				{
 					PayloadData.NumMips = DecompressedImage.MipMapImage.GetMipCount();
-					PayloadData.MipGenSettings = PayloadData.NumMips > 1 ? TextureMipGenSettings::TMGS_LeaveExistingMips : TextureMipGenSettings::TMGS_NoMipmaps;
+					PayloadData.MipGenSettings = PayloadData.NumMips > 1 ? TextureMipGenSettings::TMGS_LeaveExistingMips : TextureMipGenSettings::TMGS_FromTextureGroup;
 				}
 			}
 		}
