Interchange USD:
- Add support of MaterialX

InterchangeUSDTranslator:
- Rename AddMaterialInstanceNode into AddMaterialNode, as we now add as well materials from underlying translators
- Put AddMaterialNode and AddMeshNode free functions as member functions into the Pimpl
- Add data structures to map materials UID to mesh, since the UIDs name from the translators differ from the prim path

UsdUtil:
- Move GetMaterialXFilePaths private functions into the UsdUtils namespace as we now need as well in the InterchangeUsdTranslator.

diff --git a/Engine/Plugins/Importers/USDImporter/Source/USDSchemas/Private/Custom/MaterialXUSDShadeMaterialTranslator.cpp b/Engine/Plugins/Importers/USDImporter/Source/USDSchemas/Private/Custom/MaterialXUSDShadeMaterialTranslator.cpp
--- a/Engine/Plugins/Importers/USDImporter/Source/USDSchemas/Private/Custom/MaterialXUSDShadeMaterialTranslator.cpp
+++ b/Engine/Plugins/Importers/USDImporter/Source/USDSchemas/Private/Custom/MaterialXUSDShadeMaterialTranslator.cpp
@@ -51,44 +51,9 @@ static FAutoConsoleVariableRef CvarUseInterchangeMaterialXTranslator(
 	)
 );
 
-FName FMaterialXUsdShadeMaterialTranslator::MaterialXRenderContext = TEXT("mtlx");
-
 namespace UE::USDMaterialXTranslator::Private
 {
 #if WITH_EDITOR
-	TArray<FString> GetMaterialXFilePaths(const pxr::UsdPrim& Prim)
-	{
-		TArray<FString> Result;
-
-		// We used to just fetch "direct references" here, but stages may compose the .mtlx file reference onto the prim
-		// via another sublayer or reference, so it will be marked as an "ancestral arc" and not be included in the "direct references" filter
-		pxr::UsdPrimCompositionQuery PrimCompositionQuery = pxr::UsdPrimCompositionQuery{Prim};
-		for (const pxr::UsdPrimCompositionQueryArc& CompositionArc : PrimCompositionQuery.GetCompositionArcs())
-		{
-			if (CompositionArc.GetArcType() == pxr::PcpArcTypeReference)
-			{
-				pxr::SdfReferenceEditorProxy ReferenceEditor;
-				pxr::SdfReference UsdReference;
-
-				if (CompositionArc.GetIntroducingListEditor(&ReferenceEditor, &UsdReference))
-				{
-					// The mtlx file will be the "target layer". It's useful to get its real path via the SdfLayer
-					// interface as it already makes sure it is absolute
-					pxr::SdfLayerHandle TargetLayer = CompositionArc.GetTargetLayer();
-					FString AbsoluteFilePath = UsdToUnreal::ConvertString(TargetLayer->GetRealPath());
-
-					FString Extension = FPaths::GetExtension(AbsoluteFilePath);
-					if (Extension == FMaterialXUsdShadeMaterialTranslator::MaterialXRenderContext && FPaths::FileExists(AbsoluteFilePath))
-					{
-						Result.Add(AbsoluteFilePath);
-					}
-				}
-			}
-		}
-
-		return Result;
-	};
-
 	// Quick shallow parse of MaterialX documents ourselves so that we can find all the referenced MaterialX files and
 	// build a robust hash to use for caching our generated assets
 	FString HashMaterialXFile(const FString& MaterialXFilePath)
@@ -363,7 +328,7 @@ void FMaterialXUsdShadeMaterialTranslator::CreateAssets()
 	// this material came from: It will find an UMaterialInterface linked to that Material prim and it will happily use
 	// that as any other material.
 
-	if (Context->RenderContext != MaterialXRenderContext || !GUseInterchangeMaterialXTranslator)
+	if (Context->RenderContext != UnrealIdentifiers::MaterialXRenderContext || !GUseInterchangeMaterialXTranslator)
 	{
 		Super::CreateAssets();
 		return;
@@ -397,7 +362,7 @@ void FMaterialXUsdShadeMaterialTranslator::CreateAssets()
 	// We check for the mtlx surface output directly, because ComputeSurfaceSource will return a valid SurfaceShader
 	// in case the material just has a regular universal render context output.
 	// This is just for checking though: We will defer back to USD to let it ComputeSurfaceSource with whatever logic it has
-	const static pxr::TfToken RenderContextToken = UnrealToUsd::ConvertToken(*MaterialXRenderContext.ToString()).Get();
+	const static pxr::TfToken RenderContextToken = UnrealToUsd::ConvertToken(*UnrealIdentifiers::MaterialXRenderContext.ToString()).Get();
 	pxr::UsdShadeOutput MtlxOutput = ShadeMaterial.GetSurfaceOutput(RenderContextToken);
 	if (!MtlxOutput)
 	{
@@ -415,7 +380,7 @@ void FMaterialXUsdShadeMaterialTranslator::CreateAssets()
 	// This material prim has the mtlx render context, so maybe it is one of the ones generated by usdMtlx.
 	// Let's traverse upwards and try finding a .mtlx file reference in one of our parents.
 
-	TArray<FString> MaterialXFilePaths = UE::USDMaterialXTranslator::Private::GetMaterialXFilePaths(Prim);
+	TArray<FString> MaterialXFilePaths = UsdUtils::GetMaterialXFilePaths(Prim);
 	pxr::UsdPrim MaterialXReferencerPrim = Prim;
 	if (MaterialXFilePaths.Num() == 0)
 	{
@@ -433,7 +398,7 @@ void FMaterialXUsdShadeMaterialTranslator::CreateAssets()
 			pxr::UsdPrim MaterialXReferencerCandidate = ParentPrim;
 			while (MaterialXReferencerCandidate)
 			{
-				MaterialXFilePaths = UE::USDMaterialXTranslator::Private::GetMaterialXFilePaths(MaterialXReferencerCandidate);
+				MaterialXFilePaths = UsdUtils::GetMaterialXFilePaths(MaterialXReferencerCandidate);
 				if (MaterialXFilePaths.Num() > 0)
 				{
 					break;
diff --git a/Engine/Plugins/Importers/USDImporter/Source/USDSchemas/Private/Custom/MaterialXUSDShadeMaterialTranslator.h b/Engine/Plugins/Importers/USDImporter/Source/USDSchemas/Private/Custom/MaterialXUSDShadeMaterialTranslator.h
--- a/Engine/Plugins/Importers/USDImporter/Source/USDSchemas/Private/Custom/MaterialXUSDShadeMaterialTranslator.h
+++ b/Engine/Plugins/Importers/USDImporter/Source/USDSchemas/Private/Custom/MaterialXUSDShadeMaterialTranslator.h
@@ -14,9 +14,6 @@ class FMaterialXUsdShadeMaterialTranslator : public FUsdShadeMaterialTranslator
 {
 	using Super = FUsdShadeMaterialTranslator;
 
-public:
-	static FName MaterialXRenderContext;
-
 public:
 	using FUsdShadeMaterialTranslator::FUsdShadeMaterialTranslator;
 
diff --git a/Engine/Plugins/Importers/USDImporter/Source/USDSchemas/Private/USDSchemasModule.cpp b/Engine/Plugins/Importers/USDImporter/Source/USDSchemas/Private/USDSchemasModule.cpp
--- a/Engine/Plugins/Importers/USDImporter/Source/USDSchemas/Private/USDSchemasModule.cpp
+++ b/Engine/Plugins/Importers/USDImporter/Source/USDSchemas/Private/USDSchemasModule.cpp
@@ -25,6 +25,7 @@
 #include "USDLuxLightTranslator.h"
 #include "USDMediaSpatialAudioTranslator.h"
 #include "USDMemory.h"
+#include "USDShadeConversion.h"
 #include "USDShadeMaterialTranslator.h"
 #include "USDSkelSkeletonTranslator.h"
 #include "USDVolVolumeTranslator.h"
@@ -64,7 +65,7 @@ public:
 			Registry.Register<FUsdVolVolumeTranslator>(TEXT("UsdVolVolume"))};
 
 #if WITH_EDITOR
-		ShaderRegistry.Register(FMaterialXUsdShadeMaterialTranslator::MaterialXRenderContext);
+		ShaderRegistry.Register(UnrealIdentifiers::MaterialXRenderContext);
 		TranslatorHandles.Add(Registry.Register<FMaterialXUsdShadeMaterialTranslator>(TEXT("UsdShadeMaterial")));
 
 		// Creating skeletal meshes technically works in Standalone mode, but by checking for this we artificially block it
@@ -106,7 +107,7 @@ public:
 
 #if USE_USD_SDK && WITH_EDITOR
 		ShaderRegistry.Unregister(FMdlUsdShadeMaterialTranslator::MdlRenderContext);
-		ShaderRegistry.Unregister(FMaterialXUsdShadeMaterialTranslator::MaterialXRenderContext);
+		ShaderRegistry.Unregister(UnrealIdentifiers::MaterialXRenderContext);
 #endif	  // WITH_EDITOR
 	}
 
diff --git a/Engine/Plugins/Interchange/Runtime/Source/Import/Private/Usd/InterchangeUsdTranslator.cpp b/Engine/Plugins/Interchange/Runtime/Source/Import/Private/Usd/InterchangeUsdTranslator.cpp
--- a/Engine/Plugins/Interchange/Runtime/Source/Import/Private/Usd/InterchangeUsdTranslator.cpp
+++ b/Engine/Plugins/Interchange/Runtime/Source/Import/Private/Usd/InterchangeUsdTranslator.cpp
@@ -30,11 +30,13 @@
 
 #include "Async/Async.h"
 #include "HAL/IConsoleManager.h"
+#include "InterchangeManager.h"
 #include "InterchangeCameraNode.h"
 #include "InterchangeLightNode.h"
 #include "InterchangeMaterialInstanceNode.h"
 #include "InterchangeMeshNode.h"
 #include "InterchangeSceneNode.h"
+#include "InterchangeShaderGraphNode.h"
 #include "InterchangeTexture2DNode.h"
 #include "InterchangeTranslatorHelper.h"
 #include "MovieSceneSection.h"
@@ -124,6 +126,13 @@ namespace UE::InterchangeUsdTranslator::Private
 	// use mutable or const casts, but with the Impl we don't need to!
 	class UInterchangeUSDTranslatorImpl
 	{
+	public:
+
+		/** Add a material instance to the node container, otherwise it will add a material if it comes from a Translator (for example coming from MaterialX which cannot handle material instances) */
+		void AddMaterialNode(const UE::FUsdPrim& Prim, UInterchangeUsdTranslatorSettings* TranslatorSettings, UInterchangeBaseNodeContainer& NodeContainer);
+
+		void AddMeshNode(const UE::FUsdPrim& Prim, UInterchangeBaseNodeContainer& NodeContainer, const FTraversalInfo& Info);
+
 	public:
 		// We have to keep a stage reference so that we can parse the payloads after Translate() complete.
 		// ReleaseSource() clears this member, once translation is complete.
@@ -146,6 +155,19 @@ namespace UE::InterchangeUsdTranslator::Private
 		// For now we only generate a single LevelSequence per stage though, so we'll keep track of this
 		// here for easy access when parsing the tracks
 		UInterchangeAnimationTrackSetNode* CurrentTrackSet = nullptr;
+
+		// Array of translators that we call in the GetTexturePayload, the key has no real meaning, it's just here to avoid having duplicates and calling several times the Translate function
+		TMap<FString, UInterchangeTranslatorBase*> Translators;
+
+	private:
+
+		struct FMaterialSlotMesh
+		{
+			FString MaterialSlotName;
+			UInterchangeMeshNode* MeshNode;
+		};
+		TMap<FString, FString> PrimPathToMaterialPath; 
+		TMap<FString, TArray<FMaterialSlotMesh>> PrimPathToSlotMeshNodes;
 	};
 
 #if USE_USD_SDK
@@ -672,7 +694,7 @@ namespace UE::InterchangeUsdTranslator::Private
 		const FString* ParameterName = nullptr;
 	};
 
-	void AddMaterialInstanceNode(
+	void UInterchangeUSDTranslatorImpl::AddMaterialNode(
 		const UE::FUsdPrim& Prim,
 		UInterchangeUsdTranslatorSettings* TranslatorSettings,
 		UInterchangeBaseNodeContainer& NodeContainer
@@ -688,15 +710,87 @@ namespace UE::InterchangeUsdTranslator::Private
 			return;
 		}
 
+		auto SetMaterialSlotDependencies = [this, &NodeUid]()
+		{
+			// Now we need to check if we have to set the slot of the mesh nodes here
+			if(TArray<FMaterialSlotMesh>* SlotMeshes = PrimPathToSlotMeshNodes.Find(NodeUid))
+			{
+				if(FString* NewMaterialUID = PrimPathToMaterialPath.Find(NodeUid))
+				{
+					for(const FMaterialSlotMesh& MaterialSlotMesh : *SlotMeshes)
+					{
+						if(!MaterialSlotMesh.MeshNode->GetSlotMaterialDependencyUid(MaterialSlotMesh.MaterialSlotName, *NewMaterialUID))
+						{
+							MaterialSlotMesh.MeshNode->SetSlotMaterialDependencyUid(MaterialSlotMesh.MaterialSlotName, *NewMaterialUID);
+						}
+					}
+				}
+			}
+		};
+
+		FString RenderContext = TranslatorSettings ? TranslatorSettings->RenderContext.ToString() : FString();
+
+		// Check for any references of MaterialX
+		if(RenderContext == UnrealIdentifiers::MaterialXRenderContext)
+		{
+			TArray<FString> FilePaths = UsdUtils::GetMaterialXFilePaths(Prim);
+			for(const FString& File : FilePaths)
+			{
+				// the file has already been handled no need to do a Translate again
+				if(!Translators.Find(File))
+				{
+
+					UInterchangeManager& InterchangeManager = UInterchangeManager::GetInterchangeManager();
+					UInterchangeSourceData* SourceData = UInterchangeManager::CreateSourceData(File);
+
+					UInterchangeTranslatorBase* Translator = InterchangeManager.GetTranslatorForSourceData(SourceData);
+					//check on the Translator, it might return nullptr in case of reimport
+					if(Translator)
+					{
+						Translator->Translate(NodeContainer);
+						Translators.Add(File, Translator);
+					}
+				}
+
+				// The material from the MaterialXTranslator doesn't have the same UID, both the PrimPath have the same name but not the same path
+				// We need to retrieve that name (which is the Material name) in the translator, then we can map it to the right mesh
+				NodeContainer.BreakableIterateNodesOfType<UInterchangeShaderGraphNode>(
+					[this, &NodeName, &NodeUid](const FString&, UInterchangeShaderGraphNode* ShaderGraphNode)
+					{
+						FString UID = ShaderGraphNode->GetUniqueID();
+						if(FPaths::GetBaseFilename(UID) == NodeName)
+						{
+							PrimPathToMaterialPath.Add(NodeUid, UID);
+							return true;
+						}
+						else
+						{
+							return false;
+						}
+					});
+			}
+
+			SetMaterialSlotDependencies();
+
+			if(!FilePaths.IsEmpty())
+			{
+				return;
+			}
+		}
+		
 		UInterchangeMaterialInstanceNode* MaterialNode = NewObject<UInterchangeMaterialInstanceNode>(&NodeContainer);
 		MaterialNode->InitializeNode(NodeUid, NodeName, EInterchangeNodeContainerType::TranslatedAsset);
 		MaterialNode->SetAssetName(NodeName);
 		NodeContainer.AddNode(MaterialNode);
+		
+		// Set the material instance node to the correct mesh nodes
+		PrimPathToMaterialPath.Add(NodeUid, NodeUid);
+		SetMaterialSlotDependencies();
 
 		UsdToUnreal::FUsdPreviewSurfaceMaterialData MaterialData;
-		FString RenderContext = TranslatorSettings ? TranslatorSettings->RenderContext.ToString() : FString();
 		const bool bSuccess = UsdToUnreal::ConvertMaterial(Prim, MaterialData, TranslatorSettings ? *RenderContext : nullptr);
 
+
 		// Set all the parameter values to the interchange node
 		bool bHasUDIMTexture = false;
 		FParameterValueVisitor Visitor{Prim, NodeContainer, *MaterialNode, MaterialData.PrimvarToUVIndex};
@@ -1007,9 +1101,8 @@ namespace UE::InterchangeUsdTranslator::Private
 		}
 	}
 
-	void AddMeshNode(
+	void UInterchangeUSDTranslatorImpl::AddMeshNode(
 		const UE::FUsdPrim& Prim,
-		UInterchangeUSDTranslatorImpl& TranslatorImpl,
 		UInterchangeBaseNodeContainer& NodeContainer,
 		const FTraversalInfo& Info
 	)
@@ -1038,13 +1131,13 @@ namespace UE::InterchangeUsdTranslator::Private
 				MeshNode->SetSkeletonDependencyUid(Info.ActiveSkelQuery.GetSkeleton().GetPrimPath().GetString());
 			}
 
-			AddMorphTargetNodes(Prim, TranslatorImpl, *MeshNode, NodeContainer, Info);
+			AddMorphTargetNodes(Prim, *this, *MeshNode, NodeContainer, Info);
 
 			// When returning the payload data later, we'll need at the very least our SkeletonQuery, so
 			// here we store the Info object into the Impl
 			{
-				FWriteScopeLock ScopedInfoWriteLock{TranslatorImpl.CachedTraversalInfoLock};
-				TranslatorImpl.NodeUidToCachedTraversalInfo.Add(NodeUid, Info);
+				FWriteScopeLock ScopedInfoWriteLock{CachedTraversalInfoLock};
+				NodeUidToCachedTraversalInfo.Add(NodeUid, Info);
 			}
 		}
 		else
@@ -1060,8 +1153,8 @@ namespace UE::InterchangeUsdTranslator::Private
 				Prim,
 				TimeCode,
 				bProvideMaterialIndices,
-				TranslatorImpl.CachedMeshConversionOptions.RenderContext,
-				TranslatorImpl.CachedMeshConversionOptions.MaterialPurpose
+				CachedMeshConversionOptions.RenderContext,
+				CachedMeshConversionOptions.MaterialPurpose
 			);
 
 			for (const UsdUtils::FUsdPrimMaterialSlot& Slot : Assignments.Slots)
@@ -1102,7 +1195,23 @@ namespace UE::InterchangeUsdTranslator::Private
 					}
 				}
 
-				MeshNode->SetSlotMaterialDependencyUid(SlotName, MaterialInstanceUid);
+				// if we found a match let's set the slot to the corresponding Material (the Material path could come from another translator)
+				if(FString* NewMaterialInstanceUid = PrimPathToMaterialPath.Find(MaterialInstanceUid))
+				{
+					MeshNode->SetSlotMaterialDependencyUid(SlotName, *NewMaterialInstanceUid);
+				}
+				else // otherwise it's up to the Material to attach itself to the mesh
+				{
+					// one material can be attached to several meshes
+					if(TArray<FMaterialSlotMesh>* SlotsMeshes = PrimPathToSlotMeshNodes.Find(MaterialInstanceUid))
+					{
+						SlotsMeshes->Add({ SlotName, MeshNode });
+					}
+					else
+					{
+						PrimPathToSlotMeshNodes.Add(MaterialInstanceUid, { { SlotName, MeshNode } });
+					}
+				}
 			}
 		}
 
@@ -1586,12 +1695,12 @@ namespace UE::InterchangeUsdTranslator::Private
 		if (Prim.IsA(TEXT("Material")))
 		{
 			Prefix = &MaterialPrefix;
-			AddMaterialInstanceNode(Prim, TranslatorSettings, NodeContainer);
+			TranslatorImpl.AddMaterialNode(Prim, TranslatorSettings, NodeContainer);
 		}
 		else if (Prim.IsA(TEXT("Mesh")))
 		{
 			Prefix = &MeshPrefix;
-			AddMeshNode(Prim, TranslatorImpl, NodeContainer, Info);
+			TranslatorImpl.AddMeshNode(Prim, NodeContainer, Info);
 		}
 		else if (Prim.IsA(TEXT("Camera")))
 		{
@@ -2443,33 +2552,45 @@ TOptional<UE::Interchange::FImportImage> UInterchangeUSDTranslator::GetTexturePa
 	FString FilePath;
 	TextureGroup TextureGroup;
 	bool bDecoded = DecodeTexturePayloadKey(PayloadKey, FilePath, TextureGroup);
-	if (!bDecoded)
+	if(bDecoded)
 	{
-		return {};
-	}
-
-	// Defer back to another translator to actually parse the texture raw data
-	UE::Interchange::Private::FScopedTranslator ScopedTranslator(FilePath, Results);
-	const IInterchangeTexturePayloadInterface* TextureTranslator = ScopedTranslator.GetPayLoadInterface<IInterchangeTexturePayloadInterface>();
-	if (!ensure(TextureTranslator))
-	{
-		return {};
-	}
+		// Defer back to another translator to actually parse the texture raw data
+		UE::Interchange::Private::FScopedTranslator ScopedTranslator(FilePath, Results);
+		const IInterchangeTexturePayloadInterface* TextureTranslator = ScopedTranslator.GetPayLoadInterface<IInterchangeTexturePayloadInterface>();
+		if(ensure(TextureTranslator))
+		{
 
-	AlternateTexturePath = FilePath;
+			AlternateTexturePath = FilePath;
 
-	// The texture translators don't use the payload key, and read the texture directly from the SourceData's file path
-	const FString UnusedPayloadKey = {};
-	TexturePayloadData = TextureTranslator->GetTexturePayloadData(UnusedPayloadKey, AlternateTexturePath);
+			// The texture translators don't use the payload key, and read the texture directly from the SourceData's file path
+			const FString UnusedPayloadKey = {};
+			TexturePayloadData = TextureTranslator->GetTexturePayloadData(UnusedPayloadKey, AlternateTexturePath);
 
-	// Move compression settings onto the payload data.
-	// Note: We don't author anything else on the texture payload data here (like the sRGB flag), because those
-	// settings were already on our translated node, and presumably already made their way to the factory node.
-	// The factory should use them to override whatever it finds in this payload data, with the exception of the
-	// compression settings (which can't be stored on the translated node)
-	TexturePayloadData->CompressionSettings = TextureGroup == TEXTUREGROUP_WorldNormalMap ? TC_Normalmap : TC_Default;
+			// Move compression settings onto the payload data.
+			// Note: We don't author anything else on the texture payload data here (like the sRGB flag), because those
+			// settings were already on our translated node, and presumably already made their way to the factory node.
+			// The factory should use them to override whatever it finds in this payload data, with the exception of the
+			// compression settings (which can't be stored on the translated node)
+			TexturePayloadData->CompressionSettings = TextureGroup == TEXTUREGROUP_WorldNormalMap ? TC_Normalmap : TC_Default;
+		}
+	}
 #endif	  // USE_USD_SDK
 
+	// We did not find a suitable Payload in USD Translator, let's find one in one of the Translators (MaterialX for the moment)
+	// The best way would be to have a direct association between the payload and the right Translator, but we don't have a suitable way of knowing which Payload belongs to which Translator
+	// So let's just loop over them all
+	for(const TPair<FString, UInterchangeTranslatorBase*> & Pair : Impl->Translators)
+	{
+		if(IInterchangeTexturePayloadInterface* TexturePayloadInterface = Cast<IInterchangeTexturePayloadInterface>(Pair.Value))
+		{
+			TexturePayloadData = TexturePayloadInterface->GetTexturePayloadData(PayloadKey, AlternateTexturePath);
+			if(TexturePayloadData)
+			{
+				break;
+			}
+		}
+	}
+
 	return TexturePayloadData;
 }
 
diff --git a/Engine/Plugins/Runtime/USDCore/Source/USDUtilities/Private/USDShadeConversion.cpp b/Engine/Plugins/Runtime/USDCore/Source/USDUtilities/Private/USDShadeConversion.cpp
--- a/Engine/Plugins/Runtime/USDCore/Source/USDUtilities/Private/USDShadeConversion.cpp
+++ b/Engine/Plugins/Runtime/USDCore/Source/USDUtilities/Private/USDShadeConversion.cpp
@@ -56,9 +56,11 @@
 #include "pxr/usd/ar/resolver.h"
 #include "pxr/usd/ar/resolverScopedCache.h"
 #include "pxr/usd/sdf/layerUtils.h"
+#include "pxr/usd/usd/primCompositionQuery.h"
 #include "pxr/usd/usdShade/material.h"
 #include "pxr/usd/usdShade/materialBindingAPI.h"
 #include "pxr/usd/usdShade/nodeDefAPI.h"
+
 #include "USDIncludesEnd.h"
 
 // Required for packaging
@@ -3035,6 +3037,43 @@ void UsdUtils::SetBoolParameterValue(UMaterialInstance& Material, const TCHAR* P
 	}
 }
 
+#if WITH_EDITOR
+TArray<FString> UsdUtils::GetMaterialXFilePaths(const pxr::UsdPrim& Prim)
+{
+	FScopedUsdAllocs UsdAllocs;
+
+	TArray<FString> Result;
+
+	// We used to just fetch "direct references" here, but stages may compose the .mtlx file reference onto the prim
+	// via another sublayer or reference, so it will be marked as an "ancestral arc" and not be included in the "direct references" filter
+	pxr::UsdPrimCompositionQuery PrimCompositionQuery = pxr::UsdPrimCompositionQuery{ Prim };
+	for(const pxr::UsdPrimCompositionQueryArc& CompositionArc : PrimCompositionQuery.GetCompositionArcs())
+	{
+		if(CompositionArc.GetArcType() == pxr::PcpArcTypeReference)
+		{
+			pxr::SdfReferenceEditorProxy ReferenceEditor;
+			pxr::SdfReference UsdReference;
+
+			if(CompositionArc.GetIntroducingListEditor(&ReferenceEditor, &UsdReference))
+			{
+				// The mtlx file will be the "target layer". It's useful to get its real path via the SdfLayer
+				// interface as it already makes sure it is absolute
+				pxr::SdfLayerHandle TargetLayer = CompositionArc.GetTargetLayer();
+				FString AbsoluteFilePath = UsdToUnreal::ConvertString(TargetLayer->GetRealPath());
+
+				FString Extension = FPaths::GetExtension(AbsoluteFilePath);
+				if(Extension == UnrealIdentifiers::MaterialXRenderContext && FPaths::FileExists(AbsoluteFilePath))
+				{
+					Result.Add(AbsoluteFilePath);
+				}
+			}
+		}
+	}
+
+	return Result;
+}
+#endif // WITH_EDITOR
+
 void UsdUtils::AuthorUnrealMaterialBinding(pxr::UsdPrim& MeshOrGeomSubsetPrim, const FString& UnrealMaterialPathName)
 {
 	if (!MeshOrGeomSubsetPrim || UnrealMaterialPathName.IsEmpty())
diff --git a/Engine/Plugins/Runtime/USDCore/Source/USDUtilities/Public/USDShadeConversion.h b/Engine/Plugins/Runtime/USDCore/Source/USDUtilities/Public/USDShadeConversion.h
--- a/Engine/Plugins/Runtime/USDCore/Source/USDUtilities/Public/USDShadeConversion.h
+++ b/Engine/Plugins/Runtime/USDCore/Source/USDUtilities/Public/USDShadeConversion.h
@@ -347,6 +347,11 @@ namespace UsdUtils
 	USDUTILITIES_API void SetVectorParameterValue(UMaterialInstance& Material, const TCHAR* ParameterName, FLinearColor ParameterValue);
 	USDUTILITIES_API void SetTextureParameterValue(UMaterialInstance& Material, const TCHAR* ParameterName, UTexture* ParameterValue);
 	USDUTILITIES_API void SetBoolParameterValue(UMaterialInstance& Material, const TCHAR* ParameterName, bool bParameterValue);
+
+#if WITH_EDITOR
+	/** Retrieve MaterialX file from a prim*/
+	USDUTILITIES_API TArray<FString> GetMaterialXFilePaths(const pxr::UsdPrim& Prim);
+#endif // WITH_EDITOR
 }	 // namespace UsdUtils
 
 #endif	  // #if USE_USD_SDK
diff --git a/Engine/Plugins/Runtime/USDCore/Source/UnrealUSDWrapper/Private/UnrealUSDWrapper.cpp b/Engine/Plugins/Runtime/USDCore/Source/UnrealUSDWrapper/Private/UnrealUSDWrapper.cpp
--- a/Engine/Plugins/Runtime/USDCore/Source/UnrealUSDWrapper/Private/UnrealUSDWrapper.cpp
+++ b/Engine/Plugins/Runtime/USDCore/Source/UnrealUSDWrapper/Private/UnrealUSDWrapper.cpp
@@ -475,6 +475,10 @@ FString UnrealIdentifiers::ModelApplyDrawMode = TEXT("model:applyDrawMode");
 FString UnrealIdentifiers::UsdNamespaceDelimiter = TEXT(":");
 #endif	  // USE_USD_SDK
 
+#if WITH_EDITOR
+const FName UnrealIdentifiers::MaterialXRenderContext = TEXT("mtlx");
+#endif
+
 FUsdDelegates::FUsdImportDelegate FUsdDelegates::OnPreUsdImport;
 FUsdDelegates::FUsdImportDelegate FUsdDelegates::OnPostUsdImport;
 
diff --git a/Engine/Plugins/Runtime/USDCore/Source/UnrealUSDWrapper/Public/UnrealUSDWrapper.h b/Engine/Plugins/Runtime/USDCore/Source/UnrealUSDWrapper/Public/UnrealUSDWrapper.h
--- a/Engine/Plugins/Runtime/USDCore/Source/UnrealUSDWrapper/Public/UnrealUSDWrapper.h
+++ b/Engine/Plugins/Runtime/USDCore/Source/UnrealUSDWrapper/Public/UnrealUSDWrapper.h
@@ -552,6 +552,11 @@ namespace UnrealIdentifiers
 
 	// The character used to separate property namespaces (usually just ':')
 	extern UNREALUSDWRAPPER_API FString UsdNamespaceDelimiter;
+
+#if WITH_EDITOR
+	extern UNREALUSDWRAPPER_API const FName MaterialXRenderContext;
+#endif
+
 }	 // namespace UnrealIdentifiers
 
 struct UNREALUSDWRAPPER_API FUsdDelegates
