[Interchange] FrameRate precision fix for Rigid Animations.

#rnx

diff --git a/Engine/Plugins/Interchange/Runtime/Source/Parsers/Fbx/Private/FbxScene.cpp b/Engine/Plugins/Interchange/Runtime/Source/Parsers/Fbx/Private/FbxScene.cpp
--- a/Engine/Plugins/Interchange/Runtime/Source/Parsers/Fbx/Private/FbxScene.cpp
+++ b/Engine/Plugins/Interchange/Runtime/Source/Parsers/Fbx/Private/FbxScene.cpp
@@ -909,6 +909,9 @@ namespace UE
 				{
 					UInterchangeAnimationTrackSetNode* TrackSetNode = NewObject< UInterchangeAnimationTrackSetNode >(&NodeContainer);
 
+					double FrameRate = FbxTime::GetFrameRate(SDKScene->GetGlobalSettings().GetTimeMode());
+					TrackSetNode->SetCustomFrameRate(FrameRate);
+
 					const FString AnimTrackSetNodeUid = TEXT("\\Animation\\") + FString(RootNode->GetName());
 					const FString AnimTrackSetNodeDisplayLabel = FString(RootNode->GetName()) + TEXT("_TrackSetNode");
 					TrackSetNode->InitializeNode(AnimTrackSetNodeUid, AnimTrackSetNodeDisplayLabel, EInterchangeNodeContainerType::TranslatedAsset);
diff --git a/Engine/Plugins/Interchange/Runtime/Source/Pipelines/Private/InterchangeGenericAnimationPipeline.cpp b/Engine/Plugins/Interchange/Runtime/Source/Pipelines/Private/InterchangeGenericAnimationPipeline.cpp
--- a/Engine/Plugins/Interchange/Runtime/Source/Pipelines/Private/InterchangeGenericAnimationPipeline.cpp
+++ b/Engine/Plugins/Interchange/Runtime/Source/Pipelines/Private/InterchangeGenericAnimationPipeline.cpp
@@ -339,7 +339,8 @@ void UInterchangeGenericAnimationPipeline::ExecutePipeline(UInterchangeBaseNodeC
 			SkeletalAnimationNode->SetCustomAnimationSampleRate(CustomFrameRate);
 			SkeletalAnimationNode->SetCustomAnimationStartTime(0);
 			//stop time will be calculated once we get the curves from the Translators.
-			SkeletalAnimationNode->SetCustomAnimationStopTime(0);
+			//However to avoid error reports we set stoptime for 1 subframe.
+			SkeletalAnimationNode->SetCustomAnimationStopTime(1./CustomFrameRate);
 
 			for (const FString& AnimationTrackUid : AnimationTrackUids)
 			{
