Fix wrong UsdPreviewSurface path

diff --git a/Engine/Plugins/Interchange/Runtime/Source/Pipelines/Private/InterchangeMaterialXPipeline.cpp b/Engine/Plugins/Interchange/Runtime/Source/Pipelines/Private/InterchangeMaterialXPipeline.cpp
--- a/Engine/Plugins/Interchange/Runtime/Source/Pipelines/Private/InterchangeMaterialXPipeline.cpp
+++ b/Engine/Plugins/Interchange/Runtime/Source/Pipelines/Private/InterchangeMaterialXPipeline.cpp
@@ -31,7 +31,7 @@ namespace
 	constexpr const TCHAR* TransmissionSurfaceSubstratePath = TEXT("/Engine/Functions/Substrate/Substrate-StandardSurface-Translucent.Substrate-StandardSurface-Translucent");
 	MATERIALX_FUNCTIONS_SUBSTRATE_PATH(SurfaceUnlit);
 	MATERIALX_FUNCTIONS_SUBSTRATE_PATH(Surface);
-	constexpr const TCHAR* UsdPreviewSurfaceFunctionsPath = TEXT("/Interchange/Functions/MX_UsdPreviewSurface.MX_UsdPreviewSurface");
+	constexpr const TCHAR* UsdPreviewSurfaceFunctionsPath = TEXT("/InterchangeAssets/Functions/MX_UsdPreviewSurface.MX_UsdPreviewSurface");
 	constexpr const TCHAR* UsdPreviewSurfaceSubstratePath = TEXT("/Engine/Functions/Substrate/MF_Substrate_UsdPreviewSurface.MF_Substrate_UsdPreviewSurface");
 	
 	MATERIALX_FUNCTIONS_SUBSTRATE_PATH(OrenNayarBSDF);
