Fix crash when loading a level and some file are not saved

#rnx

diff --git a/Engine/Plugins/Interchange/Runtime/Source/Pipelines/Private/InterchangeGenericScenesPipeline.cpp b/Engine/Plugins/Interchange/Runtime/Source/Pipelines/Private/InterchangeGenericScenesPipeline.cpp
--- a/Engine/Plugins/Interchange/Runtime/Source/Pipelines/Private/InterchangeGenericScenesPipeline.cpp
+++ b/Engine/Plugins/Interchange/Runtime/Source/Pipelines/Private/InterchangeGenericScenesPipeline.cpp
@@ -1127,11 +1127,11 @@ void UInterchangeGenericLevelPipeline::ExecutePostImportPipeline(const UIntercha
 	}
 	else if (ALevelInstance* LevelInstanceActor = Cast<ALevelInstance>(CreatedAsset))
 	{
-		const UInterchangeLevelInstanceActorFactoryNode* LevelInstanceFactoryNode = Cast<UInterchangeLevelInstanceActorFactoryNode>(BaseNodeContainer->GetFactoryNode(NodeKey));
+		const UInterchangeLevelInstanceActorFactoryNode* LevelInstanceFactoryNode = Cast<UInterchangeLevelInstanceActorFactoryNode>(InBaseNodeContainer->GetFactoryNode(NodeKey));
 		FString LevelFactoryNodeUid;
 		if (LevelInstanceFactoryNode->GetCustomLevelReference(LevelFactoryNodeUid))
 		{
-			if (const UInterchangeFactoryBaseNode* ReferenceLevelFactoryNode = BaseNodeContainer->GetFactoryNode(LevelFactoryNodeUid))
+			if (const UInterchangeFactoryBaseNode* ReferenceLevelFactoryNode = InBaseNodeContainer->GetFactoryNode(LevelFactoryNodeUid))
 			{
 				FSoftObjectPath ReferenceLevelPath;
 				if (ReferenceLevelFactoryNode->GetCustomReferenceObject(ReferenceLevelPath))
@@ -1147,6 +1147,25 @@ void UInterchangeGenericLevelPipeline::ExecutePostImportPipeline(const UIntercha
 #endif
 }
 
+void UInterchangeGenericLevelPipeline::ExecutePostBroadcastPipeline(const UInterchangeBaseNodeContainer* InBaseNodeContainer, const FString& NodeKey, UObject* CreatedAsset, bool bIsAReimport)
+{
+	Super::ExecutePostBroadcastPipeline(InBaseNodeContainer, NodeKey, CreatedAsset, bIsAReimport);
+	
+	if (!CreatedAsset || !ensure(IsInGameThread()))
+	{
+		return;
+	}
+
+#if WITH_EDITORONLY_DATA
+	if (ALevelInstance* LevelInstanceActor = Cast<ALevelInstance>(CreatedAsset))
+	{
+		//Commit an edit to properly finish the setup for the new Level instance actor and new level.
+		LevelInstanceActor->EnterEdit();
+		LevelInstanceActor->ExitEdit();
+	}
+#endif
+}
+
 #if WITH_EDITORONLY_DATA
 
 void UInterchangeGenericLevelPipeline::FPostPipelineImportData::AddWorld(UWorld* World)
@@ -1191,9 +1210,9 @@ bool UInterchangeGenericLevelPipeline::FPostPipelineImportData::UpdateLevelInsta
 		return false;
 	}
 
-	bool& bWorldSaved = Worlds.FindChecked(ReferencedWorld);
+	bool& bWorldUpdate = Worlds.FindChecked(ReferencedWorld);
 	
-	if (!bWorldSaved)
+	if (!bWorldUpdate)
 	{
 		if (UInterchangeEditorUtilitiesBase* EditorUtilities = UInterchangeManager::GetInterchangeManager().GetEditorUtilities())
 		{
@@ -1202,26 +1221,27 @@ bool UInterchangeGenericLevelPipeline::FPostPipelineImportData::UpdateLevelInsta
 				UE_LOG(LogInterchangePipeline, Warning, TEXT("UInterchangeGenericAssetsPipeline: Cannot save the level instance actor (%s) referenced world (%s)"), *LevelInstanceActor->GetName(), *ReferencedWorld->GetName());
 			}
 		}
-		bWorldSaved = true;
-	}
 
-	ParentWorld->PreEditChange(nullptr);
+		// Make sure newly created level asset gets scanned
+		ULevel::ScanLevelAssets(ReferencedWorld->GetPackage()->GetName());
 
-	LevelInstanceActor->SetWorldAsset(ReferencedWorld);
-	LevelInstanceActor->UpdateLevelInstanceFromWorldAsset();
-	LevelInstanceActor->LoadLevelInstance();
-	
-	//A level instance referenced world cannot be RF_Standalone
-	ReferencedWorld->ClearFlags(RF_Standalone);
+		ParentWorld->PreEditChange(nullptr);
 
-	//Reference world must be cleanup since they are not the main world.
-	//This remove all the world managers and prevent GC issue when unloading the main world referencing this world.
-	if (ReferencedWorld->bIsWorldInitialized)
-	{
-		ReferencedWorld->CleanupWorld();
-	}
+		LevelInstanceActor->SetWorldAsset(ReferencedWorld);
+		LevelInstanceActor->UpdateLevelInstanceFromWorldAsset();
+		LevelInstanceActor->LoadLevelInstance();
 
-	ParentWorld->PostEditChange();
+		//Reference world must be cleanup since they are not the main world.
+		//This remove all the world managers and prevent GC issue when unloading the main world referencing this world.
+		if (ReferencedWorld->bIsWorldInitialized)
+		{
+			ReferencedWorld->CleanupWorld();
+		}
+
+		ParentWorld->PostEditChange();
+
+		bWorldUpdate = true;
+	}
 
 	return true;
 }
diff --git a/Engine/Plugins/Interchange/Runtime/Source/Pipelines/Public/InterchangeGenericScenesPipeline.h b/Engine/Plugins/Interchange/Runtime/Source/Pipelines/Public/InterchangeGenericScenesPipeline.h
--- a/Engine/Plugins/Interchange/Runtime/Source/Pipelines/Public/InterchangeGenericScenesPipeline.h
+++ b/Engine/Plugins/Interchange/Runtime/Source/Pipelines/Public/InterchangeGenericScenesPipeline.h
@@ -67,6 +67,7 @@ protected:
 
 	virtual void ExecutePipeline(UInterchangeBaseNodeContainer* InBaseNodeContainer, const TArray<UInterchangeSourceData*>& InSourceDatas, const FString& ContentBasePath) override;
 	virtual void ExecutePostImportPipeline(const UInterchangeBaseNodeContainer* BaseNodeContainer, const FString& NodeKey, UObject* CreatedAsset, bool bIsAReimport) override;
+	virtual void ExecutePostBroadcastPipeline(const UInterchangeBaseNodeContainer* BaseNodeContainer, const FString& NodeKey, UObject* CreatedAsset, bool bIsAReimport) override;
 
 	virtual bool CanExecuteOnAnyThread(EInterchangePipelineTask PipelineTask) override
 	{
@@ -119,7 +120,7 @@ protected:
 	private:
 		bool UpdateLevelInstanceInternal(ALevelInstance* LevelInstanceActor, UWorld* ReferenceWorld);
 
-		//The bool is to know if we have save this world already or not
+		//The bool is to know if we have already update this world
 		TMap<UWorld*, bool> Worlds;
 		TMap<ALevelInstance*, UWorld*> ReferenceWorldPerLevelInstanceToUpdates;
 	};
diff --git a/Engine/Source/Runtime/Interchange/Core/Public/InterchangePipelineBase.h b/Engine/Source/Runtime/Interchange/Core/Public/InterchangePipelineBase.h
--- a/Engine/Source/Runtime/Interchange/Core/Public/InterchangePipelineBase.h
+++ b/Engine/Source/Runtime/Interchange/Core/Public/InterchangePipelineBase.h
@@ -190,6 +190,15 @@ protected:
 	TSharedPtr<SWindow> WidgetWindow = nullptr;
 };
 
+/**
+ * Pipeline implementation:
+ *
+ * 1. ExecutePipeline - Create the factory nodes from the translated nodes. This is where the logic is execute to create the unreal asset via the factory node. Called after the translation
+ * 2. ExecutePostFactoryPipeline - Called after the factory has create the unreal asset with the associate factory node, but before calling PostEditChange.
+ * 3. ExecutePostImportPipeline - Called after the asset PostEditChange is done. If the asset use the async build framework, the asset build should be completed.
+ * 4. ExecutePostBroadcastPipeline - Called after the asset was registered to the registry manager and all broadcast calls have been done.
+ */
+
 UCLASS(BlueprintType, Blueprintable, editinlinenew, Abstract, MinimalAPI)
 class UInterchangePipelineBase : public UObject
 {
@@ -262,6 +271,20 @@ public:
 		ExecutePostImportPipeline(BaseNodeContainer, FactoryNodeKey, CreatedAsset, bIsAReimport);
 	}
 
+	/**
+	 * ScriptedExecutePostBroadcastPipeline is called after an asset is completely imported and the broadcast have been called.
+	 * This can be useful if you need to unload the asset for any reason (Level reference by level instance need to be unload).
+	 * @note - the FTaskCompletion_GameThread calls this function not the virtual one that is call by the default implementation.
+	 */
+	UFUNCTION(BlueprintNativeEvent, BlueprintCallable, Category = "Interchange | Pipeline")
+	INTERCHANGECORE_API void ScriptedExecutePostBroadcastPipeline(const UInterchangeBaseNodeContainer* BaseNodeContainer, const FString& FactoryNodeKey, UObject* CreatedAsset, bool bIsAReimport);
+	/** The default implementation, which is called if the Blueprint does not have any implementation, calls the virtual ExecutePostBroadcastPipeline(). */
+	void ScriptedExecutePostBroadcastPipeline_Implementation(const UInterchangeBaseNodeContainer* BaseNodeContainer, const FString& FactoryNodeKey, UObject* CreatedAsset, bool bIsAReimport)
+	{
+		//By default we call the virtual import pipeline execution
+		ExecutePostBroadcastPipeline(BaseNodeContainer, FactoryNodeKey, CreatedAsset, bIsAReimport);
+	}
+
 	/**
 	 * Non-virtual helper that allows Blueprint to implement an event-based function.
 	 * The Interchange manager calls this function, not the virtual one that is called by the default implementation.
@@ -551,6 +574,14 @@ protected:
 	{
 	}
 
+	/**
+	 * This function is called after the Unreal asset is completely imported. PostEditChange and all broadcast have been called.
+	 * @Note: Some Unreal assets have asynchronous build operations. It's possible they are still compiling.
+	 */
+	virtual void ExecutePostBroadcastPipeline(const UInterchangeBaseNodeContainer* BaseNodeContainer, const FString& NodeKey, UObject* CreatedAsset, bool bIsAReimport)
+	{
+	}
+
 	virtual void SetReimportSourceIndex(UClass* ReimportObjectClass, const int32 SourceFileIndex)
 	{
 	}
diff --git a/Engine/Source/Runtime/Interchange/Engine/Private/Tasks/InterchangeTaskCompletion.cpp b/Engine/Source/Runtime/Interchange/Engine/Private/Tasks/InterchangeTaskCompletion.cpp
--- a/Engine/Source/Runtime/Interchange/Engine/Private/Tasks/InterchangeTaskCompletion.cpp
+++ b/Engine/Source/Runtime/Interchange/Engine/Private/Tasks/InterchangeTaskCompletion.cpp
@@ -233,8 +233,43 @@ void UE::Interchange::FTaskCompletion_GameThread::Execute()
 					}
 				}
 
+				//Do a second pass for the post broadcast pipeline call
+				for (const FImportAsyncHelper::FImportedObjectInfo& AssetInfo : AssetInfos)
+				{
+					if (UObject* Asset = AssetInfo.ImportedObject)
+					{
+						for (UInterchangePipelineBase* PipelineBase : AsyncHelper->Pipelines)
+						{
+							PipelineBase->ScriptedExecutePostBroadcastPipeline(AsyncHelper->BaseNodeContainers[SourceIndex].Get()
+								, AssetInfo.FactoryNode ? AssetInfo.FactoryNode->GetUniqueID() : FString()
+								, Asset
+								, AssetInfo.bIsReimport);
+						}
+					}
+				}
+
 				UE_LOG(LogInterchangeEngine, Display, TEXT("Interchange import completed [%s]"), *AsyncHelper->SourceDatas[SourceIndex]->ToDisplayString());
 			});
+
+		//Iterate the Scene Actors
+		AsyncHelper->IterateImportedSceneObjectsPerSourceIndex([AsyncHelper](int32 SourceIndex, const TArray<FImportAsyncHelper::FImportedObjectInfo>& AssetInfos)
+			{
+				for (const FImportAsyncHelper::FImportedObjectInfo& SceneObjectInfo : AssetInfos)
+				{
+					if (AActor* Actor = Cast<AActor>(SceneObjectInfo.ImportedObject))
+					{
+						for (UInterchangePipelineBase* PipelineBase : AsyncHelper->Pipelines)
+						{
+							PipelineBase->ScriptedExecutePostBroadcastPipeline(AsyncHelper->BaseNodeContainers[SourceIndex].Get()
+								, SceneObjectInfo.FactoryNode ? SceneObjectInfo.FactoryNode->GetUniqueID() : FString()
+								, Actor
+								, SceneObjectInfo.bIsReimport);
+						}
+					}
+				}
+
+				UE_LOG(LogInterchangeEngine, Display, TEXT("Interchange import cancelled [%s]"), *AsyncHelper->SourceDatas[SourceIndex]->ToDisplayString());
+			});
 	}
 	else
 	{
