Surround MaterialXRenderContext by a WITH_EDITOR

diff --git a/Engine/Plugins/Interchange/Runtime/Source/Import/Private/Usd/InterchangeUsdTranslator.cpp b/Engine/Plugins/Interchange/Runtime/Source/Import/Private/Usd/InterchangeUsdTranslator.cpp
--- a/Engine/Plugins/Interchange/Runtime/Source/Import/Private/Usd/InterchangeUsdTranslator.cpp
+++ b/Engine/Plugins/Interchange/Runtime/Source/Import/Private/Usd/InterchangeUsdTranslator.cpp
@@ -731,6 +731,7 @@ namespace UE::InterchangeUsdTranslator::Private
 		FString RenderContext = TranslatorSettings ? TranslatorSettings->RenderContext.ToString() : FString();
 
 		// Check for any references of MaterialX
+#if WITH_EDITOR
 		if(RenderContext == UnrealIdentifiers::MaterialXRenderContext)
 		{
 			TArray<FString> FilePaths = UsdUtils::GetMaterialXFilePaths(Prim);
@@ -777,6 +778,7 @@ namespace UE::InterchangeUsdTranslator::Private
 				return;
 			}
 		}
+#endif // WITH_EDITOR
 		
 		UInterchangeMaterialInstanceNode* MaterialNode = NewObject<UInterchangeMaterialInstanceNode>(&NodeContainer);
 		MaterialNode->InitializeNode(NodeUid, NodeName, EInterchangeNodeContainerType::TranslatedAsset);
