[Interchange] Actor Layer/Tag support.

#rnx

diff --git a/Engine/Plugins/Interchange/Runtime/Source/FactoryNodes/Private/InterchangeActorFactoryNode.cpp b/Engine/Plugins/Interchange/Runtime/Source/FactoryNodes/Private/InterchangeActorFactoryNode.cpp
--- a/Engine/Plugins/Interchange/Runtime/Source/FactoryNodes/Private/InterchangeActorFactoryNode.cpp
+++ b/Engine/Plugins/Interchange/Runtime/Source/FactoryNodes/Private/InterchangeActorFactoryNode.cpp
@@ -96,3 +96,60 @@ void UInterchangeActorFactoryNode::CopyWithObject(const UInterchangeFactoryBaseN
 	}
 }
 
+UInterchangeActorFactoryNode::UInterchangeActorFactoryNode()
+{
+	LayerNames.Initialize(Attributes.ToSharedRef(), TEXT("__LayerNames__"));
+	Tags.Initialize(Attributes.ToSharedRef(), TEXT("__Tags__"));
+}
+
+void UInterchangeActorFactoryNode::GetLayerNames(TArray<FString>& OutLayerNames) const
+{
+	LayerNames.GetItems(OutLayerNames);
+}
+
+bool UInterchangeActorFactoryNode::AddLayerName(const FString& LayerName)
+{
+	return LayerNames.AddItem(LayerName);
+}
+
+bool UInterchangeActorFactoryNode::AddLayerNames(const TArray<FString>& InLayerNames)
+{
+	bool bSuccess = true;
+	for (const FString& LayerName : InLayerNames)
+	{
+		bSuccess &= LayerNames.AddItem(LayerName);
+	}
+
+	return bSuccess;
+}
+
+bool UInterchangeActorFactoryNode::RemoveLayerName(const FString& LayerName)
+{
+	return LayerNames.RemoveItem(LayerName);
+}
+
+void UInterchangeActorFactoryNode::GetTags(TArray<FString>& OutTags) const
+{
+	Tags.GetItems(OutTags);
+}
+
+bool UInterchangeActorFactoryNode::AddTag(const FString& Tag)
+{
+	return Tags.AddItem(Tag);
+}
+
+bool UInterchangeActorFactoryNode::AddTags(const TArray<FString>& InTags)
+{
+	bool bSuccess = true;
+	for (const FString& Tag : InTags)
+	{
+		bSuccess &= Tags.AddItem(Tag);
+	}
+
+	return bSuccess;
+}
+
+bool UInterchangeActorFactoryNode::RemoveTag(const FString& Tag)
+{
+	return Tags.RemoveItem(Tag);
+}
\ No newline at end of file
diff --git a/Engine/Plugins/Interchange/Runtime/Source/FactoryNodes/Public/InterchangeActorFactoryNode.h b/Engine/Plugins/Interchange/Runtime/Source/FactoryNodes/Public/InterchangeActorFactoryNode.h
--- a/Engine/Plugins/Interchange/Runtime/Source/FactoryNodes/Public/InterchangeActorFactoryNode.h
+++ b/Engine/Plugins/Interchange/Runtime/Source/FactoryNodes/Public/InterchangeActorFactoryNode.h
@@ -12,6 +12,8 @@ class INTERCHANGEFACTORYNODES_API UInterchangeActorFactoryNode : public UInterch
 	GENERATED_BODY()
 
 public:
+	UInterchangeActorFactoryNode();
+
 	virtual UClass* GetObjectClass() const override;
 
 	UFUNCTION(BlueprintCallable, Category = "Interchange | Node | ActorFactory")
@@ -52,6 +54,32 @@ public:
 
 	virtual void CopyWithObject(const UInterchangeFactoryBaseNode* SourceNode, UObject* Object) override;
 
+
+	/** Gets the LayerNames that this Actor is supposed to be part of. */
+	void GetLayerNames(TArray<FString>& OutLayerNames) const;
+
+	/** Add LayerName that this Actor is supposed to be part of. */
+	bool AddLayerName(const FString& InLayerName);
+
+	/** Add LayerNames that this Actor is supposed to be part of. */
+	bool AddLayerNames(const TArray<FString>& InLayerNames);
+
+	/** Remove LayerName that this Actor is supposed to be part of. */
+	bool RemoveLayerName(const FString& InLayerName);
+
+
+	/** Gets the Tags that this Actor is supposed to have. */
+	void GetTags(TArray<FString>& OutTags) const;
+
+	/** Add Tag that this Actor is supposed to have. */
+	bool AddTag(const FString& Tag);
+
+	/** Add Tags that this Actor is supposed to have. */
+	bool AddTags(const TArray<FString>& InTags);
+
+	/** Remove Tag that this Actor is supposed to have. */
+	bool RemoveTag(const FString& Tag);
+
 private:
 
 	bool ApplyCustomGlobalTransformToAsset(UObject* Asset) const;
@@ -61,4 +89,8 @@ private:
 	const UE::Interchange::FAttributeKey Macro_CustomLocalTransformKey = UE::Interchange::FAttributeKey(TEXT("LocalTransform"));
 	const UE::Interchange::FAttributeKey Macro_CustomActorClassNameKey = UE::Interchange::FAttributeKey(TEXT("ActorClassName"));
 	const UE::Interchange::FAttributeKey Macro_CustomMobilityKey = UE::Interchange::FAttributeKey(TEXT("Mobility"));
+
+	//A scene node can be part of multiple Layers.
+	UE::Interchange::TArrayAttributeHelper<FString> LayerNames;
+	UE::Interchange::TArrayAttributeHelper<FString> Tags;
 };
\ No newline at end of file
diff --git a/Engine/Plugins/Interchange/Runtime/Source/Import/Private/Scene/InterchangeActorFactory.cpp b/Engine/Plugins/Interchange/Runtime/Source/Import/Private/Scene/InterchangeActorFactory.cpp
--- a/Engine/Plugins/Interchange/Runtime/Source/Import/Private/Scene/InterchangeActorFactory.cpp
+++ b/Engine/Plugins/Interchange/Runtime/Source/Import/Private/Scene/InterchangeActorFactory.cpp
@@ -10,6 +10,15 @@
 #include "GameFramework/Actor.h"
 #include "Components/SceneComponent.h"
 
+#if WITH_EDITORONLY_DATA
+#include "Engine/World.h"
+#include "Layers/LayersSubsystem.h"
+#include "Layers/Layer.h"
+#include "Editor/EditorEngine.h"
+
+extern UNREALED_API UEditorEngine* GEditor;
+#endif
+
 #include UE_INLINE_GENERATED_CPP_BY_NAME(InterchangeActorFactory)
 
 namespace UE::InterchangeActorFactory::Private
@@ -115,6 +124,10 @@ UObject* UInterchangeActorFactory::ImportSceneObject_GameThread(const UInterchan
 				}
 			}
 		}
+
+		ProcessTags(FactoryNode, SpawnedActor);
+
+		ProcessLayerNames(FactoryNode, SpawnedActor);
 	}
 
 	return SpawnedActor;
@@ -124,3 +137,94 @@ UObject* UInterchangeActorFactory::ProcessActor(AActor& SpawnedActor, const UInt
 {
 	return SpawnedActor.GetRootComponent();
 }
+
+void UInterchangeActorFactory::ProcessTags(UInterchangeActorFactoryNode* FactoryNode, AActor* SpawnedActor)
+{
+	TArray<FString> TagsArray;
+	FactoryNode->GetTags(TagsArray);
+
+	TSet<FString> Tags(TagsArray);
+	TSet<FName> AlreadySetTags(SpawnedActor->Tags);
+
+	for (const FString& Tag : Tags)
+	{
+		FName TagName(Tag);
+		if (!AlreadySetTags.Contains(TagName))
+		{
+			SpawnedActor->Tags.Add(TagName);
+		}
+	}
+}
+
+void UInterchangeActorFactory::ProcessLayerNames(UInterchangeActorFactoryNode* FactoryNode, AActor* SpawnedActor)
+{
+	TArray<FString> LayerNamesArray;
+	FactoryNode->GetLayerNames(LayerNamesArray);
+
+	TSet<FString> LayerNames(LayerNamesArray);
+#if WITH_EDITORONLY_DATA
+	AddUniqueLayersToWorld(SpawnedActor->GetWorld(), LayerNames);
+#endif
+
+	TSet<FName> AlreadySetLayerNames(SpawnedActor->Layers);
+
+	for (const FString& LayerNameString : LayerNames)
+	{
+		FName LayerName(LayerNameString);
+		if (!AlreadySetLayerNames.Contains(LayerName))
+		{
+			SpawnedActor->Layers.Add(FName(LayerName));
+		}
+	}
+}
+
+#if WITH_EDITORONLY_DATA
+void UInterchangeActorFactory::AddUniqueLayersToWorld(UWorld* World, const TSet<FString>& LayerNames)
+{
+	if (!World || !IsValidChecked(World) || World->IsUnreachable() || LayerNames.Num() == 0)
+	{
+		return;
+	}
+
+	TSet< FName > ExistingLayers;
+	for (ULayer* Layer : World->Layers)
+	{
+		ExistingLayers.Add(Layer->GetLayerName());
+	}
+
+	int32 NumberOfExistingLayers = World->Layers.Num();
+
+	ULayersSubsystem* LayersSubsystem = GEditor ? GEditor->GetEditorSubsystem<ULayersSubsystem>() : nullptr;
+	for (const FString& LayerNameString : LayerNames)
+	{
+		FName LayerName(LayerNameString);
+
+		if (!ExistingLayers.Contains(LayerName))
+		{
+			// Use the ILayers if we are adding the layers to the currently edited world
+			if (LayersSubsystem && GWorld && World == GWorld.GetReference())
+			{
+				LayersSubsystem->CreateLayer(LayerName);
+			}
+			else
+			{
+				ULayer* NewLayer = NewObject<ULayer>(World, NAME_None, RF_Transactional);
+				if (!ensure(NewLayer != NULL))
+				{
+					continue;
+				}
+
+				World->Layers.Add(NewLayer);
+
+				NewLayer->SetLayerName(LayerName);
+				NewLayer->SetVisible(true);
+			}
+		}
+	}
+
+	if (NumberOfExistingLayers != World->Layers.Num())
+	{
+		World->Modify();
+	}
+}
+#endif
\ No newline at end of file
diff --git a/Engine/Plugins/Interchange/Runtime/Source/Import/Public/Scene/InterchangeActorFactory.h b/Engine/Plugins/Interchange/Runtime/Source/Import/Public/Scene/InterchangeActorFactory.h
--- a/Engine/Plugins/Interchange/Runtime/Source/Import/Public/Scene/InterchangeActorFactory.h
+++ b/Engine/Plugins/Interchange/Runtime/Source/Import/Public/Scene/InterchangeActorFactory.h
@@ -32,6 +32,12 @@ private:
 	// Interchange factory base interface end
 	//////////////////////////////////////////////////////////////////////////
 
+	void ProcessTags(UInterchangeActorFactoryNode* FactoryNode, AActor* SpawnedActor);
+	void ProcessLayerNames(UInterchangeActorFactoryNode* FactoryNode, AActor* SpawnedActor);
+#if WITH_EDITORONLY_DATA
+	void AddUniqueLayersToWorld(UWorld* World, const TSet<FString>& LayerNames);
+#endif
+
 protected:
 	/**
 	 * Method called in UInterchangeActorFactory::ImportSceneObject_GameThread to allow
diff --git a/Engine/Plugins/Interchange/Runtime/Source/Nodes/Private/InterchangeSceneNode.cpp b/Engine/Plugins/Interchange/Runtime/Source/Nodes/Private/InterchangeSceneNode.cpp
--- a/Engine/Plugins/Interchange/Runtime/Source/Nodes/Private/InterchangeSceneNode.cpp
+++ b/Engine/Plugins/Interchange/Runtime/Source/Nodes/Private/InterchangeSceneNode.cpp
@@ -53,13 +53,23 @@ namespace UE
 			return MeshToGlobalBindPoseReferncesString;
 		}
 
-		
-
 		const FString& FSceneNodeStaticData::GetMorphTargetCurveWeightsKey()
 		{
 			static FString MorphTargetCurvesKey(TEXT("__MorphTargetCurveWeights__Key"));
 			return MorphTargetCurvesKey;
 		}
+
+		const FString& FSceneNodeStaticData::GetLayerNamesKey()
+		{
+			static FString LayerNamesKey(TEXT("__LayerNames__Key"));
+			return LayerNamesKey;
+		}
+
+		const FString& FSceneNodeStaticData::GetTagsKey()
+		{
+			static FString TagsKey(TEXT("__Tags__Key"));
+			return TagsKey;
+		}
 	}//ns Interchange
 }//ns UE
 
@@ -69,6 +79,8 @@ UInterchangeSceneNode::UInterchangeSceneNode()
 	MeshToGlobalBindPoseReferences.Initialize(Attributes.ToSharedRef(), UE::Interchange::FSceneNodeStaticData::GetMeshToGlobalBindPoseReferencesString());
 	SlotMaterialDependencies.Initialize(Attributes.ToSharedRef(), UE::Interchange::FSceneNodeStaticData::GetSlotMaterialDependenciesString());
 	MorphTargetCurveWeights.Initialize(Attributes.ToSharedRef(), UE::Interchange::FSceneNodeStaticData::GetMorphTargetCurveWeightsKey());
+	LayerNames.Initialize(Attributes, UE::Interchange::FSceneNodeStaticData::GetLayerNamesKey());
+	Tags.Initialize(Attributes, UE::Interchange::FSceneNodeStaticData::GetTagsKey());
 }
 
 /**
@@ -448,4 +460,34 @@ bool UInterchangeSceneNode::SetCustomHasBindPose(const bool& AttributeValue)
 bool UInterchangeSceneNode::GetCustomHasBindPose(bool& AttributeValue) const
 {
 	IMPLEMENT_NODE_ATTRIBUTE_GETTER(HasBindPose, bool);
+}
+
+void UInterchangeSceneNode::GetLayerNames(TArray<FString>& OutLayerNames) const
+{
+	LayerNames.GetItems(OutLayerNames);
+}
+
+bool UInterchangeSceneNode::AddLayerName(const FString& LayerName)
+{
+	return LayerNames.AddItem(LayerName);
+}
+
+bool UInterchangeSceneNode::RemoveLayerName(const FString& LayerName)
+{
+	return LayerNames.RemoveItem(LayerName);
+}
+
+void UInterchangeSceneNode::GetTags(TArray<FString>& OutTags) const
+{
+	Tags.GetItems(OutTags);
+}
+
+bool UInterchangeSceneNode::AddTag(const FString& Tag)
+{
+	return Tags.AddItem(Tag);
+}
+
+bool UInterchangeSceneNode::RemoveTag(const FString& Tag)
+{
+	return Tags.RemoveItem(Tag);
 }
\ No newline at end of file
diff --git a/Engine/Plugins/Interchange/Runtime/Source/Nodes/Public/InterchangeSceneNode.h b/Engine/Plugins/Interchange/Runtime/Source/Nodes/Public/InterchangeSceneNode.h
--- a/Engine/Plugins/Interchange/Runtime/Source/Nodes/Public/InterchangeSceneNode.h
+++ b/Engine/Plugins/Interchange/Runtime/Source/Nodes/Public/InterchangeSceneNode.h
@@ -25,6 +25,8 @@ namespace UE
 			static const FString& GetMeshToGlobalBindPoseReferencesString();
 			static const FString& GetSlotMaterialDependenciesString();
 			static const FString& GetMorphTargetCurveWeightsKey();
+			static const FString& GetLayerNamesKey();
+			static const FString& GetTagsKey();
 		};
 
 	}//ns Interchange
@@ -239,6 +241,24 @@ public:
 	UFUNCTION(BlueprintCallable, Category = "Interchange | Node | Joint")
 	bool GetCustomHasBindPose(bool &bHasBindPose) const;
 
+	/** Gets the LayerNames that this SceneNode (Actor) is supposed to be part of. */
+	void GetLayerNames(TArray<FString>& OutLayerNames) const;
+
+	/** Add LayerName that this SceneNode (Actor) is supposed to be part of. */
+	bool AddLayerName(const FString& LayerName);
+
+	/** Remove LayerName that this SceneNode (Actor) is supposed to be part of. */
+	bool RemoveLayerName(const FString& LayerName);
+
+	/** Gets the Tags that this SceneNode (Actor) is supposed to have. */
+	void GetTags(TArray<FString>& OutTags) const;
+
+	/** Add Tag that this SceneNode (Actor) is supposed to have. */
+	bool AddTag(const FString& Tag);
+
+	/** Remove Tag that this SceneNode (Actor) is supposed to have. */
+	bool RemoveTag(const FString& Tag);
+
 private:
 
 	bool GetGlobalTransformInternal(const UE::Interchange::FAttributeKey LocalTransformKey
@@ -281,6 +301,12 @@ private:
 	//BindPose References per Mesh for a JointNode.
 	UE::Interchange::TMapAttributeHelper<FString, FMatrix> MeshToGlobalBindPoseReferences;
 
+	//A scene node can be part of multiple Layers.
+	UE::Interchange::TArrayAttributeHelper<FString> LayerNames;
+
+	//A scene node can have multiple Tags.
+	UE::Interchange::TArrayAttributeHelper<FString> Tags;
+
 	//A scene node can reference an animation asset on top of base asset:
 	const UE::Interchange::FAttributeKey Macro_CustomAnimationAssetUidToPlayKey = UE::Interchange::FAttributeKey(TEXT("AnimationAssetUidToPlay"));
 
diff --git a/Engine/Plugins/Interchange/Runtime/Source/Pipelines/Private/InterchangeGenericScenesPipeline.cpp b/Engine/Plugins/Interchange/Runtime/Source/Pipelines/Private/InterchangeGenericScenesPipeline.cpp
--- a/Engine/Plugins/Interchange/Runtime/Source/Pipelines/Private/InterchangeGenericScenesPipeline.cpp
+++ b/Engine/Plugins/Interchange/Runtime/Source/Pipelines/Private/InterchangeGenericScenesPipeline.cpp
@@ -498,6 +498,14 @@ void UInterchangeGenericLevelPipeline::ExecuteSceneNodePreImport(const FTransfor
 
 	UInterchangeUserDefinedAttributesAPI::DuplicateAllUserDefinedAttribute(SceneNode, ActorFactoryNode, false);
 
+	TArray<FString> LayerNames;
+	SceneNode->GetLayerNames(LayerNames);
+	ActorFactoryNode->AddLayerNames(LayerNames);
+
+	TArray<FString> Tags;
+	SceneNode->GetTags(Tags);
+	ActorFactoryNode->AddTags(Tags);
+
 	FString NodeUid = SceneNode->GetUniqueID() + (bRootJointNode ? TEXT("_SkeletonNode") : TEXT(""));
 	FString FactoryNodeUid = UInterchangeFactoryBaseNode::BuildFactoryNodeUid(NodeUid);
 
