[Interchange] AnimSequence fixes (ReferenceSkeleton vs Imported Bone transform without animation)

#rnx

diff --git a/Engine/Plugins/Interchange/Runtime/Source/Import/Private/Animation/InterchangeAnimSequenceFactory.cpp b/Engine/Plugins/Interchange/Runtime/Source/Import/Private/Animation/InterchangeAnimSequenceFactory.cpp
--- a/Engine/Plugins/Interchange/Runtime/Source/Import/Private/Animation/InterchangeAnimSequenceFactory.cpp
+++ b/Engine/Plugins/Interchange/Runtime/Source/Import/Private/Animation/InterchangeAnimSequenceFactory.cpp
@@ -587,12 +587,25 @@ namespace UE::Interchange::Private
 			{
 				if (const UInterchangeSceneNode* SkeletonNode = Cast<UInterchangeSceneNode>(NodeContainer->GetNode(NonAnimatedSkeletonNodeUID)))
 				{
-					//check if BindPose exists and if it does does it equal to LocalTransform
+					TOptional<FTransform> ReferenceTransform;
+					if (Skeleton)
+					{
+						const FReferenceSkeleton& RefSkeleton = Skeleton->GetReferenceSkeleton();
+						int32 BoneIndex = RefSkeleton.FindBoneIndex(*SkeletonNode->GetDisplayLabel());
+						if (BoneIndex != INDEX_NONE && RefSkeleton.GetRefBonePose().IsValidIndex(BoneIndex))
+						{
+							ReferenceTransform = RefSkeleton.GetRefBonePose()[BoneIndex];
+						}
+					}
+
+					//check if BindPose exists and if it does, does it equal to LocalTransform/ReferenceTransform
 					FTransform LocalBindPoseTransform;
 					FTransform LocalTransform;
 					if (SkeletonNode->GetCustomBindPoseLocalTransform(LocalBindPoseTransform)
 						&& SkeletonNode->GetCustomLocalTransform(LocalTransform)
-						&& !LocalBindPoseTransform.Equals(LocalTransform))
+						&& (!LocalBindPoseTransform.Equals(LocalTransform) 
+							|| (ReferenceTransform.IsSet() && (!ReferenceTransform.GetValue().Equals(LocalBindPoseTransform))))
+						)
 					{
 						//If we bake the mesh and the current non animated node is the root joint, get the global transform instead of the local
 						if (bBakeMeshes && SkeletonNode->GetUniqueID().Equals(SkeletonRootUid))
