Interchange Collision workflow:
- Remove the "from file" option from EInterchangeMeshCollision
- Fix a bug where custom collisions are not imported when None is chosen in the dropdown, also added a comment to let the user know the real intent behind it (different from the boolean to import collisions)
- Remove the reimport strategy, replacing it by emptying the previous imported collisions when importing the new ones
- Preserving the editor-generated collisions during a reimport, regenerating them with their previous settings, just updating the extents
- Storing the previous editor-generated collisions to regenerate them in the game thread
- Renaming Default Collision Type into Fallback collision type
- fixing some typos as well

diff --git a/Engine/Plugins/Interchange/Editor/Source/InterchangeEditor/Private/InterchangeFbxAssetImportDataConverter.cpp b/Engine/Plugins/Interchange/Editor/Source/InterchangeEditor/Private/InterchangeFbxAssetImportDataConverter.cpp
--- a/Engine/Plugins/Interchange/Editor/Source/InterchangeEditor/Private/InterchangeFbxAssetImportDataConverter.cpp
+++ b/Engine/Plugins/Interchange/Editor/Source/InterchangeEditor/Private/InterchangeFbxAssetImportDataConverter.cpp
@@ -123,7 +123,7 @@ namespace UE::Interchange::Private
 
 	void FillFbxStaticMeshImportData(const UInterchangeGenericAssetsPipeline* GenericAssetPipeline, UFbxStaticMeshImportData* DestinationStaticMeshImportData)
 	{
-		DestinationStaticMeshImportData->bAutoGenerateCollision = GenericAssetPipeline->MeshPipeline->bCollision && GenericAssetPipeline->MeshPipeline->Collision != EInterchangeMeshCollision::None && GenericAssetPipeline->MeshPipeline->Collision != EInterchangeMeshCollision::Custom;
+		DestinationStaticMeshImportData->bAutoGenerateCollision = GenericAssetPipeline->MeshPipeline->bCollision && GenericAssetPipeline->MeshPipeline->Collision != EInterchangeMeshCollision::None;
 		DestinationStaticMeshImportData->bBuildNanite = GenericAssetPipeline->MeshPipeline->bBuildNanite;
 		DestinationStaticMeshImportData->bBuildReversedIndexBuffer = GenericAssetPipeline->MeshPipeline->bBuildReversedIndexBuffer;
 		DestinationStaticMeshImportData->bCombineMeshes = GenericAssetPipeline->MeshPipeline->bCombineStaticMeshes;
diff --git a/Engine/Plugins/Interchange/Runtime/Source/Import/Private/Mesh/InterchangeStaticMeshFactory.cpp b/Engine/Plugins/Interchange/Runtime/Source/Import/Private/Mesh/InterchangeStaticMeshFactory.cpp
--- a/Engine/Plugins/Interchange/Runtime/Source/Import/Private/Mesh/InterchangeStaticMeshFactory.cpp
+++ b/Engine/Plugins/Interchange/Runtime/Source/Import/Private/Mesh/InterchangeStaticMeshFactory.cpp
@@ -31,7 +31,6 @@
 #include "Model.h"
 #include "Nodes/InterchangeBaseNode.h"
 #include "Nodes/InterchangeBaseNodeContainer.h"
-#include "PhysicsEngine/AggregateGeom.h"
 #include "PhysicsEngine/BodySetup.h"
 #include "PhysicsEngine/BoxElem.h"
 #include "PhysicsEngine/ConvexElem.h"
@@ -430,27 +429,17 @@ UInterchangeFactoryBase::FImportAssetResult UInterchangeStaticMeshFactory::Impor
 		{
 			LodDataNode->GetImportCollision(bImportCollision);
 			LodDataNode->GetImportCollisionType(Collision);
-			if(bImportCollision && Collision != EInterchangeMeshCollision::None)
+			if(bImportCollision)
 			{
 				if (bReimport)
 				{
-					EReimportStrategyFlags ReimportStrategy = StaticMeshFactoryNode->GetReimportStrategyFlags();
-					switch(ReimportStrategy)
-					{
-					using enum EReimportStrategyFlags;
-					case ApplyNoProperties:	// do nothing
-						break;
-
-					case ApplyPipelineProperties:
-						StaticMesh->GetBodySetup()->AggGeom.EmptyElements();
-						break;
-
-					case ApplyEditorChangedProperties:
-#if WITH_EDITORONLY_DATA
-						StaticMesh->GetBodySetup()->AggGeom.EmptyImportedElements();
+					//Let's clean only the imported collisions first in order 
+					// to store the previous editor-generated collisions to re-generate them later in the Game Thread with their properties
+#if WITH_EDITOR
+					StaticMesh->GetBodySetup()->AggGeom.EmptyImportedElements();
 #endif
-						break;
-					}
+					ImportAssetObjectData.AggregateGeom = StaticMesh->GetBodySetup()->AggGeom;
+					StaticMesh->GetBodySetup()->AggGeom.EmptyElements();
 				}
 
 				bImportedCustomCollision |= ImportBoxCollision(Arguments, StaticMesh, LodDataNode);
@@ -687,12 +676,79 @@ UInterchangeFactoryBase::FImportAssetResult UInterchangeStaticMeshFactory::EndIm
 			}
 		}
 	}
-#endif // WITH_EDITOR
 
-#if WITH_EDITOR
+	//Let's now re-generate the previous collisions with their properties, only the extents will be updated	
+	if(bReimport)
+	{
+		for(FKBoxElem& BoxElem : ImportAssetObjectData.AggregateGeom.BoxElems)
+		{
+			int32 Index = GenerateBoxAsSimpleCollision(StaticMesh, false);
+			FKBoxElem& NewBoxElem = StaticMesh->GetBodySetup()->AggGeom.BoxElems[Index];
+
+			//Preserve the new extents
+			float X = NewBoxElem.X;
+			float Y = NewBoxElem.Y;
+			float Z = NewBoxElem.Z;
+
+			//Copy the previous settings, we just want the new extents
+			NewBoxElem = BoxElem;
+			NewBoxElem.X = X;
+			NewBoxElem.Y = Y;
+			NewBoxElem.Z = Z;
+		}
+
+		for(FKSphereElem& SphereElem : ImportAssetObjectData.AggregateGeom.SphereElems)
+		{
+			int32 Index = GenerateSphereAsSimpleCollision(StaticMesh, false);
+			FKSphereElem& NewSphereElem = StaticMesh->GetBodySetup()->AggGeom.SphereElems[Index];
+
+			//Preserve the new extents
+			float Radius = NewSphereElem.Radius;
+
+			//Copy the previous settings, we just want the new extents
+			NewSphereElem = SphereElem;
+			NewSphereElem.Radius = Radius;
+		}
+
+		for(FKSphylElem& CapsuleElem: ImportAssetObjectData.AggregateGeom.SphylElems)
+		{
+			int32 Index = GenerateSphylAsSimpleCollision(StaticMesh, false);
+			FKSphylElem& NewCapsuleElem = StaticMesh->GetBodySetup()->AggGeom.SphylElems[Index];
+
+			//Preserve the new extents
+			float Radius = NewCapsuleElem.Radius;
+			float Length = NewCapsuleElem.Length;
+
+			//Copy the previous settings, we just want the new extents
+			NewCapsuleElem = CapsuleElem;
+			NewCapsuleElem.Radius = Radius;
+			NewCapsuleElem.Length = Length;
+		}
+
+		for(FKConvexElem& ConvexElem: ImportAssetObjectData.AggregateGeom.ConvexElems)
+		{
+			int32 Index = GenerateKDopAsSimpleCollision(StaticMesh, TArray<FVector>(KDopDir18, sizeof(KDopDir18) / sizeof(FVector)), false);
+			FKConvexElem& NewConvexElem = StaticMesh->GetBodySetup()->AggGeom.ConvexElems[Index];
+
+			//Preserve the new extents
+			TArray<FVector> VertexData = NewConvexElem.VertexData;
+			TArray<int32> IndexData = NewConvexElem.IndexData;
+			FBox ElemBox = NewConvexElem.ElemBox;
+			Chaos::FConvexPtr ChaosConvex = NewConvexElem.GetChaosConvexMesh();
+			NewConvexElem.ResetChaosConvexMesh();
+
+			//Copy the previous settings, we just want the new extents
+			NewConvexElem = ConvexElem;
+			NewConvexElem.VertexData = VertexData;
+			NewConvexElem.IndexData = IndexData;
+			NewConvexElem.ElemBox = ElemBox;
+			NewConvexElem.SetConvexMeshObject(MoveTemp(ChaosConvex));
+		}
+	}
+
 	if(ImportAssetObjectData.bImportCollision)
 	{
-		if(!ImportAssetObjectData.bImportedCustomCollision)
+		if(!ImportAssetObjectData.bImportedCustomCollision && ImportAssetObjectData.Collision != EInterchangeMeshCollision::None)
 		{
 			// Don't generate collisions if the mesh already has one of the requested type, otherwise it will continue to create collisions,
 			// it can happen in the case of an import, and then importing the same file without deleting the asset in the content browser (different from a reimport)
@@ -757,7 +813,7 @@ UInterchangeFactoryBase::FImportAssetResult UInterchangeStaticMeshFactory::EndIm
 				break;
 			}
 		}
-#endif
+#endif // WITH_EDITOR
 #if WITH_EDITORONLY_DATA
 		else
 		{
diff --git a/Engine/Plugins/Interchange/Runtime/Source/Import/Public/Mesh/InterchangeStaticMeshFactory.h b/Engine/Plugins/Interchange/Runtime/Source/Import/Public/Mesh/InterchangeStaticMeshFactory.h
--- a/Engine/Plugins/Interchange/Runtime/Source/Import/Public/Mesh/InterchangeStaticMeshFactory.h
+++ b/Engine/Plugins/Interchange/Runtime/Source/Import/Public/Mesh/InterchangeStaticMeshFactory.h
@@ -6,6 +6,7 @@
 #include "InterchangeFactoryBase.h"
 #include "InterchangeMeshDefinitions.h"
 #include "MeshDescription.h"
+#include "PhysicsEngine/AggregateGeom.h"
 #include "UObject/Object.h"
 #include "UObject/ObjectMacros.h"
 #include "Mesh/InterchangeMeshPayload.h"
@@ -77,11 +78,12 @@ private:
 #endif
 	struct FImportAssetObjectData
 	{
+		TArray<FMeshDescription> LodMeshDescriptions;
+		EInterchangeMeshCollision Collision = EInterchangeMeshCollision::None;
+		FKAggregateGeom AggregateGeom;
 		bool bIsAppGame = false;
 		bool bImportedCustomCollision = false;
 		bool bImportCollision = false;
-		EInterchangeMeshCollision Collision = EInterchangeMeshCollision::None;
-		TArray<FMeshDescription> LodMeshDescriptions;
 	};
 	FImportAssetObjectData ImportAssetObjectData;
 };
diff --git a/Engine/Plugins/Interchange/Runtime/Source/Nodes/Public/InterchangeMeshDefinitions.h b/Engine/Plugins/Interchange/Runtime/Source/Nodes/Public/InterchangeMeshDefinitions.h
--- a/Engine/Plugins/Interchange/Runtime/Source/Nodes/Public/InterchangeMeshDefinitions.h
+++ b/Engine/Plugins/Interchange/Runtime/Source/Nodes/Public/InterchangeMeshDefinitions.h
@@ -33,11 +33,8 @@ enum class EInterchangeMeshCollision : uint8
 	
 	/** Generates a new axis-aligned box collision mesh with all edges and corners beveled (26 total sides) encompassing the static mesh*/
 	Convex26DOP UMETA(DisplayName = "26DOP Simplified"),
-	
-	/** Generates only collisions retrieved from the file. This option is mostly useful in case of a reimport,
-	  * to only import collisions by removing the previous imported ones and preserving the ones that were editor-created.*/
-	Custom UMETA(DisplayName = "From File"),
 
+	/** Generates no collisions, but continue to import custom collisions if the file has ones*/
 	None = 0xFF
 };
 
