Interchange: Fix fbx skeletalmesh root finding.

#rnx

diff --git a/Engine/Plugins/Interchange/Runtime/Source/Parsers/Fbx/Private/FbxAPI.cpp b/Engine/Plugins/Interchange/Runtime/Source/Parsers/Fbx/Private/FbxAPI.cpp
--- a/Engine/Plugins/Interchange/Runtime/Source/Parsers/Fbx/Private/FbxAPI.cpp
+++ b/Engine/Plugins/Interchange/Runtime/Source/Parsers/Fbx/Private/FbxAPI.cpp
@@ -152,6 +152,13 @@ namespace UE
 
 				// Get The Creator of the FBX File.
 				FileDetails.FbxFileCreator = UTF8_TO_TCHAR(SDKImporter->GetFileHeaderInfo()->mCreator.Buffer());
+				{
+					//Example of creator file info string
+					//Blender (stable FBX IO) - 2.78 (sub 0) - 3.7.7
+					//Maya and Max use the same string where they specify the fbx sdk version, so we cannot know it is coming from which software
+					//We need blender creator when importing skeletal mesh containing the "armature" dummy node as the parent of the root joint. We want to remove this dummy "armature" node
+					bCreatorIsBlender = FileDetails.FbxFileCreator.StartsWith(TEXT("Blender"));
+				}
 
 				FbxDocumentInfo* DocInfo = SDKImporter->GetSceneInfo();
 				if (DocInfo)
diff --git a/Engine/Plugins/Interchange/Runtime/Source/Parsers/Fbx/Private/FbxAPI.h b/Engine/Plugins/Interchange/Runtime/Source/Parsers/Fbx/Private/FbxAPI.h
--- a/Engine/Plugins/Interchange/Runtime/Source/Parsers/Fbx/Private/FbxAPI.h
+++ b/Engine/Plugins/Interchange/Runtime/Source/Parsers/Fbx/Private/FbxAPI.h
@@ -106,6 +106,7 @@ namespace UE
 
 				double GetFrameRate() { return FrameRate; }
 
+				bool IsCreatorBlender() { return bCreatorIsBlender; }
 				/**
 				 * Critical section to avoid getting multiple payload in same time.
 				 * The FBX evaluator use a cache mechanism for evaluating global transform that is not thread safe.
@@ -136,6 +137,7 @@ namespace UE
 				bool bForceFrontXAxis = false;
 				bool bConvertSceneUnit = true;
 				bool bKeepFbxNamespace = false;
+				bool bCreatorIsBlender = false;
 
 				struct FileDetails
 				{
diff --git a/Engine/Plugins/Interchange/Runtime/Source/Parsers/Fbx/Private/FbxScene.cpp b/Engine/Plugins/Interchange/Runtime/Source/Parsers/Fbx/Private/FbxScene.cpp
--- a/Engine/Plugins/Interchange/Runtime/Source/Parsers/Fbx/Private/FbxScene.cpp
+++ b/Engine/Plugins/Interchange/Runtime/Source/Parsers/Fbx/Private/FbxScene.cpp
@@ -91,66 +91,22 @@ namespace UE
 				CreateAssetNodeReference(Parser, UnrealSceneNode, NodeAttribute, NodeContainer, UInterchangeLightNode::StaticAssetTypeName());
 			}
 
-			bool DoesChildrenHierarchyContainJoints(FbxNode* Node, TFunction<bool(FbxNode*)> IsNodeAjoint)
+			bool IsNodeUnderCommonJointRootNode(FbxNode* Node, TSet<FbxNode*>& CommonJointRootNodes)
 			{
-				if (!Node)
+				if (!Node || CommonJointRootNodes.IsEmpty())
 				{
 					return false;
 				}
-				if (IsNodeAjoint(Node))
-				{
-					return true;
-				}
-				int32 ChildCount = Node->GetChildCount();
-				for (int32 ChildIndex = 0; ChildIndex < ChildCount; ++ChildIndex)
-				{
-					FbxNode* ChildNode = Node->GetChild(ChildIndex);
-					if (DoesChildrenHierarchyContainJoints(ChildNode, IsNodeAjoint))
-					{
-						return true;
-					}
-				}
-				return false;
-			}
 
-			bool DoesTheParentOrChildrenHierarchyContainJoints(FbxNode* Node)
-			{
-				if (!Node)
+				//Simply go up the hierarchy until we match the CommonJointRootNode
+				FbxNode* IterateNode = Node;
+				while (IterateNode)
 				{
-					return false;
-				}
-				auto IsNodeAjoint = [](FbxNode* NodeToTest)
-					{
-						int32 AttributeCount = NodeToTest->GetNodeAttributeCount();
-						for (int32 AttributeIndex = 0; AttributeIndex < AttributeCount; ++AttributeIndex)
-						{
-							FbxNodeAttribute* NodeAttribute = NodeToTest->GetNodeAttributeByIndex(AttributeIndex);
-							if (NodeAttribute->GetAttributeType() == FbxNodeAttribute::eSkeleton)
-							{
-								return true;
-							}
-						}
-						return false;
-					};
-				
-				if (IsNodeAjoint(Node))
-				{
-					return true;
-				}
-
-				FbxNode* ParentNode = Node->GetParent();
-				while (ParentNode)
-				{
-					if (IsNodeAjoint(ParentNode))
+					if (CommonJointRootNodes.Contains(IterateNode))
 					{
 						return true;
 					}
-					ParentNode = ParentNode->GetParent();
-				}
-
-				if (DoesChildrenHierarchyContainJoints(Node, IsNodeAjoint))
-				{
-					return true;
+					IterateNode = IterateNode->GetParent();
 				}
 				return false;
 			}
@@ -310,7 +266,7 @@ namespace UE
 
 						case FbxNodeAttribute::eNull:
 						{
-							if (!DoesTheParentOrChildrenHierarchyContainJoints(Node))
+							if (!IsNodeUnderCommonJointRootNode(Node, CommonJointRootNodes))
 							{
 								//eNull node not in a hierarchy containing any joint will not be set has joint
 								break;
@@ -385,7 +341,7 @@ namespace UE
 						UnrealNode->AddSpecializedType(FSceneNodeStaticData::GetTransformSpecializeTypeString());
 						ApplySkeletonAttribute();
 					}
-					else if (!bIsRootNode && DoesTheParentOrChildrenHierarchyContainJoints(Node))
+					else if (!bIsRootNode && IsNodeUnderCommonJointRootNode(Node, CommonJointRootNodes))
 					{
 						UnrealNode->AddSpecializedType(FSceneNodeStaticData::GetTransformSpecializeTypeString());
 						ApplySkeletonAttribute();
@@ -466,7 +422,7 @@ namespace UE
 				//Add all custom Attributes for the node
 				while (Property.IsValid())
 				{
-					EFbxType PropertyType =  Property.GetPropertyDataType().GetType();
+					EFbxType PropertyType = Property.GetPropertyDataType().GetType();
 					if (Property.GetFlag(FbxPropertyFlags::eUserDefined) && FFbxAnimation::IsFbxPropertyTypeSupported(PropertyType))
 					{
 						FbxAnimCurveNode* CurveNode = Property.GetCurveNode();
@@ -526,6 +482,9 @@ namespace UE
 				TArray<FbxNode*> ForceJointNodes;
 				FindForceJointNode(SDKScene, ForceJointNodes);
 
+				//Cache the common root joint
+				FindCommonJointRootNode(SDKScene, ForceJointNodes);
+
 				bool bBadBindPoseMessageDisplay = false;
 				AddHierarchyRecursively(nullptr, RootNode, SDKScene, NodeContainer, PayloadContexts, ForceJointNodes, bBadBindPoseMessageDisplay);
 
@@ -534,7 +493,7 @@ namespace UE
 				{
 					if (FbxNode* Node = SDKScene->GetNode(NodeIndex))
 					{
-						if(Node != RootNode)
+						if (Node != RootNode)
 						{
 							if (Node->GetParent() == nullptr)
 							{
@@ -630,6 +589,86 @@ namespace UE
 				}
 			}
 
+			FbxNode* FFbxScene::Internal_GetRootSkeleton(FbxScene* SDKScene, FbxNode* Link)
+			{
+				FbxNode* RootBone = Link;
+
+				// get Unreal skeleton root
+				// mesh and dummy are used as bone if they are in the skeleton hierarchy
+				while (RootBone && RootBone->GetParent())
+				{
+					bool bIsBlenderArmatureBone = false;
+					if (Parser.IsCreatorBlender())
+					{
+						//Hack to support armature dummy node from blender
+						//Users do not want the null attribute node named armature which is the parent of the real root bone in blender fbx file
+						//This is a hack since if a rigid mesh group root node is named "armature" it will be skip
+						const FString RootBoneParentName(RootBone->GetParent()->GetName());
+						FbxNode* GrandFather = RootBone->GetParent()->GetParent();
+						bIsBlenderArmatureBone = (GrandFather == nullptr || GrandFather == SDKScene->GetRootNode()) && (RootBoneParentName.Compare(TEXT("armature"), ESearchCase::IgnoreCase) == 0);
+					}
+
+					FbxNodeAttribute* Attr = RootBone->GetParent()->GetNodeAttribute();
+					if (Attr &&
+						(Attr->GetAttributeType() == FbxNodeAttribute::eMesh ||
+							(Attr->GetAttributeType() == FbxNodeAttribute::eNull && !bIsBlenderArmatureBone) ||
+							Attr->GetAttributeType() == FbxNodeAttribute::eSkeleton) &&
+						RootBone->GetParent() != SDKScene->GetRootNode())
+					{
+						// in some case, skeletal mesh can be ancestor of bones
+						// this avoids this situation
+						if (Attr->GetAttributeType() == FbxNodeAttribute::eMesh)
+						{
+							FbxMesh* Mesh = (FbxMesh*)Attr;
+							if (Mesh->GetDeformerCount(FbxDeformer::eSkin) > 0)
+							{
+								break;
+							}
+						}
+
+						RootBone = RootBone->GetParent();
+					}
+					else
+					{
+						break;
+					}
+				}
+
+				return RootBone;
+			}
+
+			void FFbxScene::FindCommonJointRootNode(FbxScene* SDKScene, const TArray<FbxNode*>& ForceJointNodes)
+			{
+				//Process the ForceJointNodes and any skeleton joint node
+				int32 NodeCount = SDKScene->GetNodeCount();
+				for (int32 NodeIndex = 0; NodeIndex < NodeCount; ++NodeIndex)
+				{
+					if (FbxNode* Node = SDKScene->GetNode(NodeIndex))
+					{
+						bool bProcessNode = ForceJointNodes.Contains(Node);
+						if (!bProcessNode)
+						{
+							const int32 AttributeCount = Node->GetNodeAttributeCount();
+							for (int32 AttributeIndex = 0; AttributeIndex < AttributeCount; ++AttributeIndex)
+							{
+								if (Node->GetNodeAttributeByIndex(AttributeIndex)->GetAttributeType() == FbxNodeAttribute::eSkeleton)
+								{
+									bProcessNode = true;
+									break;
+								}
+							}
+						}
+						if (bProcessNode)
+						{
+							if (FbxNode* Root = Internal_GetRootSkeleton(SDKScene, Node))
+							{
+								CommonJointRootNodes.Add(Root);
+							}
+						}
+					}
+				}
+			}
+
 			void FFbxScene::FindForceJointNode(FbxScene* SDKScene, TArray<FbxNode*>& ForceJointNodes)
 			{
 				int32 GeometryCount = SDKScene->GetGeometryCount();
@@ -724,7 +763,7 @@ namespace UE
 						switch (NodeAttribute->GetAttributeType())
 						{
 						case FbxNodeAttribute::eNull:
-							if (!DoesTheParentOrChildrenHierarchyContainJoints(Node))
+							if (!IsNodeUnderCommonJointRootNode(Node, CommonJointRootNodes))
 							{
 								//eNull node not under any joint are not joint
 								break;
@@ -745,7 +784,7 @@ namespace UE
 						{
 							bNewSkeltalAnimationStarted = ApplySkeletonAttribute() || bNewSkeltalAnimationStarted;
 						}
-						else if (!bIsRootNode && DoesTheParentOrChildrenHierarchyContainJoints(Node))
+						else if (!bIsRootNode && IsNodeUnderCommonJointRootNode(Node, CommonJointRootNodes))
 						{
 							bNewSkeltalAnimationStarted = ApplySkeletonAttribute() || bNewSkeltalAnimationStarted;
 						}
diff --git a/Engine/Plugins/Interchange/Runtime/Source/Parsers/Fbx/Private/FbxScene.h b/Engine/Plugins/Interchange/Runtime/Source/Parsers/Fbx/Private/FbxScene.h
--- a/Engine/Plugins/Interchange/Runtime/Source/Parsers/Fbx/Private/FbxScene.h
+++ b/Engine/Plugins/Interchange/Runtime/Source/Parsers/Fbx/Private/FbxScene.h
@@ -62,8 +62,14 @@ namespace UE
 					, UInterchangeBaseNodeContainer& NodeContainer
 					, TMap<FString, TSharedPtr<FPayloadContextBase, ESPMode::ThreadSafe>>& PayloadContexts);
 
+				FbxNode* Internal_GetRootSkeleton(FbxScene* SDKScene, FbxNode* Link);
+				void FindCommonJointRootNode(FbxScene* SDKScene, const TArray<FbxNode*>& ForceJointNodes);
+
 				void FindForceJointNode(FbxScene* SDKScene, TArray<FbxNode*>& ForceJointNodes);
 
+
+				TSet<FbxNode*> CommonJointRootNodes;
+
 				FFbxParser& Parser;
 			};
 		}//ns Private
diff --git a/Engine/Plugins/Interchange/Runtime/Source/Pipelines/Private/InterchangePipelineMeshesUtilities.cpp b/Engine/Plugins/Interchange/Runtime/Source/Pipelines/Private/InterchangePipelineMeshesUtilities.cpp
--- a/Engine/Plugins/Interchange/Runtime/Source/Pipelines/Private/InterchangePipelineMeshesUtilities.cpp
+++ b/Engine/Plugins/Interchange/Runtime/Source/Pipelines/Private/InterchangePipelineMeshesUtilities.cpp
@@ -278,7 +278,7 @@ UInterchangePipelineMeshesUtilities* UInterchangePipelineMeshesUtilities::Create
 	
 	//Find all translated node we need for this pipeline
 	BaseNodeContainer->IterateNodes(
-		[&PipelineMeshesUtilities, &BaseNodeContainer, &SkeletonRootNodeUids](const FString& NodeUid, UInterchangeBaseNode* Node)
+		[&PipelineMeshesUtilities, &BaseNodeContainer](const FString& NodeUid, UInterchangeBaseNode* Node)
 		{
 			if (Node->GetNodeContainerType() == EInterchangeNodeContainerType::TranslatedAsset)
 			{
