Interchange: Mesh payload refactor, the payload async is now control by the interchange task system.

#rnx

diff --git a/Engine/Plugins/Experimental/Enterprise/DatasmithInterchange/Source/DatasmithInterchange/Private/InterchangeDatasmithTranslator.cpp b/Engine/Plugins/Experimental/Enterprise/DatasmithInterchange/Source/DatasmithInterchange/Private/InterchangeDatasmithTranslator.cpp
--- a/Engine/Plugins/Experimental/Enterprise/DatasmithInterchange/Source/DatasmithInterchange/Private/InterchangeDatasmithTranslator.cpp
+++ b/Engine/Plugins/Experimental/Enterprise/DatasmithInterchange/Source/DatasmithInterchange/Private/InterchangeDatasmithTranslator.cpp
@@ -39,6 +39,7 @@
 #include "InterchangeVariantSetNode.h"
 #include "StaticMeshOperations.h"
 
+#include "Async/ParallelFor.h"
 #include "HAL/ConsoleManager.h"
 #include "Misc/App.h"
 #include "Misc/PackageName.h"
@@ -720,16 +721,15 @@ TOptional<UE::Interchange::FImportLightProfile> UInterchangeDatasmithTranslator:
 	return TextureTranslator->GetLightProfilePayloadData(PayloadKey, AlternateTexturePath);
 }
 
-TFuture<TOptional<UE::Interchange::FMeshPayloadData>> UInterchangeDatasmithTranslator::GetMeshPayloadData(const FInterchangeMeshPayLoadKey& PayLoadKey, const FTransform& MeshGlobalTransform) const
+TOptional<UE::Interchange::FMeshPayloadData> UInterchangeDatasmithTranslator::GetMeshPayloadData(const FInterchangeMeshPayLoadKey& PayLoadKey, const FTransform& MeshGlobalTransform) const
 {
 	using namespace UE::DatasmithInterchange;
 
-	TPromise<TOptional<UE::Interchange::FMeshPayloadData>> EmptyPromise;
-	EmptyPromise.SetValue(TOptional<UE::Interchange::FMeshPayloadData>());
+	TOptional<UE::Interchange::FMeshPayloadData> EmptyPayload = TOptional<UE::Interchange::FMeshPayloadData>();
 
 	if (!LoadedExternalSource || !LoadedExternalSource->GetDatasmithScene())
 	{
-		return EmptyPromise.GetFuture();
+		return EmptyPayload;
 	}
 
 	int32 MeshIndex = 0;
@@ -737,59 +737,33 @@ TFuture<TOptional<UE::Interchange::FMeshPayloadData>> UInterchangeDatasmithTrans
 	TSharedPtr<IDatasmithScene> DatasmithScene = LoadedExternalSource->GetDatasmithScene();
 	if (MeshIndex < 0 || MeshIndex >= DatasmithScene->GetMeshesCount())
 	{
-		return EmptyPromise.GetFuture();
+		return EmptyPayload;
 	}
 
 	TSharedPtr<IDatasmithMeshElement> MeshElement = DatasmithScene->GetMesh(MeshIndex);
 	if (!MeshElement.IsValid())
 	{
-		return EmptyPromise.GetFuture();
+		return EmptyPayload;
 	}
 
-	// If on the GameThread and Datasmith translator cannot be parallelized
-	// just run it on the spot.
-	if (IsInGameThread() && AsyncMode == EAsyncExecution::TaskGraphMainThread)
+	UE::Interchange::FMeshPayloadData StaticMeshPayloadData;
+	if (GetMeshDescription(MeshElement, MeshGlobalTransform, StaticMeshPayloadData))
 	{
-		TOptional<UE::Interchange::FMeshPayloadData> Result;
-
-		UE::Interchange::FMeshPayloadData StaticMeshPayloadData;
-		if (GetMeshDescription(MeshElement, MeshGlobalTransform, StaticMeshPayloadData))
-		{
-			Result.Emplace(MoveTemp(StaticMeshPayloadData));
-
-			TPromise<TOptional<UE::Interchange::FMeshPayloadData>> Promise;
-			Promise.SetValue(Result);
-
-			return Promise.GetFuture();
-		}
-
-		return EmptyPromise.GetFuture();
+		TOptional<UE::Interchange::FMeshPayloadData> Payload;
+		Payload = MoveTemp(StaticMeshPayloadData);
+		return Payload;
 	}
 
-	return Async(AsyncMode, [this, MeshElement = MoveTemp(MeshElement), MeshGlobalTransform]
-		{
-			TOptional<UE::Interchange::FMeshPayloadData> Result;
-
-			UE::Interchange::FMeshPayloadData StaticMeshPayloadData;
-			if (GetMeshDescription(MeshElement, MeshGlobalTransform, StaticMeshPayloadData))
-			{
-				Result.Emplace(MoveTemp(StaticMeshPayloadData));
-			}
-
-			return Result;
-		}
-	);
+	return EmptyPayload;
 }
 
-TFuture<TOptional<UE::Interchange::FAnimationPayloadData>> UInterchangeDatasmithTranslator::GetAnimationPayloadData(const UE::Interchange::FAnimationPayloadQuery& PayloadQuery) const
+TOptional<UE::Interchange::FAnimationPayloadData> UInterchangeDatasmithTranslator::GetAnimationPayloadData(const UE::Interchange::FAnimationPayloadQuery& PayloadQuery) const
 {
-	TPromise<TOptional<UE::Interchange::FAnimationPayloadData>> EmptyPromise;
-	EmptyPromise.SetValue(TOptional<UE::Interchange::FAnimationPayloadData>());
-	
+	TOptional<UE::Interchange::FAnimationPayloadData> Result;
 
 	if (!LoadedExternalSource || !LoadedExternalSource->GetDatasmithScene())
 	{
-		return EmptyPromise.GetFuture();
+		return Result;
 	}
 
 	TSharedPtr<IDatasmithBaseAnimationElement> AnimationElement;
@@ -800,7 +774,7 @@ TFuture<TOptional<UE::Interchange::FAnimationPayloadData>> UInterchangeDatasmith
 		if (!ensure(AnimationElement))
 		{
 			// #ueent_logwarning:
-			return EmptyPromise.GetFuture();
+			return Result;
 		}
 
 		FrameRate = PayloadDescPtr->Key;
@@ -808,37 +782,53 @@ TFuture<TOptional<UE::Interchange::FAnimationPayloadData>> UInterchangeDatasmith
 
 	if (PayloadQuery.PayloadKey.Type != EInterchangeAnimationPayLoadType::NONE)
 	{
-		return Async(EAsyncExecution::TaskGraph, [this, AnimationElement = MoveTemp(AnimationElement), FrameRate, PayloadQuery]
-			{
-
-				UE::Interchange::FAnimationPayloadData TransformPayloadData(PayloadQuery.SceneNodeUniqueID, PayloadQuery.PayloadKey);
-				TOptional<UE::Interchange::FAnimationPayloadData> Result;
-
-				if (UE::DatasmithInterchange::AnimUtils::GetAnimationPayloadData(*AnimationElement, FrameRate, PayloadQuery.PayloadKey.Type, TransformPayloadData))
-				{
-					Result.Emplace(MoveTemp(TransformPayloadData));
-				}
-
-				return Result;
-			}
-		);
+		UE::Interchange::FAnimationPayloadData TransformPayloadData(PayloadQuery.SceneNodeUniqueID, PayloadQuery.PayloadKey);
+		if (UE::DatasmithInterchange::AnimUtils::GetAnimationPayloadData(*AnimationElement, FrameRate, PayloadQuery.PayloadKey.Type, TransformPayloadData))
+		{
+			Result = MoveTemp(TransformPayloadData);
+		}
 	}
 
-	return EmptyPromise.GetFuture();
+	return Result;
 }
 
 TArray<UE::Interchange::FAnimationPayloadData> UInterchangeDatasmithTranslator::GetAnimationPayloadData(const TArray<UE::Interchange::FAnimationPayloadQuery>& PayloadQueries) const
 {
-	TArray<TFuture<TOptional<UE::Interchange::FAnimationPayloadData>>> AnimationPayloadFutures;
-	for (const UE::Interchange::FAnimationPayloadQuery& PayloadQuery : PayloadQueries)
+	TArray<TOptional<UE::Interchange::FAnimationPayloadData>> AnimationPayloadOptionals;
+	int32 PayloadCount = PayloadQueries.Num();
+	AnimationPayloadOptionals.AddDefaulted(PayloadCount);
+
+	const int32 BatchSize = 5;
+	if (PayloadQueries.Num() > BatchSize)
+	{
+		const int32 NumBatches = (PayloadCount / BatchSize) + 1;
+		ParallelFor(NumBatches, [&](int32 BatchIndex)
+			{
+				int32 PayloadIndexOffset = BatchIndex * BatchSize;
+				for (int32 PayloadIndex = PayloadIndexOffset; PayloadIndex < PayloadIndexOffset + BatchSize; ++PayloadIndex)
+				{
+					if (PayloadQueries.IsValidIndex(PayloadIndex))
+					{
+						AnimationPayloadOptionals[PayloadIndex] = GetAnimationPayloadData(PayloadQueries[PayloadIndex]);
+					}
+				}
+			}, EParallelForFlags::BackgroundPriority);// ParallelFor
+	}
+	else
 	{
-		AnimationPayloadFutures.Add(GetAnimationPayloadData(PayloadQuery));
+		for (int32 PayloadIndex = 0; PayloadIndex < PayloadCount; ++PayloadIndex)
+		{
+			if (PayloadQueries.IsValidIndex(PayloadIndex))
+			{
+				AnimationPayloadOptionals[PayloadIndex] = GetAnimationPayloadData(PayloadQueries[PayloadIndex]);
+			}
+		}
 	}
 
+
 	TArray<UE::Interchange::FAnimationPayloadData> AnimationPayloads;
-	for (TFuture<TOptional<UE::Interchange::FAnimationPayloadData>>& AnimationPayloadFuture : AnimationPayloadFutures)
+	for (TOptional<UE::Interchange::FAnimationPayloadData>& OptionalPayloadData : AnimationPayloadOptionals)
 	{
-		TOptional<UE::Interchange::FAnimationPayloadData> OptionalPayloadData = AnimationPayloadFuture.Get();
 		if (!OptionalPayloadData.IsSet())
 		{
 			continue;
@@ -849,17 +839,16 @@ TArray<UE::Interchange::FAnimationPayloadData> UInterchangeDatasmithTranslator::
 	return AnimationPayloads;
 }
 
-TFuture<TOptional<UE::Interchange::FVariantSetPayloadData>> UInterchangeDatasmithTranslator::GetVariantSetPayloadData(const FString& PayloadKey) const
+TOptional<UE::Interchange::FVariantSetPayloadData> UInterchangeDatasmithTranslator::GetVariantSetPayloadData(const FString& PayloadKey) const
 {
 	using namespace UE::Interchange;
 	using namespace UE::DatasmithInterchange;
 
-	TPromise<TOptional<FVariantSetPayloadData>> EmptyPromise;
-	EmptyPromise.SetValue(TOptional<FVariantSetPayloadData>());
+	TOptional<FVariantSetPayloadData> Result;
 
 	if (!LoadedExternalSource || !LoadedExternalSource->GetDatasmithScene())
 	{
-		return EmptyPromise.GetFuture();
+		return Result;
 	}
 
 	TSharedPtr<IDatasmithScene> DatasmithScene = LoadedExternalSource->GetDatasmithScene();
@@ -870,7 +859,7 @@ TFuture<TOptional<UE::Interchange::FVariantSetPayloadData>> UInterchangeDatasmit
 	if (2 != PayloadKey.ParseIntoArray(PayloadTokens, TEXT(";")))
 	{
 		// Invalid payload
-		return EmptyPromise.GetFuture();
+		return Result;
 	}
 
 	int32 LevelVariantSetIndex = FCString::Atoi(*PayloadTokens[0]);
@@ -882,23 +871,17 @@ TFuture<TOptional<UE::Interchange::FVariantSetPayloadData>> UInterchangeDatasmit
 		TSharedPtr<IDatasmithVariantSetElement> VariantSet = LevelVariantSetElement->GetVariantSet(VariantSetIndex);
 		if (ensure(VariantSet) && VariantSet->GetVariantsCount() > 0)
 		{
-			return Async(EAsyncExecution::TaskGraph, [this, VariantSet = MoveTemp(VariantSet)]
-					{
-						FVariantSetPayloadData PayloadData;
-						TOptional<FVariantSetPayloadData> Result;
-
-						if (VariantSetUtils::GetVariantSetPayloadData(*VariantSet, PayloadData))
-						{
-							Result.Emplace(MoveTemp(PayloadData));
-						}
-
-						return Result;
-					}
-				);
+			TSharedPtr<TPromise<TOptional<FVariantSetPayloadData>>> Promise = MakeShared<TPromise<TOptional<FVariantSetPayloadData>>>();
+			FVariantSetPayloadData PayloadData;
+			if (VariantSetUtils::GetVariantSetPayloadData(*VariantSet, PayloadData))
+			{
+				Result = MoveTemp(PayloadData);
+				return Result;
+			}
 		}
 	}
 
-	return EmptyPromise.GetFuture();
+	return Result;
 }
 
 void UInterchangeDatasmithTranslator::ImportFinish()
diff --git a/Engine/Plugins/Experimental/Enterprise/DatasmithInterchange/Source/DatasmithInterchange/Public/InterchangeDatasmithTranslator.h b/Engine/Plugins/Experimental/Enterprise/DatasmithInterchange/Source/DatasmithInterchange/Public/InterchangeDatasmithTranslator.h
--- a/Engine/Plugins/Experimental/Enterprise/DatasmithInterchange/Source/DatasmithInterchange/Public/InterchangeDatasmithTranslator.h
+++ b/Engine/Plugins/Experimental/Enterprise/DatasmithInterchange/Source/DatasmithInterchange/Public/InterchangeDatasmithTranslator.h
@@ -124,16 +124,16 @@ public:
 	/* IInterchangeTextureLightProfilePayloadInterface End */
 
 	/* IInterchangeStaticMeshPayloadInterface Begin */
-	virtual TFuture<TOptional<UE::Interchange::FMeshPayloadData>> GetMeshPayloadData(const FInterchangeMeshPayLoadKey& PayLoadKey, const FTransform& MeshGlobalTransform) const override;
+	virtual TOptional<UE::Interchange::FMeshPayloadData> GetMeshPayloadData(const FInterchangeMeshPayLoadKey& PayLoadKey, const FTransform& MeshGlobalTransform) const override;
 	/* IInterchangeStaticMeshPayloadInterface End */
 
 	/* IInterchangeAnimationPayloadInterface Begin */
-	TFuture<TOptional<UE::Interchange::FAnimationPayloadData>> GetAnimationPayloadData(const UE::Interchange::FAnimationPayloadQuery& PayloadQuery) const;
+	TOptional<UE::Interchange::FAnimationPayloadData> GetAnimationPayloadData(const UE::Interchange::FAnimationPayloadQuery& PayloadQuery) const;
 	virtual TArray<UE::Interchange::FAnimationPayloadData> GetAnimationPayloadData(const TArray<UE::Interchange::FAnimationPayloadQuery>& PayloadQueries) const override;
 	/* IInterchangeAnimationPayloadInterface End */
 
 	/* IInterchangeVariantSetPayloadInterface Begin */
-	virtual TFuture<TOptional<UE::Interchange::FVariantSetPayloadData>> GetVariantSetPayloadData(const FString& PayloadKey) const override;
+	virtual TOptional<UE::Interchange::FVariantSetPayloadData> GetVariantSetPayloadData(const FString& PayloadKey) const override;
 	/* IInterchangeVariantSetPayloadInterface End */
 
 private:
diff --git a/Engine/Plugins/Interchange/Runtime/Source/Import/Private/Animation/InterchangeAnimSequenceFactory.cpp b/Engine/Plugins/Interchange/Runtime/Source/Import/Private/Animation/InterchangeAnimSequenceFactory.cpp
--- a/Engine/Plugins/Interchange/Runtime/Source/Import/Private/Animation/InterchangeAnimSequenceFactory.cpp
+++ b/Engine/Plugins/Interchange/Runtime/Source/Import/Private/Animation/InterchangeAnimSequenceFactory.cpp
@@ -358,7 +358,11 @@ namespace UE::Interchange::Private
 			//First bake the channel weight curves
 			FRichCurve& ChannelWeightCurve = Curves[0];
 #if WITH_EDITORONLY_DATA
-			ChannelWeightCurve.BakeCurve(1.0f / TargetSequence->ImportResampleFramerate);
+			//Cannot bake a curve with only one frame
+			if (ChannelWeightCurve.GetNumKeys() > 1)
+			{
+				ChannelWeightCurve.BakeCurve(1.0f / TargetSequence->ImportResampleFramerate);
+			}
 #endif
 
 			// use the primary curve to generate inbetween shape curves + a modified primary curve
@@ -926,6 +930,98 @@ UClass* UInterchangeAnimSequenceFactory::GetFactoryClass() const
 	return UAnimSequence::StaticClass();
 }
 
+void UInterchangeAnimSequenceFactory::CreatePayloadTasks(const FImportAssetObjectParams& Arguments, bool bAsync, TArray<TSharedPtr<UE::Interchange::FInterchangeTaskBase>>& PayloadTasks)
+{
+	const IInterchangeAnimationPayloadInterface* AnimSequenceTranslatorPayloadInterface = Cast<IInterchangeAnimationPayloadInterface>(Arguments.Translator);
+	if (!AnimSequenceTranslatorPayloadInterface)
+	{
+		return;
+	}
+
+	UInterchangeAnimSequenceFactoryNode* AnimSequenceFactoryNode = Cast<UInterchangeAnimSequenceFactoryNode>(Arguments.AssetNode);
+	if (AnimSequenceFactoryNode == nullptr)
+	{
+		return;
+	}
+
+	bool bImportBoneTracks = false;
+	AnimSequenceFactoryNode->GetCustomImportBoneTracks(bImportBoneTracks);
+	if (bImportBoneTracks)
+	{
+		//Get the sample rate, default to 30Hz in case the attribute is missing
+		double SampleRate = 30.0;
+		AnimSequenceFactoryNode->GetCustomImportBoneTracksSampleRate(SampleRate);
+
+		double RangeStart = 0.0;
+		AnimSequenceFactoryNode->GetCustomImportBoneTracksRangeStart(RangeStart);
+
+		double RangeEnd = 1.0 / SampleRate; //One frame duration per default
+		AnimSequenceFactoryNode->GetCustomImportBoneTracksRangeStop(RangeEnd);
+
+
+		TMap<FString, FInterchangeAnimationPayLoadKey> PayloadKeys;
+		AnimSequenceFactoryNode->GetSceneNodeAnimationPayloadKeys(PayloadKeys);
+		BoneAnimationPayloadQueries.Empty(PayloadKeys.Num());
+		BoneAnimationPayloadResults.Empty(PayloadKeys.Num());
+
+		if (AnimSequenceTranslatorPayloadInterface->PreferGroupingBoneAnimationQueriesTogether())
+		{
+			//Create one task for all queries
+			for (const TTuple<FString, FInterchangeAnimationPayLoadKey>& SceneNodeUidAndPayloadKey : PayloadKeys)
+			{
+				UE::Interchange::FAnimationPayloadQuery AnimationPayloadQuery(SceneNodeUidAndPayloadKey.Key, SceneNodeUidAndPayloadKey.Value, SampleRate, RangeStart, RangeEnd);
+				BoneAnimationPayloadQueries.Add(SceneNodeUidAndPayloadKey, AnimationPayloadQuery);
+				//Allocate the results so we do not need any mutex
+				BoneAnimationPayloadResults.FindOrAdd(SceneNodeUidAndPayloadKey, UE::Interchange::FAnimationPayloadData(SceneNodeUidAndPayloadKey.Key, SceneNodeUidAndPayloadKey.Value));
+			}
+			
+			TSharedPtr<UE::Interchange::FInterchangeTaskLambda, ESPMode::ThreadSafe> TaskGetAnimationPayload = MakeShared<UE::Interchange::FInterchangeTaskLambda, ESPMode::ThreadSafe>(bAsync ? UE::Interchange::EInterchangeTaskThread::AsyncThread : UE::Interchange::EInterchangeTaskThread::GameThread
+				, [this, AnimSequenceTranslatorPayloadInterface]()
+				{
+					TArray<UE::Interchange::FAnimationPayloadQuery> TmpPayloadQueries;
+					BoneAnimationPayloadQueries.GenerateValueArray(TmpPayloadQueries);
+					TArray<UE::Interchange::FAnimationPayloadData> AnimationPayloads = AnimSequenceTranslatorPayloadInterface->GetAnimationPayloadData(TmpPayloadQueries);
+					if (ensure(AnimationPayloads.Num() == BoneAnimationPayloadQueries.Num()))
+					{
+						TArray<TTuple<FString, FInterchangeAnimationPayLoadKey>> TmpPayloadKeys;
+						BoneAnimationPayloadQueries.GenerateKeyArray(TmpPayloadKeys);
+						for (int32 AnimationIndex = 0; AnimationIndex < AnimationPayloads.Num(); ++AnimationIndex)
+						{
+							if (ensure(BoneAnimationPayloadResults.Contains(TmpPayloadKeys[AnimationIndex])))
+							{
+								BoneAnimationPayloadResults.FindChecked(TmpPayloadKeys[AnimationIndex]) = AnimationPayloads[AnimationIndex];
+							}
+						}
+					}
+				});
+			PayloadTasks.Add(TaskGetAnimationPayload);
+		}
+		else
+		{
+			//Create one task per query
+			for (const TTuple<FString, FInterchangeAnimationPayLoadKey>& SceneNodeUidAndPayloadKey : PayloadKeys)
+			{
+				UE::Interchange::FAnimationPayloadQuery AnimationPayloadQuery(SceneNodeUidAndPayloadKey.Key, SceneNodeUidAndPayloadKey.Value, SampleRate, RangeStart, RangeEnd);
+				BoneAnimationPayloadQueries.Add(SceneNodeUidAndPayloadKey, AnimationPayloadQuery);
+				//Allocate the results so we do not need any mutex
+				BoneAnimationPayloadResults.FindOrAdd(SceneNodeUidAndPayloadKey, UE::Interchange::FAnimationPayloadData(SceneNodeUidAndPayloadKey.Key, SceneNodeUidAndPayloadKey.Value));
+
+				TSharedPtr<UE::Interchange::FInterchangeTaskLambda, ESPMode::ThreadSafe> TaskGetAnimationPayload = MakeShared<UE::Interchange::FInterchangeTaskLambda, ESPMode::ThreadSafe>(bAsync ? UE::Interchange::EInterchangeTaskThread::AsyncThread : UE::Interchange::EInterchangeTaskThread::GameThread
+					, [this, SceneNodeUidAndPayloadKey, AnimSequenceTranslatorPayloadInterface]()
+					{
+						UE::Interchange::FAnimationPayloadQuery& AnimationQuery = BoneAnimationPayloadQueries.FindChecked(SceneNodeUidAndPayloadKey);
+						TArray<UE::Interchange::FAnimationPayloadData> AnimationPayloads = AnimSequenceTranslatorPayloadInterface->GetAnimationPayloadData({ AnimationQuery });
+						if (ensure(AnimationPayloads.IsValidIndex(0) && BoneAnimationPayloadResults.Contains(SceneNodeUidAndPayloadKey)))
+						{
+							BoneAnimationPayloadResults.FindChecked(SceneNodeUidAndPayloadKey) = AnimationPayloads[0];
+						}
+					});
+				PayloadTasks.Add(TaskGetAnimationPayload);
+			}
+		}
+	}
+}
+
 UInterchangeFactoryBase::FImportAssetResult UInterchangeAnimSequenceFactory::BeginImportAsset_GameThread(const FImportAssetObjectParams& Arguments)
 {
 	TRACE_CPUPROFILER_EVENT_SCOPE(UInterchangeAnimSequenceFactory::BeginImportAsset_GameThread);
@@ -1101,24 +1197,20 @@ UInterchangeFactoryBase::FImportAssetResult UInterchangeAnimSequenceFactory::Imp
 		double RangeEnd = 1.0 / SampleRate; //One frame duration per default
 		AnimSequenceFactoryNode->GetCustomImportBoneTracksRangeStop(RangeEnd);
 
-
 		TMap<FString, FInterchangeAnimationPayLoadKey> PayloadKeys;
 		AnimSequenceFactoryNode->GetSceneNodeAnimationPayloadKeys(PayloadKeys);
 
-		TArray<UE::Interchange::FAnimationPayloadQuery> PayloadQueries;
-
-		for (const TTuple<FString, FInterchangeAnimationPayLoadKey>& SceneNodeUidAndPayloadKey : PayloadKeys)
-		{
-			PayloadQueries.Add(UE::Interchange::FAnimationPayloadQuery(SceneNodeUidAndPayloadKey.Key, SceneNodeUidAndPayloadKey.Value, SampleRate, RangeStart, RangeEnd));
-		}
-
-		TArray<UE::Interchange::FAnimationPayloadData> AnimationPayloads = AnimSequenceTranslatorPayloadInterface->GetAnimationPayloadData(PayloadQueries);
-		
 		BoneTrackData.MergedRangeEnd = RangeEnd;
 		BoneTrackData.MergedRangeStart = RangeStart;
 		BoneTrackData.PreProcessedAnimationPayloads.Reset();
-		for (UE::Interchange::FAnimationPayloadData& AnimationPayload : AnimationPayloads)
+		for (const TTuple<FString, FInterchangeAnimationPayLoadKey>& SceneNodeUidAndPayloadKey : PayloadKeys)
 		{
+			if (!BoneAnimationPayloadResults.Contains(SceneNodeUidAndPayloadKey))
+			{
+				continue;
+			}
+			UE::Interchange::FAnimationPayloadData& AnimationPayload = BoneAnimationPayloadResults.FindChecked(SceneNodeUidAndPayloadKey);
+
 			const UInterchangeSceneNode* SkeletonSceneNode = Cast<UInterchangeSceneNode>(Arguments.NodeContainer->GetNode(AnimationPayload.SceneNodeUniqueID));
 			if (!SkeletonSceneNode)
 			{
@@ -1161,7 +1253,7 @@ UInterchangeFactoryBase::FImportAssetResult UInterchangeAnimSequenceFactory::Imp
 		AnimSequenceFactoryNode->GetMorphTargetNodeAnimationPayloadKeys(MorphTargetNodeAnimationPayloads);
 
 		TArray<UE::Interchange::FAnimationPayloadQuery> PayloadQueries;
-		TMap<FString, TFuture<TOptional<UE::Interchange::FAnimationPayloadData>>> MorphTargetCurveWeightInstanceAnimationPayloads;
+		TMap<FString, TOptional<UE::Interchange::FAnimationPayloadData>> MorphTargetCurveWeightInstanceAnimationPayloads;
 
 		for (const TPair<FString, FInterchangeAnimationPayLoadKey>& MorphTargetNodeUidAnimationPayload : MorphTargetNodeAnimationPayloads)
 		{
@@ -1174,30 +1266,21 @@ UInterchangeFactoryBase::FImportAssetResult UInterchangeAnimSequenceFactory::Imp
 			{
 				if (MorphTargetNodeUidAnimationPayload.Value.Type == EInterchangeAnimationPayLoadType::MORPHTARGETCURVEWEIGHTINSTANCE)
 				{
-					MorphTargetCurveWeightInstanceAnimationPayloads.Add(PayloadKey, Async(EAsyncExecution::TaskGraph, [&MorphTargetNodeUidAnimationPayload]
-						{
-							TOptional<UE::Interchange::FAnimationPayloadData> Result;
-							UE::Interchange::FAnimationPayloadData AnimationPayLoadData(MorphTargetNodeUidAnimationPayload.Key, MorphTargetNodeUidAnimationPayload.Value);
-
-							TArray<FString> PayLoadKeys;
-							MorphTargetNodeUidAnimationPayload.Value.UniqueId.ParseIntoArray(PayLoadKeys, TEXT(":"));
-
-							if (PayLoadKeys.Num() != 2)
-							{
-								return Result;
-							}
-
-							float Weight;
-							LexFromString(Weight, *PayLoadKeys[1]);
-
-							AnimationPayLoadData.Curves.SetNum(1);
-							FRichCurve& Curve = AnimationPayLoadData.Curves[0];
-							Curve.AddKey(0, Weight);
+					TOptional<UE::Interchange::FAnimationPayloadData>& Result = MorphTargetCurveWeightInstanceAnimationPayloads.FindOrAdd(PayloadKey);
+					UE::Interchange::FAnimationPayloadData AnimationPayLoadData(MorphTargetNodeUidAnimationPayload.Key, MorphTargetNodeUidAnimationPayload.Value);
+					TArray<FString> PayLoadKeys;
+					MorphTargetNodeUidAnimationPayload.Value.UniqueId.ParseIntoArray(PayLoadKeys, TEXT(":"));
+					if (PayLoadKeys.Num() == 2)
+					{
+						float Weight;
+						LexFromString(Weight, *PayLoadKeys[1]);
 
-							Result.Emplace(AnimationPayLoadData);
+						AnimationPayLoadData.Curves.SetNum(1);
+						FRichCurve& Curve = AnimationPayLoadData.Curves[0];
+						Curve.AddKey(0, Weight);
 
-							return Result;
-						}));
+						Result.Emplace(AnimationPayLoadData);
+					}
 				}
 				else
 				{
@@ -1215,9 +1298,9 @@ UInterchangeFactoryBase::FImportAssetResult UInterchangeAnimSequenceFactory::Imp
 		}
 		AnimationCurvesPayloads.Empty();
 
-		for (TPair<FString, TFuture<TOptional<UE::Interchange::FAnimationPayloadData>>>& CurveNameAndPayload : MorphTargetCurveWeightInstanceAnimationPayloads)
+		for (TPair<FString, TOptional<UE::Interchange::FAnimationPayloadData>>& CurveNameAndPayload : MorphTargetCurveWeightInstanceAnimationPayloads)
 		{
-			TOptional<UE::Interchange::FAnimationPayloadData> AnimationCurvePayload = CurveNameAndPayload.Value.Get();
+			TOptional<UE::Interchange::FAnimationPayloadData> AnimationCurvePayload = CurveNameAndPayload.Value;
 			if (!AnimationCurvePayload.IsSet())
 			{
 				UE_LOG(LogInterchangeImport, Warning, TEXT("Invalid animation morph target curve payload key [%s] in AnimSequence asset %s."), *CurveNameAndPayload.Key, *Arguments.AssetName);
diff --git a/Engine/Plugins/Interchange/Runtime/Source/Import/Private/Fbx/InterchangeFbxTranslator.cpp b/Engine/Plugins/Interchange/Runtime/Source/Import/Private/Fbx/InterchangeFbxTranslator.cpp
--- a/Engine/Plugins/Interchange/Runtime/Source/Import/Private/Fbx/InterchangeFbxTranslator.cpp
+++ b/Engine/Plugins/Interchange/Runtime/Source/Import/Private/Fbx/InterchangeFbxTranslator.cpp
@@ -329,7 +329,7 @@ TOptional<UE::Interchange::FImportImage> UInterchangeFbxTranslator::GetTexturePa
 	return TextureTranslator->GetTexturePayloadData(PayLoadKey, AlternateTexturePath);
 }
 
-TFuture<TOptional<UE::Interchange::FMeshPayloadData>> UInterchangeFbxTranslator::GetMeshPayloadData(const FInterchangeMeshPayLoadKey& PayLoadKey, const FTransform& MeshGlobalTransform) const
+TOptional<UE::Interchange::FMeshPayloadData> UInterchangeFbxTranslator::GetMeshPayloadData(const FInterchangeMeshPayLoadKey& PayLoadKey, const FTransform& MeshGlobalTransform) const
 {
 	TSharedPtr<TPromise<TOptional<UE::Interchange::FMeshPayloadData>>> Promise = MakeShared<TPromise<TOptional<UE::Interchange::FMeshPayloadData>>>();
 
@@ -418,8 +418,7 @@ TFuture<TOptional<UE::Interchange::FMeshPayloadData>> UInterchangeFbxTranslator:
 	{
 		if (!Dispatcher.IsValid())
 		{
-			Promise->SetValue(TOptional<UE::Interchange::FMeshPayloadData>());
-			return Promise->GetFuture();
+			return TOptional<UE::Interchange::FMeshPayloadData>();
 		}
 
 		// Create a json command to read the fbx file
@@ -458,7 +457,8 @@ TFuture<TOptional<UE::Interchange::FMeshPayloadData>> UInterchangeFbxTranslator:
 			Promise->SetValue(TOptional<UE::Interchange::FMeshPayloadData>{});
 		}
 	}
-	return Promise->GetFuture();
+	//We must stall since the payload async is now control by the interchange task system
+	return Promise->GetFuture().Get();
 }
 
 TArray<UE::Interchange::FAnimationPayloadData> UInterchangeFbxTranslator::GetAnimationPayloadData(const TArray<UE::Interchange::FAnimationPayloadQuery>& PayloadQueries) const
diff --git a/Engine/Plugins/Interchange/Runtime/Source/Import/Private/Gltf/InterchangeGltfTranslator.cpp b/Engine/Plugins/Interchange/Runtime/Source/Import/Private/Gltf/InterchangeGltfTranslator.cpp
--- a/Engine/Plugins/Interchange/Runtime/Source/Import/Private/Gltf/InterchangeGltfTranslator.cpp
+++ b/Engine/Plugins/Interchange/Runtime/Source/Import/Private/Gltf/InterchangeGltfTranslator.cpp
@@ -30,7 +30,7 @@
 #include "Texture/InterchangeImageWrapperTranslator.h"
 
 #include "Algo/Find.h"
-#include "Async/Async.h"
+#include "Async/ParallelFor.h"
 #include "Misc/App.h"
 #include "StaticMeshAttributes.h"
 #include "SkeletalMeshAttributes.h"
@@ -906,60 +906,78 @@ TOptional< UE::Interchange::FImportImage > UInterchangeGLTFTranslator::GetTextur
 	return TexturePayloadData;
 }
 
-TFuture<TOptional<UE::Interchange::FAnimationPayloadData>> UInterchangeGLTFTranslator::GetAnimationPayloadData(const UE::Interchange::FAnimationPayloadQuery& PayloadQuery) const
+TOptional<UE::Interchange::FAnimationPayloadData> UInterchangeGLTFTranslator::GetAnimationPayloadData(const UE::Interchange::FAnimationPayloadQuery& PayloadQuery) const
 {
-	return Async(EAsyncExecution::TaskGraph, [this, PayloadQuery]
+	UE::Interchange::FAnimationPayloadData AnimationPayLoadData(PayloadQuery.SceneNodeUniqueID, PayloadQuery.PayloadKey);
+	TOptional<UE::Interchange::FAnimationPayloadData> Result;
+	switch (PayloadQuery.PayloadKey.Type)
+	{
+	case EInterchangeAnimationPayLoadType::CURVE:
+		if (UE::Interchange::Gltf::Private::GetTransformAnimationPayloadData(PayloadQuery.PayloadKey.UniqueId, GltfAsset, AnimationPayLoadData))
+		{
+			Result = AnimationPayLoadData;
+		}
+		break;
+	case EInterchangeAnimationPayLoadType::MORPHTARGETCURVE:
+		if (UE::Interchange::Gltf::Private::GetMorphTargetAnimationPayloadData(PayloadQuery.PayloadKey.UniqueId, GltfAsset, AnimationPayLoadData))
+		{
+			Result = AnimationPayLoadData;
+		}
+		break;
+	case EInterchangeAnimationPayLoadType::BAKED:
+		AnimationPayLoadData.BakeFrequency = PayloadQuery.TimeDescription.BakeFrequency;
+		AnimationPayLoadData.RangeStartTime = PayloadQuery.TimeDescription.RangeStartSecond;
+		AnimationPayLoadData.RangeEndTime = PayloadQuery.TimeDescription.RangeStopSecond;
+		if (UE::Interchange::Gltf::Private::GetBakedAnimationTransformPayloadData(PayloadQuery.PayloadKey.UniqueId, GltfAsset, AnimationPayLoadData))
 		{
+			Result = AnimationPayLoadData;
+		}
+		break;
+	case EInterchangeAnimationPayLoadType::STEPCURVE:
+	case EInterchangeAnimationPayLoadType::NONE:
+	default:
+		break;
+	}
 
-			TOptional<UE::Interchange::FAnimationPayloadData> Result;
-			UE::Interchange::FAnimationPayloadData AnimationPayLoadData(PayloadQuery.SceneNodeUniqueID, PayloadQuery.PayloadKey);
+	return Result;
+}
 
-			switch (PayloadQuery.PayloadKey.Type)
+TArray<UE::Interchange::FAnimationPayloadData> UInterchangeGLTFTranslator::GetAnimationPayloadData(const TArray<UE::Interchange::FAnimationPayloadQuery>& PayloadQueries) const
+{
+	TArray<TOptional<UE::Interchange::FAnimationPayloadData>> AnimationPayloadOptionals;
+	int32 PayloadCount = PayloadQueries.Num();
+	AnimationPayloadOptionals.AddDefaulted(PayloadCount);
+	
+	const int32 BatchSize = 10;
+	if (PayloadQueries.Num() > BatchSize)
+	{
+		const int32 NumBatches = (PayloadCount / BatchSize) + 1;
+		ParallelFor(NumBatches, [&](int32 BatchIndex)
 			{
-			case EInterchangeAnimationPayLoadType::CURVE:
-				if (UE::Interchange::Gltf::Private::GetTransformAnimationPayloadData(PayloadQuery.PayloadKey.UniqueId, GltfAsset, AnimationPayLoadData))
-				{
-					Result.Emplace(AnimationPayLoadData);
-				}
-				break;
-			case EInterchangeAnimationPayLoadType::MORPHTARGETCURVE:
-				if (UE::Interchange::Gltf::Private::GetMorphTargetAnimationPayloadData(PayloadQuery.PayloadKey.UniqueId, GltfAsset, AnimationPayLoadData))
-				{
-					Result.Emplace(AnimationPayLoadData);
-				}
-				break;
-			case EInterchangeAnimationPayLoadType::BAKED:
-				AnimationPayLoadData.BakeFrequency = PayloadQuery.TimeDescription.BakeFrequency;
-				AnimationPayLoadData.RangeStartTime = PayloadQuery.TimeDescription.RangeStartSecond;
-				AnimationPayLoadData.RangeEndTime = PayloadQuery.TimeDescription.RangeStopSecond;
-				if (UE::Interchange::Gltf::Private::GetBakedAnimationTransformPayloadData(PayloadQuery.PayloadKey.UniqueId, GltfAsset, AnimationPayLoadData))
+				int32 PayloadIndexOffset = BatchIndex * BatchSize;
+				for (int32 PayloadIndex = PayloadIndexOffset; PayloadIndex < PayloadIndexOffset + BatchSize; ++PayloadIndex)
 				{
-					Result.Emplace(AnimationPayLoadData);
+					if (PayloadQueries.IsValidIndex(PayloadIndex))
+					{
+						AnimationPayloadOptionals[PayloadIndex] = GetAnimationPayloadData(PayloadQueries[PayloadIndex]);
+					}
 				}
-				break;
-			case EInterchangeAnimationPayLoadType::STEPCURVE:
-			case EInterchangeAnimationPayLoadType::NONE:
-			default:
-				break;
+			}, EParallelForFlags::BackgroundPriority);// ParallelFor
+	}
+	else
+	{
+		for (int32 PayloadIndex = 0; PayloadIndex < PayloadCount; ++PayloadIndex)
+		{
+			if (PayloadQueries.IsValidIndex(PayloadIndex))
+			{
+				AnimationPayloadOptionals[PayloadIndex] = GetAnimationPayloadData(PayloadQueries[PayloadIndex]);
 			}
-
-			return Result;
 		}
-	);
-}
-
-TArray<UE::Interchange::FAnimationPayloadData> UInterchangeGLTFTranslator::GetAnimationPayloadData(const TArray<UE::Interchange::FAnimationPayloadQuery>& PayloadQueries) const
-{
-	TArray<TFuture<TOptional<UE::Interchange::FAnimationPayloadData>>> AnimationPayloadFutures;
-	for (const UE::Interchange::FAnimationPayloadQuery& PayloadQuery: PayloadQueries)
-	{
-		AnimationPayloadFutures.Add(GetAnimationPayloadData(PayloadQuery));
 	}
 
 	TArray<UE::Interchange::FAnimationPayloadData> AnimationPayloads;
-	for (TFuture<TOptional<UE::Interchange::FAnimationPayloadData>>& AnimationPayloadFuture : AnimationPayloadFutures)
+	for (TOptional<UE::Interchange::FAnimationPayloadData>& OptionalPayloadData : AnimationPayloadOptionals)
 	{
-		TOptional<UE::Interchange::FAnimationPayloadData> OptionalPayloadData = AnimationPayloadFuture.Get();
 		if (!OptionalPayloadData.IsSet())
 		{
 			continue;
@@ -990,12 +1008,11 @@ void UInterchangeGLTFTranslator::SetTextureSRGB(UInterchangeBaseNodeContainer& N
 	}
 }
 
-TFuture<TOptional<UE::Interchange::FVariantSetPayloadData>> UInterchangeGLTFTranslator::GetVariantSetPayloadData(const FString& PayloadKey) const
+TOptional<UE::Interchange::FVariantSetPayloadData> UInterchangeGLTFTranslator::GetVariantSetPayloadData(const FString& PayloadKey) const
 {
 	using namespace UE::Interchange;
 
-	TPromise<TOptional<FVariantSetPayloadData>> EmptyPromise;
-	EmptyPromise.SetValue(TOptional<FVariantSetPayloadData>());
+	TOptional<FVariantSetPayloadData> Result;
 
 	TArray<FString> PayloadTokens;
 
@@ -1003,7 +1020,7 @@ TFuture<TOptional<UE::Interchange::FVariantSetPayloadData>> UInterchangeGLTFTran
 	if (GltfAsset.Variants.Num() + 1 != PayloadKey.ParseIntoArray(PayloadTokens, TEXT(";")))
 	{
 		// Invalid payload
-		return EmptyPromise.GetFuture();
+		return Result;
 	}
 
 	//FString PayloadKey = FileName;
@@ -1012,23 +1029,17 @@ TFuture<TOptional<UE::Interchange::FVariantSetPayloadData>> UInterchangeGLTFTran
 		if (PayloadTokens[Index + 1] != GltfAsset.Variants[Index])
 		{
 			// Invalid payload
-			return EmptyPromise.GetFuture();
+			return Result;
 		}
 	}
 
-	return Async(EAsyncExecution::TaskGraph, [this]
-			{
-				FVariantSetPayloadData PayloadData;
-				TOptional<FVariantSetPayloadData> Result;
-
-				if (this->GetVariantSetPayloadData(PayloadData))
-				{
-					Result.Emplace(MoveTemp(PayloadData));
-				}
+	FVariantSetPayloadData PayloadData;
+	if (this->GetVariantSetPayloadData(PayloadData))
+	{
+		Result = MoveTemp(PayloadData);
+	}
 
-				return Result;
-			}
-		);
+	return Result;
 }
 
 void UInterchangeGLTFTranslator::HandleGltfVariants(UInterchangeBaseNodeContainer& NodeContainer, const FString& FileName) const
@@ -1208,45 +1219,39 @@ bool UInterchangeGLTFTranslator::GetVariantSetPayloadData(UE::Interchange::FVari
 	return true;
 }
 
-TFuture< TOptional< UE::Interchange::FMeshPayloadData > > UInterchangeGLTFTranslator::GetMeshPayloadData(const FInterchangeMeshPayLoadKey& PayLoadKey, const FTransform& MeshGlobalTransform) const
+TOptional< UE::Interchange::FMeshPayloadData > UInterchangeGLTFTranslator::GetMeshPayloadData(const FInterchangeMeshPayLoadKey& PayLoadKey, const FTransform& MeshGlobalTransform) const
 {
-	return Async(EAsyncExecution::TaskGraph, [this, PayLoadKey, MeshGlobalTransform]
-		{
-			UE::Interchange::FMeshPayloadData MeshPayLoadData;
-			bool bSuccessfullAcquisition = false;
+	UE::Interchange::FMeshPayloadData MeshPayloadData;
+	bool bSuccessfullAcquisition = false;
 
-			switch (PayLoadKey.Type)
-			{
-			case EInterchangeMeshPayLoadType::STATIC:
-				bSuccessfullAcquisition = UE::Interchange::Gltf::Private::GetStaticMeshPayloadDataForPayLoadKey(GltfAsset, PayLoadKey.UniqueId, MeshGlobalTransform, MeshPayLoadData.MeshDescription);
-				break;
-			case EInterchangeMeshPayLoadType::SKELETAL:
-				bSuccessfullAcquisition = UE::Interchange::Gltf::Private::GetSkeletalMeshDescriptionForPayLoadKey(GltfAsset, PayLoadKey.UniqueId, MeshGlobalTransform, MeshPayLoadData.MeshDescription, &MeshPayLoadData.JointNames);
-				break;
-			case EInterchangeMeshPayLoadType::MORPHTARGET:
-				//GLTF handles morph targets as simple Meshes
-				bSuccessfullAcquisition = UE::Interchange::Gltf::Private::GetStaticMeshPayloadDataForPayLoadKey(GltfAsset, PayLoadKey.UniqueId, MeshGlobalTransform, MeshPayLoadData.MeshDescription);
-				break;
-			case EInterchangeMeshPayLoadType::NONE:
-			default:
-				break;
-			}
-
-			TOptional<UE::Interchange::FMeshPayloadData> Result;
-			if (bSuccessfullAcquisition)
-			{
-				if (!FStaticMeshOperations::ValidateAndFixData(MeshPayLoadData.MeshDescription, PayLoadKey.UniqueId))
-				{
-					UInterchangeResultError_Generic* ErrorResult = AddMessage<UInterchangeResultError_Generic>();
-					ErrorResult->SourceAssetName = SourceData ? SourceData->GetFilename() : FString();
-					ErrorResult->Text = LOCTEXT("GetMeshPayloadData_ValidateMeshDescriptionFail", "Invalid mesh data (NAN) was found and changed to zero. This may affect the mesh rendering.");
-				}
-
-				Result.Emplace(MeshPayLoadData);
-			}
+	switch (PayLoadKey.Type)
+	{
+	case EInterchangeMeshPayLoadType::STATIC:
+		bSuccessfullAcquisition = UE::Interchange::Gltf::Private::GetStaticMeshPayloadDataForPayLoadKey(GltfAsset, PayLoadKey.UniqueId, MeshGlobalTransform, MeshPayloadData.MeshDescription);
+		break;
+	case EInterchangeMeshPayLoadType::SKELETAL:
+		bSuccessfullAcquisition = UE::Interchange::Gltf::Private::GetSkeletalMeshDescriptionForPayLoadKey(GltfAsset, PayLoadKey.UniqueId, MeshGlobalTransform, MeshPayloadData.MeshDescription, &MeshPayloadData.JointNames);
+		break;
+	case EInterchangeMeshPayLoadType::MORPHTARGET:
+		//GLTF handles morph targets as simple Meshes
+		bSuccessfullAcquisition = UE::Interchange::Gltf::Private::GetStaticMeshPayloadDataForPayLoadKey(GltfAsset, PayLoadKey.UniqueId, MeshGlobalTransform, MeshPayloadData.MeshDescription);
+		break;
+	case EInterchangeMeshPayLoadType::NONE:
+	default:
+		break;
+	}
 
-			return Result;
-		});
+	if (bSuccessfullAcquisition)
+	{
+		if (!FStaticMeshOperations::ValidateAndFixData(MeshPayloadData.MeshDescription, PayLoadKey.UniqueId))
+		{
+			UInterchangeResultError_Generic* ErrorResult = AddMessage<UInterchangeResultError_Generic>();
+			ErrorResult->SourceAssetName = SourceData ? SourceData->GetFilename() : FString();
+			ErrorResult->Text = LOCTEXT("GetMeshPayloadData_ValidateMeshDescriptionFail", "Invalid mesh data (NAN) was found and changed to zero. This may affect the mesh rendering.");
+		}
+		return MeshPayloadData;
+	}
+	return TOptional<UE::Interchange::FMeshPayloadData>();
 }
 
 void UInterchangeGLTFTranslator::HandleGltfSkeletons(UInterchangeBaseNodeContainer& NodeContainer, const FString& SceneNodeUid, const TArray<int32>& SkinnedMeshNodes, TSet<int>& UnusedMeshIndices) const
diff --git a/Engine/Plugins/Interchange/Runtime/Source/Import/Private/Mesh/InterchangeOBJTranslator.cpp b/Engine/Plugins/Interchange/Runtime/Source/Import/Private/Mesh/InterchangeOBJTranslator.cpp
--- a/Engine/Plugins/Interchange/Runtime/Source/Import/Private/Mesh/InterchangeOBJTranslator.cpp
+++ b/Engine/Plugins/Interchange/Runtime/Source/Import/Private/Mesh/InterchangeOBJTranslator.cpp
@@ -1581,25 +1581,20 @@ bool UInterchangeOBJTranslator::Translate(UInterchangeBaseNodeContainer& BaseNod
 }
 
 
-TFuture<TOptional<UE::Interchange::FMeshPayloadData>> UInterchangeOBJTranslator::GetMeshPayloadData(const FInterchangeMeshPayLoadKey& PayLoadKey, const FTransform& MeshGlobalTransform) const
+TOptional<UE::Interchange::FMeshPayloadData> UInterchangeOBJTranslator::GetMeshPayloadData(const FInterchangeMeshPayLoadKey& PayLoadKey, const FTransform& MeshGlobalTransform) const
 {
-	return Async(EAsyncExecution::TaskGraph, [this, PayLoadKey, MeshGlobalTransform]
-		{
-			using namespace UE::Interchange;
-
-			FMeshPayloadData Payload;
-			Payload.MeshDescription = ObjDataPtr->MakeMeshDescriptionForGroup(PayLoadKey.UniqueId, MeshGlobalTransform);
+	using namespace UE::Interchange;
 
-			if (!FStaticMeshOperations::ValidateAndFixData(Payload.MeshDescription, PayLoadKey.UniqueId))
-			{
-				UInterchangeResultError_Generic* ErrorResult = AddMessage<UInterchangeResultError_Generic>();
-				ErrorResult->SourceAssetName = SourceData ? SourceData->GetFilename() : FString();
-				ErrorResult->Text = LOCTEXT("GetMeshPayloadData_ValidateMeshDescriptionFail", "Invalid mesh data (NAN) was found and fix to zero. Mesh render can be bad.");
-			}
+	FMeshPayloadData Payload;
+	Payload.MeshDescription = ObjDataPtr->MakeMeshDescriptionForGroup(PayLoadKey.UniqueId, MeshGlobalTransform);
 
-			return TOptional<FMeshPayloadData>(Payload);
-		}
-	);
+	if (!FStaticMeshOperations::ValidateAndFixData(Payload.MeshDescription, PayLoadKey.UniqueId))
+	{
+		UInterchangeResultError_Generic* ErrorResult = AddMessage<UInterchangeResultError_Generic>();
+		ErrorResult->SourceAssetName = SourceData ? SourceData->GetFilename() : FString();
+		ErrorResult->Text = LOCTEXT("GetMeshPayloadData_ValidateMeshDescriptionFail", "Invalid mesh data (NAN) was found and fix to zero. Mesh render can be bad.");
+	}
+	return Payload;
 }
 
 
diff --git a/Engine/Plugins/Interchange/Runtime/Source/Import/Private/Mesh/InterchangeSkeletalMeshFactory.cpp b/Engine/Plugins/Interchange/Runtime/Source/Import/Private/Mesh/InterchangeSkeletalMeshFactory.cpp
--- a/Engine/Plugins/Interchange/Runtime/Source/Import/Private/Mesh/InterchangeSkeletalMeshFactory.cpp
+++ b/Engine/Plugins/Interchange/Runtime/Source/Import/Private/Mesh/InterchangeSkeletalMeshFactory.cpp
@@ -75,6 +75,26 @@ namespace UE
 				const UInterchangeSceneNode* SceneNode = nullptr;
 				TOptional<FTransform> SceneGlobalTransform;
 				FInterchangeMeshPayLoadKey TranslatorPayloadKey;
+
+				//Return a new key with the translator key merge with the transform
+				FInterchangeMeshPayLoadKey GetTranslatorAndTransformPayloadKey() const
+				{
+					FInterchangeMeshPayLoadKey GlobalPayloadKey = TranslatorPayloadKey;
+					GlobalPayloadKey.UniqueId = GetUniqueId();
+					return GlobalPayloadKey;
+				}
+
+				FInterchangeMeshPayLoadKey GetMorphTargetAndTransformPayloadKey(const FInterchangeMeshPayLoadKey& MorphTargetKey) const
+				{
+					FInterchangeMeshPayLoadKey GlobalPayloadKey = MorphTargetKey;
+					if (SceneGlobalTransform.IsSet())
+					{
+						GlobalPayloadKey.UniqueId += SceneGlobalTransform->ToString();
+					}
+					return GlobalPayloadKey;
+				}
+				
+				//Return the translator key merge with the transform
 				FString GetUniqueId() const
 				{
 					FString UniqueId = TranslatorPayloadKey.UniqueId;
@@ -88,7 +108,7 @@ namespace UE
 
 			void FillMorphTargetMeshDescriptionsPerMorphTargetName(const FMeshNodeContext& MeshNodeContext
 																 , TMap<FString, TOptional<UE::Interchange::FMeshPayloadData>>& MorphTargetMeshDescriptionsPerMorphTargetName
-																 , const IInterchangeMeshPayloadInterface* MeshTranslatorPayloadInterface
+																 , UInterchangeSkeletalMeshFactory::FLodPayloads& LodPayloads
 																 , const int32 VertexOffset
 																 , const UInterchangeBaseNodeContainer* NodeContainer
 																 , FString AssetName)
@@ -96,7 +116,7 @@ namespace UE
 				TRACE_CPUPROFILER_EVENT_SCOPE(FillMorphTargetMeshDescriptionsPerMorphTargetName)
 				TArray<FString> MorphTargetUids;
 				MeshNodeContext.MeshNode->GetMorphTargetDependencies(MorphTargetUids);
-				TMap<FString, TFuture<TOptional<UE::Interchange::FMeshPayloadData>>> TempMorphTargetMeshDescriptionsPerMorphTargetName;
+				TMap<FString, TOptional<UE::Interchange::FMeshPayloadData>> TempMorphTargetMeshDescriptionsPerMorphTargetName;
 				TempMorphTargetMeshDescriptionsPerMorphTargetName.Reserve(MorphTargetUids.Num());
 				for (const FString& MorphTargetUid : MorphTargetUids)
 				{
@@ -109,9 +129,9 @@ namespace UE
 							continue;
 						}
 						FInterchangeMeshPayLoadKey& PayLoadKey = OptionalPayLoadKey.GetValue();
-
+						FInterchangeMeshPayLoadKey GlobalMorphPayLoadKey = MeshNodeContext.GetMorphTargetAndTransformPayloadKey(PayLoadKey);
 						//Add the map entry key, the translator will be call after to bulk get all the needed payload
-						TempMorphTargetMeshDescriptionsPerMorphTargetName.Add(PayLoadKey.UniqueId, MeshTranslatorPayloadInterface->GetMeshPayloadData(PayLoadKey, MeshNodeContext.SceneGlobalTransform.Get(FTransform::Identity)));
+						TempMorphTargetMeshDescriptionsPerMorphTargetName.Add(PayLoadKey.UniqueId, LodPayloads.MorphPayloadPerKey.FindChecked(GlobalMorphPayLoadKey));
 					}
 				}
 
@@ -133,7 +153,7 @@ namespace UE
 							continue;
 						}
 
-						TOptional<UE::Interchange::FMeshPayloadData> MorphTargetMeshPayload = TempMorphTargetMeshDescriptionsPerMorphTargetName.FindChecked(MorphTargetPayloadKeyString).Get();
+						TOptional<UE::Interchange::FMeshPayloadData> MorphTargetMeshPayload = TempMorphTargetMeshDescriptionsPerMorphTargetName.FindChecked(MorphTargetPayloadKeyString);
 						if (!MorphTargetMeshPayload.IsSet())
 						{
 							UE_LOG(LogInterchangeImport, Warning, TEXT("Invalid skeletal mesh morph target payload key [%s] for SkeletalMesh asset %s."), *MorphTargetPayloadKeyString, *AssetName);
@@ -454,18 +474,14 @@ namespace UE
 			void RetrieveAllSkeletalMeshPayloadsAndFillImportData(const UInterchangeSkeletalMeshFactoryNode* SkeletalMeshFactoryNode
 																  , FSkeletalMeshImportData& DestinationImportData
 																  , TArray<FMeshNodeContext>& MeshReferences
+																  , UInterchangeSkeletalMeshFactory::FLodPayloads& LodPayloads
 																  , TArray<SkeletalMeshImportData::FBone>& RefBonesBinary
 																  , const UInterchangeSkeletalMeshFactory::FImportAssetObjectParams& Arguments
-																  , const IInterchangeMeshPayloadInterface* MeshTranslatorPayloadInterface
 																  , const bool bSkinControlPointToTimeZero
 																  , const UInterchangeBaseNodeContainer* NodeContainer
 																  , const FString& RootJointNodeId)
 			{
 				TRACE_CPUPROFILER_EVENT_SCOPE(RetrieveAllSkeletalMeshPayloadsAndFillImportData)
-				if (!MeshTranslatorPayloadInterface)
-				{
-					return;
-				}
 				FMeshDescription LodMeshDescription;
 				FSkeletalMeshAttributes SkeletalMeshAttributes(LodMeshDescription);
 				SkeletalMeshAttributes.Register();
@@ -478,17 +494,18 @@ namespace UE
 				bool bKeepSectionsSeparate = false;
 				SkeletalMeshFactoryNode->GetCustomKeepSectionsSeparate(bKeepSectionsSeparate);
 
+				bool bImportVertexAttributes = false;
+				SkeletalMeshFactoryNode->GetCustomImportVertexAttributes(bImportVertexAttributes);
+
 				bool bImportMorphTarget = true;
 				SkeletalMeshFactoryNode->GetCustomImportMorphTarget(bImportMorphTarget);
 
-				TMap<const FMeshNodeContext*, TFuture<TOptional<UE::Interchange::FMeshPayloadData>>> LodMeshPayloadPerTranslatorPayloadKey;
+				TMap<const FMeshNodeContext*, TOptional<UE::Interchange::FMeshPayloadData>> LodMeshPayloadPerTranslatorPayloadKey;
 				LodMeshPayloadPerTranslatorPayloadKey.Reserve(MeshReferences.Num());
 
 				TMap<FString, TOptional<UE::Interchange::FMeshPayloadData>> MorphTargetMeshDescriptionsPerMorphTargetName;
 				int32 MorphTargetCount = 0;
 
-				bool bImportVertexAttributes = false;
-				SkeletalMeshFactoryNode->GetCustomImportVertexAttributes(bImportVertexAttributes);
 				struct FInternalInstanceData
 				{
 					bool ScaleGreaterThenOne = false;
@@ -508,10 +525,8 @@ namespace UE
 
 				for (const FMeshNodeContext& MeshNodeContext : MeshReferences)
 				{
-					const FInternalInstanceData& InstanceData = MeshInstancesDatas.FindChecked(MeshNodeContext.TranslatorPayloadKey.UniqueId);
-					FTransform ApplyTransformWhenFetchPayload = InstanceData.ShouldFetchWithTransform() ? MeshNodeContext.SceneGlobalTransform.Get(FTransform::Identity) : FTransform::Identity;
 					//Add the payload entry key, the payload data will be fill later in bulk by the translator
-					LodMeshPayloadPerTranslatorPayloadKey.Add(&MeshNodeContext, MeshTranslatorPayloadInterface->GetMeshPayloadData(MeshNodeContext.TranslatorPayloadKey, ApplyTransformWhenFetchPayload));
+					LodMeshPayloadPerTranslatorPayloadKey.Add(&MeshNodeContext, LodPayloads.MeshPayloadPerKey.FindChecked(MeshNodeContext.GetTranslatorAndTransformPayloadKey()));
 					//Count the morph target dependencies so we can reserve the right amount
 					MorphTargetCount += (bImportMorphTarget && MeshNodeContext.MeshNode) ? MeshNodeContext.MeshNode->GetMorphTargetDependeciesCount() : 0;
 				}
@@ -520,7 +535,7 @@ namespace UE
 				//Fill the lod mesh description using all combined mesh part
 				TArray<SkeletalMeshImportData::FMeshInfo> MeshInfos;
 				int32 MeshInfoVertexOffset = 0;
-				for(TPair<const FMeshNodeContext*, TFuture<TOptional<UE::Interchange::FMeshPayloadData>>>& MeshNodeContextAndFuture: LodMeshPayloadPerTranslatorPayloadKey)
+				for(TPair<const FMeshNodeContext*, TOptional<UE::Interchange::FMeshPayloadData>>& MeshNodeContextAndFuture: LodMeshPayloadPerTranslatorPayloadKey)
 				{
 					if (!MeshNodeContextAndFuture.Key)
 					{
@@ -528,7 +543,7 @@ namespace UE
 					}
 					const FMeshNodeContext& MeshNodeContext = *MeshNodeContextAndFuture.Key;
 					TRACE_CPUPROFILER_EVENT_SCOPE(RetrieveAllSkeletalMeshPayloadsAndFillImportData::GetPayload)
-					TOptional<UE::Interchange::FMeshPayloadData> LodMeshPayload = MeshNodeContextAndFuture.Value.Get();
+					TOptional<UE::Interchange::FMeshPayloadData> LodMeshPayload = MeshNodeContextAndFuture.Value;
 					if (!LodMeshPayload.IsSet())
 					{
 						UE_LOG(LogInterchangeImport, Warning, TEXT("Invalid skeletal mesh payload key [%s] for SkeletalMesh asset %s."), *MeshNodeContext.TranslatorPayloadKey.UniqueId, *Arguments.AssetName);
@@ -650,7 +665,7 @@ namespace UE
 					{
 						FillMorphTargetMeshDescriptionsPerMorphTargetName(MeshNodeContext
 																		, MorphTargetMeshDescriptionsPerMorphTargetName
-																		, MeshTranslatorPayloadInterface
+																		, LodPayloads
 																		, VertexOffset
 																		, Arguments.NodeContainer
 																		, Arguments.AssetName);
@@ -1013,6 +1028,263 @@ UClass* UInterchangeSkeletalMeshFactory::GetFactoryClass() const
 	return USkeletalMesh::StaticClass();
 }
 
+void UInterchangeSkeletalMeshFactory::CreatePayloadTasks(const FImportAssetObjectParams& Arguments, bool bAsync, TArray<TSharedPtr<UE::Interchange::FInterchangeTaskBase>>& PayloadTasks)
+{
+	TRACE_CPUPROFILER_EVENT_SCOPE(UInterchangeSkeletalMeshFactory::CreateAsset)
+
+#if !WITH_EDITOR || !WITH_EDITORONLY_DATA
+
+	UE_LOG(LogInterchangeImport, Error, TEXT("Cannot import SkeletalMesh asset at runtime. This is an editor-only feature."));
+	return;
+
+#else
+	using namespace UE::Interchange;
+
+	if (!Arguments.AssetNode || !Arguments.AssetNode->GetObjectClass()->IsChildOf(GetFactoryClass()))
+	{
+		return;
+	}
+
+	UInterchangeSkeletalMeshFactoryNode* SkeletalMeshFactoryNode = Cast<UInterchangeSkeletalMeshFactoryNode>(Arguments.AssetNode);
+	if (SkeletalMeshFactoryNode == nullptr)
+	{
+		return;
+	}
+
+	const IInterchangeMeshPayloadInterface* MeshTranslatorPayloadInterface = Cast<IInterchangeMeshPayloadInterface>(Arguments.Translator);
+	if (!MeshTranslatorPayloadInterface)
+	{
+		UInterchangeResultError_Generic* Message = AddMessage<UInterchangeResultError_Generic>();
+		Message->Text = FText::Format(NSLOCTEXT("InterchangeSkeletalMeshFactory", "CreatePayloadTasks_TranslatorInterfaceMissing", "Cannot import skeletalMesh {0}, the translator {1} does not implement the IInterchangeSkeletalMeshPayloadInterface.")
+			, FText::FromString(Arguments.AssetName)
+			, FText::FromString(Arguments.Translator->GetName()));
+		return;
+	}
+
+	FTransform GlobalOffsetTransform = FTransform::Identity;
+	bool bBakeMeshes = false;
+	if (UInterchangeCommonPipelineDataFactoryNode* CommonPipelineDataFactoryNode = UInterchangeCommonPipelineDataFactoryNode::GetUniqueInstance(Arguments.NodeContainer))
+	{
+		CommonPipelineDataFactoryNode->GetCustomGlobalOffsetTransform(GlobalOffsetTransform);
+		CommonPipelineDataFactoryNode->GetBakeMeshes(bBakeMeshes);
+	}
+
+	bool bImportMorphTarget = true;
+	SkeletalMeshFactoryNode->GetCustomImportMorphTarget(bImportMorphTarget);
+
+	int32 LodCount = SkeletalMeshFactoryNode->GetLodDataCount();
+	TArray<FString> LodDataUniqueIds;
+	SkeletalMeshFactoryNode->GetLodDataUniqueIds(LodDataUniqueIds);
+	ensure(LodDataUniqueIds.Num() == LodCount);
+	PayloadsPerLodIndex.Reserve(LodCount);
+	for (int32 LodIndex = 0; LodIndex < LodCount; ++LodIndex)
+	{
+		FString LodUniqueId = LodDataUniqueIds[LodIndex];
+		const UInterchangeSkeletalMeshLodDataNode* LodDataNode = Cast<UInterchangeSkeletalMeshLodDataNode>(Arguments.NodeContainer->GetNode(LodUniqueId));
+		if (!LodDataNode)
+		{
+			continue;
+		}
+
+		FString SkeletonNodeUid;
+		if (!LodDataNode->GetCustomSkeletonUid(SkeletonNodeUid))
+		{
+			continue;
+		}
+		const UInterchangeSkeletonFactoryNode* SkeletonNode = Cast<UInterchangeSkeletonFactoryNode>(Arguments.NodeContainer->GetNode(SkeletonNodeUid));
+		if (!SkeletonNode)
+		{
+			continue;
+		}
+
+		FString RootJointNodeId;
+		if (!SkeletonNode->GetCustomRootJointUid(RootJointNodeId))
+		{
+			continue;
+		}
+
+		const UInterchangeSceneNode* RootJointNode = Cast<UInterchangeSceneNode>(Arguments.NodeContainer->GetNode(RootJointNodeId));
+		if (!RootJointNode)
+		{
+			continue;
+		}
+
+		FLodPayloads& LodPayloads = PayloadsPerLodIndex.FindOrAdd(LodIndex);
+
+		TArray<SkeletalMeshImportData::FBone> RefBonesBinary;
+		bool bUseTimeZeroAsBindPose = false;
+		bool bDiffPose = false;
+
+		TArray <Private::FJointInfo> JointInfos;
+		TArray<FString> BoneNotBindNames;
+		Private::FSkeletonHelper::RecursiveAddBones(Arguments.NodeContainer
+			, RootJointNodeId
+			, JointInfos
+			, INDEX_NONE
+			, RefBonesBinary
+			, bUseTimeZeroAsBindPose
+			, bDiffPose
+			, BoneNotBindNames);
+
+		FTransform RootJointNodeGlobalTransform;
+		ensure(RootJointNode->GetCustomGlobalTransform(Arguments.NodeContainer, GlobalOffsetTransform, RootJointNodeGlobalTransform));
+		FTransform RootJointNodeLocalTransform;
+		ensure(RootJointNode->GetCustomLocalTransform(RootJointNodeLocalTransform));
+		FTransform BakeToRootJointTransfromModifier = RootJointNodeGlobalTransform.Inverse() * RootJointNodeLocalTransform; //It is used for !bBakeMeshes, and so the GlobalTransform will be inversed out when multiplied into the CustomBindPoseGlobalTransform
+		FTransform BakeFromRootJointTransfromModifier = RootJointNodeLocalTransform.Inverse() * RootJointNodeGlobalTransform * GlobalOffsetTransform.Inverse(); //GlobalOffseTransform will be added by the BindPoseGlobalTransform when bBakeMeshes && !bRootAncestorOfMeshDependency is used
+
+		TArray<Private::FMeshNodeContext> MeshReferences;
+		//Scope to query the mesh node
+		{
+			TArray<FString> MeshUids;
+			LodDataNode->GetMeshUids(MeshUids);
+			MeshReferences.Reserve(MeshUids.Num());
+			for (const FString& MeshUid : MeshUids)
+			{
+				Private::FMeshNodeContext MeshReference;
+				MeshReference.MeshNode = Cast<UInterchangeMeshNode>(Arguments.NodeContainer->GetNode(MeshUid));
+				if (!MeshReference.MeshNode)
+				{
+					//The reference is a scene node and we need to bake the geometry
+					MeshReference.SceneNode = Cast<UInterchangeSceneNode>(Arguments.NodeContainer->GetNode(MeshUid));
+					if (!ensure(MeshReference.SceneNode != nullptr))
+					{
+						continue;
+					}
+					FString MeshDependencyUid;
+					MeshReference.SceneNode->GetCustomAssetInstanceUid(MeshDependencyUid);
+					MeshReference.MeshNode = Cast<UInterchangeMeshNode>(Arguments.NodeContainer->GetNode(MeshDependencyUid));
+					bool bRootAncestorOfSceneNode = Arguments.NodeContainer->GetIsAncestor(MeshReference.SceneNode->GetUniqueID(), RootJointNode->GetParentUid());
+					//Cache the scene node global matrix, we will use this matrix to bake the vertices, add the node geometric mesh offset to this matrix to bake it properly
+					FTransform SceneNodeTransform;
+					if (!bUseTimeZeroAsBindPose || !MeshReference.SceneNode->GetCustomTimeZeroGlobalTransform(Arguments.NodeContainer, GlobalOffsetTransform, SceneNodeTransform))
+					{
+						ensure(MeshReference.SceneNode->GetCustomBindPoseGlobalTransform(Arguments.NodeContainer, GlobalOffsetTransform, SceneNodeTransform));
+						if (bRootAncestorOfSceneNode)
+						{
+							if (!bBakeMeshes)
+							{
+								SceneNodeTransform *= BakeToRootJointTransfromModifier;
+							}
+						}
+						else
+						{
+							if (bBakeMeshes)
+							{
+								SceneNodeTransform = BakeFromRootJointTransfromModifier * SceneNodeTransform;
+							}
+							else
+							{
+								SceneNodeTransform *= GlobalOffsetTransform.Inverse();
+							}
+						}
+					}
+
+					FTransform SceneNodeGeometricTransform;
+					if (MeshReference.SceneNode->GetCustomGeometricTransform(SceneNodeGeometricTransform))
+					{
+						SceneNodeTransform = SceneNodeGeometricTransform * SceneNodeTransform;
+					}
+					MeshReference.SceneGlobalTransform = SceneNodeTransform;
+				}
+				else
+				{
+					MeshReference.SceneGlobalTransform = GlobalOffsetTransform;
+				}
+
+				if (!ensure(MeshReference.MeshNode != nullptr))
+				{
+					continue;
+				}
+
+				TOptional<FInterchangeMeshPayLoadKey> OptionalPayLoadKey = MeshReference.MeshNode->GetPayLoadKey();
+				if (OptionalPayLoadKey.IsSet())
+				{
+					MeshReference.TranslatorPayloadKey = OptionalPayLoadKey.GetValue();
+				}
+				else
+				{
+					continue;
+				}
+				MeshReferences.Add(MeshReference);
+			}
+		}
+
+		struct FInternalInstanceData
+		{
+			bool ScaleGreaterThenOne = false;
+			int32 Count = 0;
+			bool ShouldFetchWithTransform() const
+			{
+				return Count == 1 || ScaleGreaterThenOne;
+			}
+		};
+		TMap<FString, FInternalInstanceData> MeshInstancesDatas;
+		for (const Private::FMeshNodeContext& MeshNodeContext : MeshReferences)
+		{
+			FInternalInstanceData& InstanceData = MeshInstancesDatas.FindOrAdd(MeshNodeContext.TranslatorPayloadKey.UniqueId);
+			InstanceData.Count++;
+			InstanceData.ScaleGreaterThenOne |= MeshNodeContext.SceneGlobalTransform->GetScale3D().GetAbs().GetMax() > 1.0;
+		}
+		
+		//Reserve the correct amount since we point in the array for the lambda so array must not be resized at any moments after we create the tasks
+		LodPayloads.MeshPayloadPerKey.Reserve(MeshReferences.Num());
+		int32 MorphTargetCount = 0;
+		for (const Private::FMeshNodeContext& MeshNodeContext : MeshReferences)
+		{
+			//Count the morph target dependencies so we can reserve the right amount
+			if (bImportMorphTarget)
+			{
+				MorphTargetCount += MeshNodeContext.MeshNode->GetMorphTargetDependeciesCount();
+			}
+		}
+		LodPayloads.MorphPayloadPerKey.Reserve(MorphTargetCount);
+
+		for (const Private::FMeshNodeContext& MeshNodeContext : MeshReferences)
+		{
+			const FInternalInstanceData& InstanceData = MeshInstancesDatas.FindChecked(MeshNodeContext.TranslatorPayloadKey.UniqueId);
+			FTransform ApplyTransformWhenFetchPayload = InstanceData.ShouldFetchWithTransform() ? MeshNodeContext.SceneGlobalTransform.Get(FTransform::Identity) : FTransform::Identity;
+			//Create the payload task
+			TOptional<FMeshPayloadData>& MeshPayload = LodPayloads.MeshPayloadPerKey.FindOrAdd(MeshNodeContext.GetTranslatorAndTransformPayloadKey());
+			TSharedPtr<FInterchangeTaskLambda, ESPMode::ThreadSafe> TaskGetMeshPayload = MakeShared<FInterchangeTaskLambda, ESPMode::ThreadSafe>(bAsync ? EInterchangeTaskThread::AsyncThread : EInterchangeTaskThread::GameThread
+				, [&MeshPayload, ApplyTransformWhenFetchPayload, MeshTranslatorPayloadInterface, PayLoadKey = MeshNodeContext.TranslatorPayloadKey]()
+				{
+					MeshPayload = MeshTranslatorPayloadInterface->GetMeshPayloadData(PayLoadKey, ApplyTransformWhenFetchPayload);
+				});
+			PayloadTasks.Add(TaskGetMeshPayload);
+
+			//Count the morph target dependencies so we can reserve the right amount
+			if (bImportMorphTarget)
+			{
+				TArray<FString> MorphTargetUids;
+				MeshNodeContext.MeshNode->GetMorphTargetDependencies(MorphTargetUids);
+				for (const FString& MorphTargetUid : MorphTargetUids)
+				{
+					if (const UInterchangeMeshNode* MorphTargetMeshNode = Cast<UInterchangeMeshNode>(Arguments.NodeContainer->GetNode(MorphTargetUid)))
+					{
+						TOptional<FInterchangeMeshPayLoadKey> OptionalPayLoadKey = MorphTargetMeshNode->GetPayLoadKey();
+						if (!OptionalPayLoadKey.IsSet())
+						{
+							UE_LOG(LogInterchangeImport, Warning, TEXT("Empty LOD morph target mesh reference payload when importing SkeletalMesh asset %s."), *Arguments.AssetName);
+							continue;
+						}
+						FInterchangeMeshPayLoadKey& MorphPayLoadKey = OptionalPayLoadKey.GetValue();
+						FInterchangeMeshPayLoadKey GlobalMorphPayLoadKey = MeshNodeContext.GetMorphTargetAndTransformPayloadKey(MorphPayLoadKey);
+						TOptional<FMeshPayloadData>& MorphPayload = LodPayloads.MorphPayloadPerKey.FindOrAdd(GlobalMorphPayLoadKey);
+						TSharedPtr<FInterchangeTaskLambda, ESPMode::ThreadSafe> TaskGetMorphPayload = MakeShared<FInterchangeTaskLambda, ESPMode::ThreadSafe>(bAsync ? EInterchangeTaskThread::AsyncThread : EInterchangeTaskThread::GameThread
+							, [&MorphPayload, ApplyTransformWhenFetchPayload, MeshTranslatorPayloadInterface, MorphPayLoadKey]()
+							{
+								MorphPayload = MeshTranslatorPayloadInterface->GetMeshPayloadData(MorphPayLoadKey, ApplyTransformWhenFetchPayload);
+							});
+						PayloadTasks.Add(TaskGetMorphPayload);
+					}
+				}
+			}
+		}
+	}
+#endif
+}
+
 UInterchangeFactoryBase::FImportAssetResult UInterchangeSkeletalMeshFactory::BeginImportAsset_GameThread(const FImportAssetObjectParams& Arguments)
 {
 	TRACE_CPUPROFILER_EVENT_SCOPE(UInterchangeSkeletalMeshFactory::BeginImportAssetObject_GameThread)
@@ -1442,6 +1714,10 @@ UInterchangeFactoryBase::FImportAssetResult UInterchangeSkeletalMeshFactory::Imp
 			Message->Text = WarningMessage_InvalidRootJoint;
 			continue;
 		}
+
+		//We should have a valid lod payload data
+		FLodPayloads& LodPayloads = PayloadsPerLodIndex.FindChecked(LodIndex);
+
 		FImportAssetObjectLODData& ImportAssetObjectLODData = ImportAssetObjectData.LodDatas[CurrentLodIndex];
 		ensure(ImportAssetObjectLODData.LodIndex == CurrentLodIndex);
 
@@ -1549,9 +1825,9 @@ UInterchangeFactoryBase::FImportAssetResult UInterchangeSkeletalMeshFactory::Imp
 		UE::Interchange::Private::RetrieveAllSkeletalMeshPayloadsAndFillImportData(SkeletalMeshFactoryNode
 																					, SkeletalMeshImportData
 																					, MeshReferences
+																					, LodPayloads
 																					, ImportAssetObjectLODData.RefBonesBinary
 																					, Arguments
-																					, MeshTranslatorPayloadInterface
 																					, bSkinControlPointToTimeZero
 																					, Arguments.NodeContainer
 																					, RootJointNodeId);
diff --git a/Engine/Plugins/Interchange/Runtime/Source/Import/Private/Mesh/InterchangeStaticMeshFactory.cpp b/Engine/Plugins/Interchange/Runtime/Source/Import/Private/Mesh/InterchangeStaticMeshFactory.cpp
--- a/Engine/Plugins/Interchange/Runtime/Source/Import/Private/Mesh/InterchangeStaticMeshFactory.cpp
+++ b/Engine/Plugins/Interchange/Runtime/Source/Import/Private/Mesh/InterchangeStaticMeshFactory.cpp
@@ -148,6 +148,161 @@ namespace UE::Interchange::Private::StaticMesh
 	}
 } //ns UE::Interchange::Private::StaticMesh
 
+void UInterchangeStaticMeshFactory::CreatePayloadTasks(const FImportAssetObjectParams& Arguments, bool bAsync, TArray<TSharedPtr<UE::Interchange::FInterchangeTaskBase>>& PayloadTasks)
+{
+	TRACE_CPUPROFILER_EVENT_SCOPE(UInterchangeStaticMeshFactory::CreatePayloadTasks);
+
+	UInterchangeStaticMeshFactoryNode* StaticMeshFactoryNode = Cast<UInterchangeStaticMeshFactoryNode>(Arguments.AssetNode);
+	if (StaticMeshFactoryNode == nullptr)
+	{
+		return;
+	}
+
+	const int32 LodCount = FMath::Min(StaticMeshFactoryNode->GetLodDataCount(), MAX_STATIC_MESH_LODS);
+
+	// Now import geometry for each LOD
+	TArray<FString> LodDataUniqueIds;
+	StaticMeshFactoryNode->GetLodDataUniqueIds(LodDataUniqueIds);
+	ensure(LodDataUniqueIds.Num() >= LodCount);
+
+	const IInterchangeMeshPayloadInterface* MeshTranslatorPayloadInterface = Cast<IInterchangeMeshPayloadInterface>(Arguments.Translator);
+	if (!MeshTranslatorPayloadInterface)
+	{
+		UE_LOG(LogInterchangeImport, Error, TEXT("Cannot import static mesh. The translator does not implement IInterchangeMeshPayloadInterface."));
+		return;
+	}
+
+	FTransform GlobalOffsetTransform = FTransform::Identity;
+	bool bBakeMeshes = false;
+	bool bBakePivotMeshes = false;
+	if (UInterchangeCommonPipelineDataFactoryNode* CommonPipelineDataFactoryNode = UInterchangeCommonPipelineDataFactoryNode::GetUniqueInstance(Arguments.NodeContainer))
+	{
+		CommonPipelineDataFactoryNode->GetCustomGlobalOffsetTransform(GlobalOffsetTransform);
+		CommonPipelineDataFactoryNode->GetBakeMeshes(bBakeMeshes);
+		if (!bBakeMeshes)
+		{
+			CommonPipelineDataFactoryNode->GetBakePivotMeshes(bBakePivotMeshes);
+		}
+	}
+
+	PayloadsPerLodIndex.Reserve(LodCount);
+	int32 CurrentLodIndex = 0;
+	for (int32 LodIndex = 0; LodIndex < LodCount; ++LodIndex)
+	{
+		FString LodUniqueId = LodDataUniqueIds[LodIndex];
+		const UInterchangeStaticMeshLodDataNode* LodDataNode = Cast<UInterchangeStaticMeshLodDataNode>(Arguments.NodeContainer->GetNode(LodUniqueId));
+		if (!LodDataNode)
+		{
+			UE_LOG(LogInterchangeImport, Warning, TEXT("Invalid LOD when importing StaticMesh asset %s."), *Arguments.AssetName);
+			continue;
+		}
+
+		FLodPayloads& LodPayloads = PayloadsPerLodIndex.FindOrAdd(LodIndex);
+
+		auto AddMeshPayloads = [&Arguments, MeshTranslatorPayloadInterface, bBakeMeshes, bBakePivotMeshes, &GlobalOffsetTransform, &PayloadTasks, bAsync]
+			(const TArray<FString>& MeshUids, TMap<FInterchangeMeshPayLoadKey, FMeshPayload>& PayloadPerKey)
+			{
+				for (const FString& MeshUid : MeshUids)
+				{
+					FTransform GlobalMeshTransform;
+					const UInterchangeBaseNode* Node = Arguments.NodeContainer->GetNode(MeshUid);
+					const UInterchangeMeshNode* MeshNode = Cast<const UInterchangeMeshNode>(Node);
+					if (MeshNode == nullptr)
+					{
+						// MeshUid must refer to a scene node
+						const UInterchangeSceneNode* SceneNode = Cast<const UInterchangeSceneNode>(Node);
+						if (!ensure(SceneNode))
+						{
+							UE_LOG(LogInterchangeImport, Warning, TEXT("Invalid LOD mesh reference when importing StaticMesh asset %s."), *Arguments.AssetName);
+							continue;
+						}
+
+						if (bBakeMeshes)
+						{
+							// Get the transform from the scene node
+							FTransform SceneNodeGlobalTransform;
+							if (SceneNode->GetCustomGlobalTransform(Arguments.NodeContainer, GlobalOffsetTransform, SceneNodeGlobalTransform))
+							{
+								GlobalMeshTransform = SceneNodeGlobalTransform;
+							}
+						}
+						UE::Interchange::Private::MeshHelper::AddSceneNodeGeometricAndPivotToGlobalTransform(GlobalMeshTransform, SceneNode, bBakeMeshes, bBakePivotMeshes);
+						// And get the mesh node which it references
+						FString MeshDependencyUid;
+						SceneNode->GetCustomAssetInstanceUid(MeshDependencyUid);
+						MeshNode = Cast<UInterchangeMeshNode>(Arguments.NodeContainer->GetNode(MeshDependencyUid));
+					}
+					else if (bBakeMeshes)
+					{
+						//If we have a mesh that is not reference by a scene node, we must apply the global offset.
+						GlobalMeshTransform = GlobalOffsetTransform;
+					}
+
+					if (!ensure(MeshNode))
+					{
+						UE_LOG(LogInterchangeImport, Warning, TEXT("Invalid LOD mesh reference when importing StaticMesh asset %s."), *Arguments.AssetName);
+						continue;
+					}
+
+					TOptional<FInterchangeMeshPayLoadKey> OptionalPayLoadKey = MeshNode->GetPayLoadKey();
+					if (!ensure(OptionalPayLoadKey.IsSet()))
+					{
+						UE_LOG(LogInterchangeImport, Warning, TEXT("Empty LOD mesh reference payload when importing StaticMesh asset %s."), *Arguments.AssetName);
+						continue;
+					}
+
+					FInterchangeMeshPayLoadKey& PayLoadKey = OptionalPayLoadKey.GetValue();
+					
+					FInterchangeMeshPayLoadKey GlobalPayLoadKey = PayLoadKey;
+					GlobalPayLoadKey.UniqueId += GlobalMeshTransform.ToString();
+					if (!PayloadPerKey.Contains(GlobalPayLoadKey))
+					{
+						FMeshPayload& Payload = PayloadPerKey.FindOrAdd(GlobalPayLoadKey);
+						Payload.Transform = GlobalMeshTransform;
+						Payload.MeshName = PayLoadKey.UniqueId;
+						TSharedPtr<UE::Interchange::FInterchangeTaskLambda, ESPMode::ThreadSafe> TaskGetMeshPayload = MakeShared<UE::Interchange::FInterchangeTaskLambda, ESPMode::ThreadSafe>(bAsync ? UE::Interchange::EInterchangeTaskThread::AsyncThread : UE::Interchange::EInterchangeTaskThread::GameThread
+							, [&Payload, GlobalMeshTransform, MeshTranslatorPayloadInterface, PayLoadKey]()
+							{
+								if (ensure(!Payload.PayloadData.IsSet()))
+								{
+									Payload.PayloadData = MeshTranslatorPayloadInterface->GetMeshPayloadData(PayLoadKey, GlobalMeshTransform);
+								}
+							});
+						PayloadTasks.Add(TaskGetMeshPayload);
+					}
+				}
+			};
+
+		TArray<FString> MeshUids;
+		LodDataNode->GetMeshUids(MeshUids);
+		LodPayloads.MeshPayloadPerKey.Reserve(MeshUids.Num());
+		AddMeshPayloads(MeshUids, LodPayloads.MeshPayloadPerKey);
+
+		if (LodIndex == 0)
+		{
+			TArray<FString> BoxCollisionMeshUids;
+			LodDataNode->GetBoxCollisionMeshUids(BoxCollisionMeshUids);
+			LodPayloads.CollisionBoxPayloadPerKey.Reserve(BoxCollisionMeshUids.Num());
+			AddMeshPayloads(BoxCollisionMeshUids, LodPayloads.CollisionBoxPayloadPerKey);
+
+			TArray<FString> CapsuleCollisionMeshUids;
+			LodDataNode->GetCapsuleCollisionMeshUids(CapsuleCollisionMeshUids);
+			LodPayloads.CollisionCapsulePayloadPerKey.Reserve(CapsuleCollisionMeshUids.Num());
+			AddMeshPayloads(CapsuleCollisionMeshUids, LodPayloads.CollisionCapsulePayloadPerKey);
+
+			TArray<FString> SphereCollisionMeshUids;
+			LodDataNode->GetSphereCollisionMeshUids(SphereCollisionMeshUids);
+			LodPayloads.CollisionSpherePayloadPerKey.Reserve(SphereCollisionMeshUids.Num());
+			AddMeshPayloads(SphereCollisionMeshUids, LodPayloads.CollisionSpherePayloadPerKey);
+
+			TArray<FString> ConvexCollisionMeshUids;
+			LodDataNode->GetConvexCollisionMeshUids(ConvexCollisionMeshUids);
+			LodPayloads.CollisionConvexPayloadPerKey.Reserve(ConvexCollisionMeshUids.Num());
+			AddMeshPayloads(ConvexCollisionMeshUids, LodPayloads.CollisionConvexPayloadPerKey);
+		}
+	}
+}
+
 UInterchangeFactoryBase::FImportAssetResult UInterchangeStaticMeshFactory::BeginImportAsset_GameThread(const FImportAssetObjectParams& Arguments)
 {
 	TRACE_CPUPROFILER_EVENT_SCOPE(UInterchangeStaticMeshFactory::BeginImportAsset_GameThread);
@@ -328,14 +483,23 @@ UInterchangeFactoryBase::FImportAssetResult UInterchangeStaticMeshFactory::Impor
 		TArray<FString> MeshUids;
 		LodDataNode->GetMeshUids(MeshUids);
 
-		// Fill the lod mesh description using all combined mesh parts
-		TArray<FMeshPayload> MeshPayloads = GetMeshPayloads(Arguments, MeshUids);
+		FLodPayloads LodPayloads;
+		if (PayloadsPerLodIndex.Contains(LodIndex))
+		{
+			// Fill the lod mesh description using all combined mesh parts
+			LodPayloads = PayloadsPerLodIndex.FindChecked(LodIndex);
+		}
+		else
+		{
+			UE_LOG(LogInterchangeImport, Error, TEXT("LOD %d do not have any valid payload to create a mesh when importing StaticMesh asset %s."), LodIndex, *Arguments.AssetName);
+			continue;
+		}
 
 		// Just move the mesh description from the first valid payload then append the rest
 		bool bFirstValidMoved = false;
-		for (FMeshPayload& MeshPayload : MeshPayloads)
+		for (TPair<FInterchangeMeshPayLoadKey, FMeshPayload>& KeyAndPayload : LodPayloads.MeshPayloadPerKey)
 		{
-			const TOptional<UE::Interchange::FMeshPayloadData>& LodMeshPayload = MeshPayload.PayloadData.Get();
+			const TOptional<UE::Interchange::FMeshPayloadData>& LodMeshPayload = KeyAndPayload.Value.PayloadData;
 			if (!LodMeshPayload.IsSet())
 			{
 				UE_LOG(LogInterchangeImport, Warning, TEXT("Invalid static mesh payload key for StaticMesh asset %s."), *Arguments.AssetName);
@@ -442,9 +606,9 @@ UInterchangeFactoryBase::FImportAssetResult UInterchangeStaticMeshFactory::Impor
 					StaticMesh->GetBodySetup()->AggGeom.EmptyElements();
 				}
 
-				bImportedCustomCollision |= ImportBoxCollision(Arguments, StaticMesh, LodDataNode);
-				bImportedCustomCollision |= ImportCapsuleCollision(Arguments, StaticMesh, LodDataNode);
-				bImportedCustomCollision |= ImportSphereCollision(Arguments, StaticMesh, LodDataNode);
+				bImportedCustomCollision |= ImportBoxCollision(Arguments, StaticMesh);
+				bImportedCustomCollision |= ImportCapsuleCollision(Arguments, StaticMesh);
+				bImportedCustomCollision |= ImportSphereCollision(Arguments, StaticMesh);
 				bImportedCustomCollision |= ImportConvexCollision(Arguments, StaticMesh, LodDataNode);
 			}
 		}
@@ -1014,101 +1178,6 @@ void UInterchangeStaticMeshFactory::BuildObject_GameThread(const FSetupObjectPar
 #endif
 }
 
-
-TArray<UInterchangeStaticMeshFactory::FMeshPayload> UInterchangeStaticMeshFactory::GetMeshPayloads(const FImportAssetObjectParams& Arguments, const TArray<FString>& MeshUids) const
-{
-	TRACE_CPUPROFILER_EVENT_SCOPE(UInterchangeStaticMeshFactory::GetMeshPayloads);
-
-	TArray<FMeshPayload> Payloads;
-	Payloads.Reserve(MeshUids.Num());
-
-	const IInterchangeMeshPayloadInterface* MeshTranslatorPayloadInterface = Cast<IInterchangeMeshPayloadInterface>(Arguments.Translator);
-	if (!MeshTranslatorPayloadInterface)
-	{
-		UE_LOG(LogInterchangeImport, Error, TEXT("Cannot import static mesh. The translator does not implement IInterchangeMeshPayloadInterface."));
-		return Payloads;
-	}
-
-	FTransform GlobalOffsetTransform = FTransform::Identity;
-	bool bBakeMeshes = false;
-	bool bBakePivotMeshes = false;
-	if (UInterchangeCommonPipelineDataFactoryNode* CommonPipelineDataFactoryNode = UInterchangeCommonPipelineDataFactoryNode::GetUniqueInstance(Arguments.NodeContainer))
-	{
-		CommonPipelineDataFactoryNode->GetCustomGlobalOffsetTransform(GlobalOffsetTransform);
-		CommonPipelineDataFactoryNode->GetBakeMeshes(bBakeMeshes);
-		if (!bBakeMeshes)
-		{
-			CommonPipelineDataFactoryNode->GetBakePivotMeshes(bBakePivotMeshes);
-		}
-	}
-
-	for (const FString& MeshUid : MeshUids)
-	{
-		FMeshPayload Payload;
-
-		const UInterchangeBaseNode* Node = Arguments.NodeContainer->GetNode(MeshUid);
-		const UInterchangeMeshNode* MeshNode = Cast<const UInterchangeMeshNode>(Node);
-		if (MeshNode == nullptr)
-		{
-			// MeshUid must refer to a scene node
-			const UInterchangeSceneNode* SceneNode = Cast<const UInterchangeSceneNode>(Node);
-			if (!ensure(SceneNode))
-			{
-				UE_LOG(LogInterchangeImport, Warning, TEXT("Invalid LOD mesh reference when importing StaticMesh asset %s."), *Arguments.AssetName);
-				continue;
-			}
-
-			if (bBakeMeshes)
-			{
-				// Get the transform from the scene node
-				FTransform SceneNodeGlobalTransform;
-				if (SceneNode->GetCustomGlobalTransform(Arguments.NodeContainer, GlobalOffsetTransform, SceneNodeGlobalTransform))
-				{
-					Payload.Transform = SceneNodeGlobalTransform;
-				}
-			}
-
-			UE::Interchange::Private::MeshHelper::AddSceneNodeGeometricAndPivotToGlobalTransform(Payload.Transform, SceneNode, bBakeMeshes, bBakePivotMeshes);
-
-			// And get the mesh node which it references
-			FString MeshDependencyUid;
-			SceneNode->GetCustomAssetInstanceUid(MeshDependencyUid);
-			MeshNode = Cast<UInterchangeMeshNode>(Arguments.NodeContainer->GetNode(MeshDependencyUid));
-		}
-		else
-		{
-			if (bBakeMeshes)
-			{
-				//If we have a mesh that is not reference by a scene node, we must apply the global offset.
-				Payload.Transform = GlobalOffsetTransform;
-			}
-		}
-
-		if (!ensure(MeshNode))
-		{
-			UE_LOG(LogInterchangeImport, Warning, TEXT("Invalid LOD mesh reference when importing StaticMesh asset %s."), *Arguments.AssetName);
-			continue;
-		}
-
-		TOptional<FInterchangeMeshPayLoadKey> OptionalPayLoadKey = MeshNode->GetPayLoadKey();
-		if (!ensure(OptionalPayLoadKey.IsSet()))
-		{
-			UE_LOG(LogInterchangeImport, Warning, TEXT("Empty LOD mesh reference payload when importing StaticMesh asset %s."), *Arguments.AssetName);
-			continue;
-		}
-
-		FInterchangeMeshPayLoadKey& PayLoadKey = OptionalPayLoadKey.GetValue();
-
-		Payload.MeshName = PayLoadKey.UniqueId;
-		Payload.PayloadData = MeshTranslatorPayloadInterface->GetMeshPayloadData(PayLoadKey, Payload.Transform);
-
-		Payloads.Emplace(MoveTemp(Payload));
-	}
-
-	return Payloads;
-}
-
-
 bool UInterchangeStaticMeshFactory::AddConvexGeomFromVertices(const FImportAssetObjectParams& Arguments, const FMeshDescription& MeshDescription, const FTransform& Transform, FKAggregateGeom& AggGeom)
 {
 	FStaticMeshConstAttributes Attributes(MeshDescription);
@@ -1545,22 +1614,20 @@ bool UInterchangeStaticMeshFactory::AddCapsuleGeomFromVertices(const FImportAsse
 }
 
 
-bool UInterchangeStaticMeshFactory::ImportBoxCollision(const FImportAssetObjectParams& Arguments, UStaticMesh* StaticMesh, const UInterchangeStaticMeshLodDataNode* LodDataNode)
+bool UInterchangeStaticMeshFactory::ImportBoxCollision(const FImportAssetObjectParams& Arguments, UStaticMesh* StaticMesh)
 {
 	using namespace UE::Interchange;
 
 	bool bResult = false;
 
-	TArray<FString> BoxCollisionMeshUids;
-	LodDataNode->GetBoxCollisionMeshUids(BoxCollisionMeshUids);
+	TMap<FInterchangeMeshPayLoadKey, FMeshPayload> BoxCollisionPayloads = PayloadsPerLodIndex.FindChecked(0).CollisionBoxPayloadPerKey;
 
 	FKAggregateGeom& AggGeo = StaticMesh->GetBodySetup()->AggGeom;
 
-	TArray<FMeshPayload> MeshPayloads = GetMeshPayloads(Arguments, BoxCollisionMeshUids);
-	for (const FMeshPayload& MeshPayload : MeshPayloads)
+	for (TPair<FInterchangeMeshPayLoadKey, FMeshPayload>& BoxCollisionPayload : BoxCollisionPayloads)
 	{
 		const FTransform Transform = FTransform::Identity;
-		const TOptional<FMeshPayloadData>& PayloadData = MeshPayload.PayloadData.Get();
+		const TOptional<FMeshPayloadData>& PayloadData = BoxCollisionPayload.Value.PayloadData;
 
 		if (!PayloadData.IsSet())
 		{
@@ -1593,22 +1660,20 @@ bool UInterchangeStaticMeshFactory::ImportBoxCollision(const FImportAssetObjectP
 }
 
 
-bool UInterchangeStaticMeshFactory::ImportCapsuleCollision(const FImportAssetObjectParams& Arguments, UStaticMesh* StaticMesh, const UInterchangeStaticMeshLodDataNode* LodDataNode)
+bool UInterchangeStaticMeshFactory::ImportCapsuleCollision(const FImportAssetObjectParams& Arguments, UStaticMesh* StaticMesh)
 {
 	using namespace UE::Interchange;
 
 	bool bResult = false;
 
-	TArray<FString> CapsuleCollisionMeshUids;
-	LodDataNode->GetCapsuleCollisionMeshUids(CapsuleCollisionMeshUids);
+	TMap<FInterchangeMeshPayLoadKey, FMeshPayload> CapsuleCollisionPayloads = PayloadsPerLodIndex.FindChecked(0).CollisionCapsulePayloadPerKey;
 
 	FKAggregateGeom& AggGeo = StaticMesh->GetBodySetup()->AggGeom;
 
-	TArray<FMeshPayload> MeshPayloads = GetMeshPayloads(Arguments, CapsuleCollisionMeshUids);
-	for (const FMeshPayload& MeshPayload : MeshPayloads)
+	for (TPair<FInterchangeMeshPayLoadKey, FMeshPayload>& CapsuleCollisionPayload : CapsuleCollisionPayloads)
 	{
 		const FTransform& Transform = FTransform::Identity;
-		const TOptional<FMeshPayloadData>& PayloadData = MeshPayload.PayloadData.Get();
+		const TOptional<FMeshPayloadData>& PayloadData = CapsuleCollisionPayload.Value.PayloadData;
 
 		if (!PayloadData.IsSet())
 		{
@@ -1641,22 +1706,20 @@ bool UInterchangeStaticMeshFactory::ImportCapsuleCollision(const FImportAssetObj
 }
 
 
-bool UInterchangeStaticMeshFactory::ImportSphereCollision(const FImportAssetObjectParams& Arguments, UStaticMesh* StaticMesh, const UInterchangeStaticMeshLodDataNode* LodDataNode)
+bool UInterchangeStaticMeshFactory::ImportSphereCollision(const FImportAssetObjectParams& Arguments, UStaticMesh* StaticMesh)
 {
 	using namespace UE::Interchange;
 
 	bool bResult = false;
 
-	TArray<FString> SphereCollisionMeshUids;
-	LodDataNode->GetSphereCollisionMeshUids(SphereCollisionMeshUids);
+	TMap<FInterchangeMeshPayLoadKey, FMeshPayload> SphereCollisionPayloads = PayloadsPerLodIndex.FindChecked(0).CollisionSpherePayloadPerKey;
 
 	FKAggregateGeom& AggGeo = StaticMesh->GetBodySetup()->AggGeom;
 
-	TArray<FMeshPayload> MeshPayloads = GetMeshPayloads(Arguments, SphereCollisionMeshUids);
-	for (const FMeshPayload& MeshPayload : MeshPayloads)
+	for (TPair<FInterchangeMeshPayLoadKey, FMeshPayload>& SphereCollisionPayload : SphereCollisionPayloads)
 	{
 		const FTransform& Transform = FTransform::Identity;
-		const TOptional<FMeshPayloadData>& PayloadData = MeshPayload.PayloadData.Get();
+		const TOptional<FMeshPayloadData>& PayloadData = SphereCollisionPayload.Value.PayloadData;
 
 		if (!PayloadData.IsSet())
 		{
@@ -1696,18 +1759,15 @@ bool UInterchangeStaticMeshFactory::ImportConvexCollision(const FImportAssetObje
 
 	bool bResult = false;
 
-	TArray<FString> ConvexCollisionMeshUids;
-	LodDataNode->GetConvexCollisionMeshUids(ConvexCollisionMeshUids);
-
-	TArray<FMeshPayload> MeshPayloads = GetMeshPayloads(Arguments, ConvexCollisionMeshUids);
+	TMap<FInterchangeMeshPayLoadKey, FMeshPayload> ConvexCollisionPayloads = PayloadsPerLodIndex.FindChecked(0).CollisionConvexPayloadPerKey;
 
 	bool bOneConvexHullPerUCX;
 	if (!LodDataNode->GetOneConvexHullPerUCX(bOneConvexHullPerUCX) || !bOneConvexHullPerUCX)
 	{
-		for (const FMeshPayload& MeshPayload : MeshPayloads)
+		for (TPair<FInterchangeMeshPayLoadKey, FMeshPayload>& ConvexCollisionPayload : ConvexCollisionPayloads)
 		{
 			const FTransform& Transform = FTransform::Identity;
-			const TOptional<FMeshPayloadData>& PayloadData = MeshPayload.PayloadData.Get();
+			const TOptional<FMeshPayloadData>& PayloadData = ConvexCollisionPayload.Value.PayloadData;
 
 			if (!PayloadData.IsSet())
 			{
@@ -1729,10 +1789,10 @@ bool UInterchangeStaticMeshFactory::ImportConvexCollision(const FImportAssetObje
 	{
 		FKAggregateGeom& AggGeo = StaticMesh->GetBodySetup()->AggGeom;
 
-		for (const FMeshPayload& MeshPayload : MeshPayloads)
+		for (TPair<FInterchangeMeshPayLoadKey, FMeshPayload>& ConvexCollisionPayload : ConvexCollisionPayloads)
 		{
 			const FTransform& Transform = FTransform::Identity;
-			TOptional<FMeshPayloadData> PayloadData = MeshPayload.PayloadData.Get();
+			TOptional<FMeshPayloadData> PayloadData = ConvexCollisionPayload.Value.PayloadData;
 
 			if (!PayloadData.IsSet())
 			{
diff --git a/Engine/Plugins/Interchange/Runtime/Source/Import/Private/Scene/InterchangeSceneVariantSetsFactory.cpp b/Engine/Plugins/Interchange/Runtime/Source/Import/Private/Scene/InterchangeSceneVariantSetsFactory.cpp
--- a/Engine/Plugins/Interchange/Runtime/Source/Import/Private/Scene/InterchangeSceneVariantSetsFactory.cpp
+++ b/Engine/Plugins/Interchange/Runtime/Source/Import/Private/Scene/InterchangeSceneVariantSetsFactory.cpp
@@ -98,7 +98,7 @@ namespace UE::Interchange::Private
 		FString DisplayText;
 		VariantSetNode.GetCustomDisplayText(DisplayText);
 
-		const TOptional<UE::Interchange::FVariantSetPayloadData> PayloadData = PayloadInterface.GetVariantSetPayloadData(PayloadKey).Consume();
+		const TOptional<UE::Interchange::FVariantSetPayloadData> PayloadData = PayloadInterface.GetVariantSetPayloadData(PayloadKey);
 		if (!PayloadData.IsSet())
 		{
 			UE_LOG(LogInterchangeImport, Warning, TEXT("No payload for variant set %s."), *DisplayText);
diff --git a/Engine/Plugins/Interchange/Runtime/Source/Import/Private/Usd/InterchangeUsdTranslator.cpp b/Engine/Plugins/Interchange/Runtime/Source/Import/Private/Usd/InterchangeUsdTranslator.cpp
--- a/Engine/Plugins/Interchange/Runtime/Source/Import/Private/Usd/InterchangeUsdTranslator.cpp
+++ b/Engine/Plugins/Interchange/Runtime/Source/Import/Private/Usd/InterchangeUsdTranslator.cpp
@@ -29,6 +29,7 @@
 #include "UsdWrappers/UsdTyped.h"
 
 #include "Async/Async.h"
+#include "Async/ParallelFor.h"
 #include "HAL/IConsoleManager.h"
 #include "InterchangeManager.h"
 #include "InterchangeCameraNode.h"
@@ -39,6 +40,7 @@
 #include "InterchangeShaderGraphNode.h"
 #include "InterchangeTexture2DNode.h"
 #include "InterchangeTranslatorHelper.h"
+#include "Mesh/InterchangeMeshPayload.h"
 #include "MovieSceneSection.h"
 #include "Rendering/SkeletalMeshLODImporterData.h"
 #include "StaticMeshAttributes.h"
@@ -2466,79 +2468,70 @@ void UInterchangeUSDTranslator::SetSettings(const UInterchangeTranslatorSettings
 	}
 }
 
-TFuture<TOptional<UE::Interchange::FMeshPayloadData>> UInterchangeUSDTranslator::GetMeshPayloadData(
+TOptional<UE::Interchange::FMeshPayloadData> UInterchangeUSDTranslator::GetMeshPayloadData(
 	const FInterchangeMeshPayLoadKey& PayloadKey,
 	const FTransform& MeshGlobalTransform
 ) const
 {
 	using namespace UE::InterchangeUsdTranslator::Private;
-
-	return Async(
-		EAsyncExecution::TaskGraph,
-		[this, PayloadKey, MeshGlobalTransform]
-		{
-			TOptional<UE::Interchange::FMeshPayloadData> Result;
-
+	bool bSuccess = false;
+	TOptional<UE::Interchange::FMeshPayloadData> Result;
 #if USE_USD_SDK
-			UInterchangeUSDTranslatorImpl* ImplPtr = Impl.Get();
-			if (!ImplPtr)
-			{
-				return Result;
-			}
+	UInterchangeUSDTranslatorImpl* ImplPtr = Impl.Get();
+	if (!ImplPtr)
+	{
+		return Result;
+	}
 
-			UsdToUnreal::FUsdMeshConversionOptions OptionsCopy = ImplPtr->CachedMeshConversionOptions;
-			OptionsCopy.AdditionalTransform = MeshGlobalTransform;
+	UsdToUnreal::FUsdMeshConversionOptions OptionsCopy = ImplPtr->CachedMeshConversionOptions;
+	OptionsCopy.AdditionalTransform = MeshGlobalTransform;
 
-			bool bSuccess = false;
-			UE::Interchange::FMeshPayloadData MeshPayloadData;
-			switch (PayloadKey.Type)
-			{
-				case EInterchangeMeshPayLoadType::STATIC:
-				{
-					bSuccess = UE::InterchangeUsdTranslator::Private::GetStaticMeshPayloadData(
-						PayloadKey.UniqueId,
-						*ImplPtr,
-						OptionsCopy,
-						MeshPayloadData.MeshDescription
-					);
-					break;
-				}
-				case EInterchangeMeshPayLoadType::SKELETAL:
-				{
-					bSuccess = UE::InterchangeUsdTranslator::Private::GetSkeletalMeshPayloadData(
-						PayloadKey.UniqueId,
-						*ImplPtr,
-						OptionsCopy,
-						MeshPayloadData.MeshDescription,
-						MeshPayloadData.JointNames
-					);
-					break;
-				}
-				case EInterchangeMeshPayLoadType::MORPHTARGET:
-				{
-					bSuccess = UE::InterchangeUsdTranslator::Private::GetMorphTargetPayloadData(
-						PayloadKey.UniqueId,
-						*ImplPtr,
-						OptionsCopy,
-						MeshPayloadData.MeshDescription,
-						MeshPayloadData.MorphTargetName
-					);
-					break;
-				}
-				case EInterchangeMeshPayLoadType::NONE:	   // Fallthrough
-				default:
-					break;
-			}
+	UE::Interchange::FMeshPayloadData MeshPayloadData;
+	switch (PayloadKey.Type)
+	{
+		case EInterchangeMeshPayLoadType::STATIC:
+		{
+			bSuccess = UE::InterchangeUsdTranslator::Private::GetStaticMeshPayloadData(
+				PayloadKey.UniqueId,
+				*ImplPtr,
+				OptionsCopy,
+				MeshPayloadData.MeshDescription
+			);
+			break;
+		}
+		case EInterchangeMeshPayLoadType::SKELETAL:
+		{
+			bSuccess = UE::InterchangeUsdTranslator::Private::GetSkeletalMeshPayloadData(
+				PayloadKey.UniqueId,
+				*ImplPtr,
+				OptionsCopy,
+				MeshPayloadData.MeshDescription,
+				MeshPayloadData.JointNames
+			);
+			break;
+		}
+		case EInterchangeMeshPayLoadType::MORPHTARGET:
+		{
+			bSuccess = UE::InterchangeUsdTranslator::Private::GetMorphTargetPayloadData(
+				PayloadKey.UniqueId,
+				*ImplPtr,
+				OptionsCopy,
+				MeshPayloadData.MeshDescription,
+				MeshPayloadData.MorphTargetName
+			);
+			break;
+		}
+		case EInterchangeMeshPayLoadType::NONE:	   // Fallthrough
+		default:
+			break;
+	}
 
-			if (bSuccess)
-			{
-				Result.Emplace(MeshPayloadData);
-			}
+	if (bSuccess)
+	{
+		Result.Emplace(MeshPayloadData);
+	}
 #endif	  // USE_USD_SDK
-
-			return Result;
-		}
-	);
+	return Result;
 }
 
 TOptional<UE::Interchange::FImportImage> UInterchangeUSDTranslator::GetTexturePayloadData(
@@ -2692,102 +2685,154 @@ TArray<UE::Interchange::FAnimationPayloadData> UInterchangeUSDTranslator::GetAni
 {
 	using namespace UE::Interchange;
 	using namespace UE::InterchangeUsdTranslator::Private;
-
+	//This is the results we return
 	TArray<UE::Interchange::FAnimationPayloadData> AnimationPayloads;
 
+	//Maps to help sorting the queries by payload type
+	TArray<int32> BakeQueryIndexes;
+	TArray<TArray<UE::Interchange::FAnimationPayloadData>> BakeAnimationPayloads;
+	TArray<int32> CurveQueryIndexes;
+	TArray<TArray<UE::Interchange::FAnimationPayloadData>> CurveAnimationPayloads;
+	
+	//Get all curves with a parallel for
+	int32 PayloadCount = PayloadQueries.Num();
+	for (int32 PayloadIndex = 0; PayloadIndex < PayloadCount; ++PayloadIndex)
+	{
+		const UE::Interchange::FAnimationPayloadQuery& PayloadQuery = PayloadQueries[PayloadIndex];
+		EInterchangeAnimationPayLoadType QueryType = PayloadQuery.PayloadKey.Type;
+		if (QueryType == EInterchangeAnimationPayLoadType::BAKED)
+		{
+			BakeQueryIndexes.Add(PayloadIndex);
+		}
+		else
+		{
+			CurveQueryIndexes.Add(PayloadIndex);
+		}
+	}
+
+	
 #if USE_USD_SDK
-	TMap<FString, TArray<const UE::Interchange::FAnimationPayloadQuery*>> BatchedBakeQueries;
-	BatchedBakeQueries.Reserve(PayloadQueries.Num());
 
-	// Inspect the queries we got
-	TArray<TFuture<TArray<UE::Interchange::FAnimationPayloadData>>> AnimationPayloadFutures;
-	for (const UE::Interchange::FAnimationPayloadQuery& PayloadQuery : PayloadQueries)
+	//Import the Baked curve payloads
+	if(BakeQueryIndexes.Num() > 0)
 	{
-		switch (PayloadQuery.PayloadKey.Type)
-		{
-			case EInterchangeAnimationPayLoadType::CURVE:	 // Fallthrough
-			case EInterchangeAnimationPayLoadType::STEPCURVE:
-			{
-				// Property track animation queries.
-				//
-				// We're fine handling these in isolation (currently GetAnimationPayloadData is called with
-				// a single query at a time for these): Emit a separate task for each right away
-				AnimationPayloadFutures.Add(Async(
-					EAsyncExecution::TaskGraph,
-					[this, &PayloadQuery]
-					{
-						TArray<FAnimationPayloadData> Result;
-						FAnimationPayloadData AnimationPayLoadData{PayloadQuery.SceneNodeUniqueID, PayloadQuery.PayloadKey};
-						if (GetPropertyAnimationCurvePayloadData(Impl->UsdStage, PayloadQuery.PayloadKey.UniqueId, AnimationPayLoadData))
-						{
-							Result.Emplace(AnimationPayLoadData);
-						}
-						return Result;
-					}
-				));
-				break;
-			}
-			case EInterchangeAnimationPayLoadType::BAKED:
-			{
-				// Joint transform animation queries.
-				//
-				// Currently we'll receive the PayloadQueries for all joints of a skeletal animation on the same GetAnimationPayloadData
-				// call. Unfortunately in USD we must compute all joint transforms every time, even if all we need is data for a single
-				// joint. For efficiency then, we group up all the queries for the separate joints of the same skeleton into one batch
-				// task that we can resolve in one pass
-				const FString BakedQueryHash = HashAnimPayloadQuery(PayloadQuery);
-				TArray<const UE::Interchange::FAnimationPayloadQuery*>& Queries = BatchedBakeQueries.FindOrAdd(BakedQueryHash);
-				Queries.Add(&PayloadQuery);
-				break;
-			}
-			case EInterchangeAnimationPayLoadType::MORPHTARGETCURVE:
+		int32 BakePayloadCount = BakeQueryIndexes.Num();
+		TMap<FString, TArray<const UE::Interchange::FAnimationPayloadQuery*>> BatchedBakeQueries;
+		BatchedBakeQueries.Reserve(BakePayloadCount);
+		
+		//Get the BAKED transform synchronously, since there is some interchange task that parallel them
+		for (int32 BakePayloadIndex = 0; BakePayloadIndex < BakePayloadCount; ++BakePayloadIndex)
+		{
+			if (!ensure(BakeQueryIndexes.IsValidIndex(BakePayloadIndex)))
 			{
-				// Morph target curve queries.
-				AnimationPayloadFutures.Add(Async(
-					EAsyncExecution::TaskGraph,
-					[this, &PayloadQuery]
-					{
-						TArray<FAnimationPayloadData> Result;
-						FAnimationPayloadData AnimationPayLoadData{PayloadQuery.SceneNodeUniqueID, PayloadQuery.PayloadKey};
-						if (GetMorphTargetAnimationCurvePayloadData(*Impl, PayloadQuery.PayloadKey.UniqueId, AnimationPayLoadData))
-						{
-							Result.Emplace(AnimationPayLoadData);
-						}
-						return Result;
-					}
-				));
-				break;
+				continue;
 			}
-			case EInterchangeAnimationPayLoadType::NONE:
-			default:
+			int32 PayloadIndex = BakeQueryIndexes[BakePayloadIndex];
+			if (!PayloadQueries.IsValidIndex(PayloadIndex))
 			{
-				break;
+				continue;
 			}
+			const UE::Interchange::FAnimationPayloadQuery& PayloadQuery = PayloadQueries[PayloadIndex];
+			check(PayloadQuery.PayloadKey.Type == EInterchangeAnimationPayLoadType::BAKED);
+			// Joint transform animation queries.
+			//
+			// Currently we'll receive the PayloadQueries for all joints of a skeletal animation on the same GetAnimationPayloadData
+			// call. Unfortunately in USD we must compute all joint transforms every time, even if all we need is data for a single
+			// joint. For efficiency then, we group up all the queries for the separate joints of the same skeleton into one batch
+			// task that we can resolve in one pass
+			const FString BakedQueryHash = HashAnimPayloadQuery(PayloadQuery);
+			TArray<const UE::Interchange::FAnimationPayloadQuery*>& Queries = BatchedBakeQueries.FindOrAdd(BakedQueryHash);
+			Queries.Add(&PayloadQuery);
+		}
+		// Emit the batched joint transform animation tasks
+		for (const TPair<FString, TArray<const UE::Interchange::FAnimationPayloadQuery*>>& BatchedBakedQueryPair : BatchedBakeQueries)
+		{
+			const TArray<const UE::Interchange::FAnimationPayloadQuery*>& Queries = BatchedBakedQueryPair.Value;
+			TArray<UE::Interchange::FAnimationPayloadData> Result;
+			GetJointAnimationCurvePayloadData(*Impl, Queries, Result);
+			BakeAnimationPayloads.Add(Result);
+		}
+
+		// Append the bake curves results
+		for (TArray<UE::Interchange::FAnimationPayloadData>& AnimationPayload : BakeAnimationPayloads)
+		{
+			AnimationPayloads.Append(AnimationPayload);
 		}
 	}
 
-	// Emit the batched joint transform animation tasks
-	for (const TPair<FString, TArray<const UE::Interchange::FAnimationPayloadQuery*>>& BatchedBakedQueryPair : BatchedBakeQueries)
+	//Import normal curves
+	if(CurveQueryIndexes.Num() > 0)
 	{
-		AnimationPayloadFutures.Add(Async(
-			EAsyncExecution::TaskGraph,
-			[this, &BatchedBakedQueryPair]
+		auto GetAnimPayloadLambda = [&](int32 PayloadIndex)
 			{
-				const TArray<const UE::Interchange::FAnimationPayloadQuery*>& Queries = BatchedBakedQueryPair.Value;
-				TArray<UE::Interchange::FAnimationPayloadData> Result;
-				if (GetJointAnimationCurvePayloadData(*Impl, Queries, Result))
+				if (!PayloadQueries.IsValidIndex(PayloadIndex))
 				{
-					return Result;
+					return;
+				}
+				const UE::Interchange::FAnimationPayloadQuery& PayloadQuery = PayloadQueries[PayloadIndex];
+				EInterchangeAnimationPayLoadType PayloadType = PayloadQuery.PayloadKey.Type;
+				if(PayloadType == EInterchangeAnimationPayLoadType::CURVE
+					|| PayloadType == EInterchangeAnimationPayLoadType::STEPCURVE)
+				{
+					// Property track animation queries.
+					//
+					// We're fine handling these in isolation (currently GetAnimationPayloadData is called with
+					// a single query at a time for these): Emit a separate task for each right away
+					FAnimationPayloadData AnimationPayLoadData{ PayloadQuery.SceneNodeUniqueID, PayloadQuery.PayloadKey };
+					if (GetPropertyAnimationCurvePayloadData(Impl->UsdStage, PayloadQuery.PayloadKey.UniqueId, AnimationPayLoadData))
+					{
+						CurveAnimationPayloads[PayloadIndex].Emplace(AnimationPayLoadData);
+					}
+				}
+				else if(PayloadType == EInterchangeAnimationPayLoadType::MORPHTARGETCURVE)
+				{
+					// Morph target curve queries.
+					FAnimationPayloadData AnimationPayLoadData{ PayloadQuery.SceneNodeUniqueID, PayloadQuery.PayloadKey };
+					if (GetMorphTargetAnimationCurvePayloadData(*Impl, PayloadQuery.PayloadKey.UniqueId, AnimationPayLoadData))
+					{
+						CurveAnimationPayloads[PayloadIndex].Emplace(AnimationPayLoadData);
+					}
+				}
+			};
+
+		//Get all curves with a parallel for if there is many
+		int32 CurvePayloadCount = CurveQueryIndexes.Num();
+		CurveAnimationPayloads.AddDefaulted(CurvePayloadCount);
+		const int32 BatchSize = 10;
+		if (CurvePayloadCount > BatchSize)
+		{
+			const int32 NumBatches = (CurvePayloadCount / BatchSize) + 1;
+			ParallelFor(NumBatches, [&](int32 BatchIndex)
+				{
+					int32 PayloadIndexOffset = BatchIndex * BatchSize;
+					for (int32 PayloadIndex = PayloadIndexOffset; PayloadIndex < PayloadIndexOffset + BatchSize; ++PayloadIndex)
+					{
+						//The last batch can be incomplete
+						if (!CurveQueryIndexes.IsValidIndex(PayloadIndex))
+						{
+							break;
+						}
+						GetAnimPayloadLambda(CurveQueryIndexes[PayloadIndex]);
+					}
+				}, EParallelForFlags::BackgroundPriority);// ParallelFor
+		}
+		else
+		{
+			for (int32 PayloadIndex = 0; PayloadIndex < CurvePayloadCount; ++PayloadIndex)
+			{
+				int32 PayloadQueriesIndex = CurveQueryIndexes[PayloadIndex];
+				if (PayloadQueries.IsValidIndex(PayloadQueriesIndex))
+				{
+					GetAnimPayloadLambda(PayloadQueriesIndex);
 				}
-				return TArray<UE::Interchange::FAnimationPayloadData>{};
 			}
-		));
-	}
+		}
 
-	// Wait for all async tasks to complete and collect the results
-	for (TFuture<TArray<UE::Interchange::FAnimationPayloadData>>& AnimationPayloadFuture : AnimationPayloadFutures)
-	{
-		AnimationPayloads.Append(AnimationPayloadFuture.Get());
+		// Append the curves results
+		for (TArray<UE::Interchange::FAnimationPayloadData>& AnimationPayload : CurveAnimationPayloads)
+		{
+			AnimationPayloads.Append(AnimationPayload);
+		}
 	}
 #endif	  // USE_USD_SDK
 
diff --git a/Engine/Plugins/Interchange/Runtime/Source/Import/Public/Animation/InterchangeAnimSequenceFactory.h b/Engine/Plugins/Interchange/Runtime/Source/Import/Public/Animation/InterchangeAnimSequenceFactory.h
--- a/Engine/Plugins/Interchange/Runtime/Source/Import/Public/Animation/InterchangeAnimSequenceFactory.h
+++ b/Engine/Plugins/Interchange/Runtime/Source/Import/Public/Animation/InterchangeAnimSequenceFactory.h
@@ -26,6 +26,8 @@ public:
 	virtual UClass* GetFactoryClass() const override;
 	virtual EInterchangeFactoryAssetType GetFactoryAssetType() override { return EInterchangeFactoryAssetType::Animations; }
 	
+	virtual void CreatePayloadTasks(const FImportAssetObjectParams& Arguments, bool bAsync, TArray<TSharedPtr<UE::Interchange::FInterchangeTaskBase>>& PayloadTasks) override;
+
 	virtual FImportAssetResult BeginImportAsset_GameThread(const FImportAssetObjectParams& Arguments) override;
 	
 	virtual FImportAssetResult ImportAsset_Async(const FImportAssetObjectParams& Arguments) override;
@@ -66,6 +68,22 @@ private:
 	//Morph target curves payload data
 	FMorphTargetData MorphTargetData;
 
+	//Animation queries
+	TMap< TTuple<FString, FInterchangeAnimationPayLoadKey>, UE::Interchange::FAnimationPayloadQuery> BoneAnimationPayloadQueries;
+	
+	//Animation queries results
+	TMap< TTuple<FString, FInterchangeAnimationPayLoadKey>, UE::Interchange::FAnimationPayloadData> BoneAnimationPayloadResults;
+
+	friend bool operator==(const TTuple<FString, FInterchangeAnimationPayLoadKey>& A, const TTuple<FString, FInterchangeAnimationPayLoadKey>& B)
+	{
+		return A.Key.Equals(B.Key) && A.Value == B.Value;
+	}
+
+	friend uint32 GetTypeHash(const TTuple<FString, FInterchangeAnimationPayLoadKey>& Tupple)
+	{
+		return HashCombine(GetTypeHash(Tupple.Key), GetTypeHash(Tupple.Value));
+	}
+
 	// Interchange factory base interface end
 	//////////////////////////////////////////////////////////////////////////
 };
diff --git a/Engine/Plugins/Interchange/Runtime/Source/Import/Public/Animation/InterchangeAnimationPayloadInterface.h b/Engine/Plugins/Interchange/Runtime/Source/Import/Public/Animation/InterchangeAnimationPayloadInterface.h
--- a/Engine/Plugins/Interchange/Runtime/Source/Import/Public/Animation/InterchangeAnimationPayloadInterface.h
+++ b/Engine/Plugins/Interchange/Runtime/Source/Import/Public/Animation/InterchangeAnimationPayloadInterface.h
@@ -26,6 +26,16 @@ class INTERCHANGEIMPORT_API IInterchangeAnimationPayloadInterface
 	GENERATED_BODY()
 public:
 
+	/**
+	 * Return true if the translator want to import all bone animations in one query, false otherwise.
+	 * 
+	 * @Note - The fbx translator use the fbx sdk which cache the global transform but dirty the cache every time we evaluate at a different time. The goal is to evaluate all bones at the same time.
+	 */
+	virtual bool PreferGroupingBoneAnimationQueriesTogether() const
+	{
+		return false;
+	}
+
 	/**
 	 * Get animation payload data for the specified payload key.
 	 * It return an array of FRichCurve (Rich curve are float curve we can interpolate) or array of "Step" curve or an array of Baked Transformations, depending on Type
diff --git a/Engine/Plugins/Interchange/Runtime/Source/Import/Public/Fbx/InterchangeFbxTranslator.h b/Engine/Plugins/Interchange/Runtime/Source/Import/Public/Fbx/InterchangeFbxTranslator.h
--- a/Engine/Plugins/Interchange/Runtime/Source/Import/Public/Fbx/InterchangeFbxTranslator.h
+++ b/Engine/Plugins/Interchange/Runtime/Source/Import/Public/Fbx/InterchangeFbxTranslator.h
@@ -97,13 +97,20 @@ public:
 	 * @param PayloadKey - The key to retrieve the a particular payload contain into the specified source data.
 	 * @return a PayloadData containing the imported data. The TOptional will not be set if there is an error.
 	 */
-	virtual TFuture<TOptional<UE::Interchange::FMeshPayloadData>> GetMeshPayloadData(const FInterchangeMeshPayLoadKey& PayLoadKey, const FTransform& MeshGlobalTransform) const override;
+	virtual TOptional<UE::Interchange::FMeshPayloadData> GetMeshPayloadData(const FInterchangeMeshPayLoadKey& PayLoadKey, const FTransform& MeshGlobalTransform) const override;
 
 	///* IInterchangeMeshPayloadInterface End */
 
 	//////////////////////////////////////////////////////////////////////////
 	/* IInterchangeAnimationPayloadInterface Begin */
+	
+	virtual bool PreferGroupingBoneAnimationQueriesTogether() const override
+	{
+		return true;
+	}
+
 	virtual TArray<UE::Interchange::FAnimationPayloadData> GetAnimationPayloadData(const TArray<UE::Interchange::FAnimationPayloadQuery>& PayloadQueries) const override;
+	
 	/* IInterchangeAnimationPayloadInterface End */
 private:
 	FString CreateLoadFbxFileCommand(const FString& FbxFilePath, const bool bConvertScene, const bool bForceFrontXAxis, const bool bConvertSceneUnit, const bool bKeepFbxNamespace) const;
diff --git a/Engine/Plugins/Interchange/Runtime/Source/Import/Public/Gltf/InterchangeGltfTranslator.h b/Engine/Plugins/Interchange/Runtime/Source/Import/Public/Gltf/InterchangeGltfTranslator.h
--- a/Engine/Plugins/Interchange/Runtime/Source/Import/Public/Gltf/InterchangeGltfTranslator.h
+++ b/Engine/Plugins/Interchange/Runtime/Source/Import/Public/Gltf/InterchangeGltfTranslator.h
@@ -42,7 +42,7 @@ public:
 
 	/* IInterchangeStaticMeshPayloadInterface Begin */
 
-	virtual TFuture<TOptional<UE::Interchange::FMeshPayloadData>> GetMeshPayloadData(const FInterchangeMeshPayLoadKey& PayLoadKey, const FTransform& MeshGlobalTransform) const override;
+	virtual TOptional<UE::Interchange::FMeshPayloadData> GetMeshPayloadData(const FInterchangeMeshPayLoadKey& PayLoadKey, const FTransform& MeshGlobalTransform) const override;
 
 	/* IInterchangeStaticMeshPayloadInterface End */
 
@@ -53,12 +53,12 @@ public:
 	/* IInterchangeTexturePayloadInterface End */
 
 	/* IInterchangeAnimationPayloadInterface Begin */
-	TFuture<TOptional<UE::Interchange::FAnimationPayloadData>> GetAnimationPayloadData(const UE::Interchange::FAnimationPayloadQuery& PayloadQuery) const;
+	TOptional<UE::Interchange::FAnimationPayloadData> GetAnimationPayloadData(const UE::Interchange::FAnimationPayloadQuery& PayloadQuery) const;
 	virtual TArray<UE::Interchange::FAnimationPayloadData> GetAnimationPayloadData(const TArray<UE::Interchange::FAnimationPayloadQuery>& PayloadQueries) const override;
 	/* IInterchangeAnimationPayloadInterface End */
 
 	/* IInterchangeVariantSetPayloadInterface Begin */
-	virtual TFuture<TOptional<UE::Interchange::FVariantSetPayloadData>> GetVariantSetPayloadData(const FString& PayloadKey) const override;
+	virtual TOptional<UE::Interchange::FVariantSetPayloadData> GetVariantSetPayloadData(const FString& PayloadKey) const override;
 	/* IInterchangeVariantSetPayloadInterface End */
 
 protected:
diff --git a/Engine/Plugins/Interchange/Runtime/Source/Import/Public/Mesh/InterchangeMeshPayloadInterface.h b/Engine/Plugins/Interchange/Runtime/Source/Import/Public/Mesh/InterchangeMeshPayloadInterface.h
--- a/Engine/Plugins/Interchange/Runtime/Source/Import/Public/Mesh/InterchangeMeshPayloadInterface.h
+++ b/Engine/Plugins/Interchange/Runtime/Source/Import/Public/Mesh/InterchangeMeshPayloadInterface.h
@@ -34,7 +34,7 @@ public:
 	 * @param MeshGlobalTransform - The transform to apply to the raw data when creating the payload data.
 	 * @return a PayloadData containing the the data ask with the key.
 	 */
-	virtual TFuture<TOptional<UE::Interchange::FMeshPayloadData>> GetMeshPayloadData(const FInterchangeMeshPayLoadKey& PayLoadKey, const FTransform& MeshGlobalTransform) const = 0;
+	virtual TOptional<UE::Interchange::FMeshPayloadData> GetMeshPayloadData(const FInterchangeMeshPayLoadKey& PayLoadKey, const FTransform& MeshGlobalTransform) const = 0;
 };
 
 
diff --git a/Engine/Plugins/Interchange/Runtime/Source/Import/Public/Mesh/InterchangeOBJTranslator.h b/Engine/Plugins/Interchange/Runtime/Source/Import/Public/Mesh/InterchangeOBJTranslator.h
--- a/Engine/Plugins/Interchange/Runtime/Source/Import/Public/Mesh/InterchangeOBJTranslator.h
+++ b/Engine/Plugins/Interchange/Runtime/Source/Import/Public/Mesh/InterchangeOBJTranslator.h
@@ -49,7 +49,7 @@ public:
 	 * @param PayloadKey - The key to retrieve the a particular payload contain into the specified source data.
 	 * @return a PayloadData containing the the data ask with the key.
 	 */
-	virtual TFuture<TOptional<UE::Interchange::FMeshPayloadData>> GetMeshPayloadData(const FInterchangeMeshPayLoadKey& PayLoadKey, const FTransform& MeshGlobalTransform) const override;
+	virtual TOptional<UE::Interchange::FMeshPayloadData> GetMeshPayloadData(const FInterchangeMeshPayLoadKey& PayLoadKey, const FTransform& MeshGlobalTransform) const override;
 
 	/* IInterchangeStaticMeshPayloadInterface End */
 
diff --git a/Engine/Plugins/Interchange/Runtime/Source/Import/Public/Mesh/InterchangeSkeletalMeshFactory.h b/Engine/Plugins/Interchange/Runtime/Source/Import/Public/Mesh/InterchangeSkeletalMeshFactory.h
--- a/Engine/Plugins/Interchange/Runtime/Source/Import/Public/Mesh/InterchangeSkeletalMeshFactory.h
+++ b/Engine/Plugins/Interchange/Runtime/Source/Import/Public/Mesh/InterchangeSkeletalMeshFactory.h
@@ -6,6 +6,8 @@
 #include "ClothingAsset.h"
 #include "CoreMinimal.h"
 #include "InterchangeFactoryBase.h"
+#include "InterchangeMeshNode.h"
+#include "Mesh/InterchangeMeshPayload.h"
 #include "Rendering/SkeletalMeshLODImporterData.h"
 #include "UObject/Object.h"
 #include "UObject/ObjectMacros.h"
@@ -53,6 +55,7 @@ public:
 
 	virtual UClass* GetFactoryClass() const override;
 	virtual EInterchangeFactoryAssetType GetFactoryAssetType() override { return EInterchangeFactoryAssetType::Meshes; }
+	virtual void CreatePayloadTasks(const FImportAssetObjectParams& Arguments, bool bAsync, TArray<TSharedPtr<UE::Interchange::FInterchangeTaskBase>>& PayloadTasks) override;
 	virtual FImportAssetResult BeginImportAsset_GameThread(const FImportAssetObjectParams& Arguments) override;
 	virtual FImportAssetResult ImportAsset_Async(const FImportAssetObjectParams& Arguments) override;
 	virtual FImportAssetResult EndImportAsset_GameThread(const FImportAssetObjectParams& Arguments) override;
@@ -69,9 +72,20 @@ public:
 
 	// Interchange factory base interface end
 	//////////////////////////////////////////////////////////////////////////
+
+	struct FLodPayloads
+	{
+		TMap<FInterchangeMeshPayLoadKey, TOptional<UE::Interchange::FMeshPayloadData>> MeshPayloadPerKey;
+		TMap<FInterchangeMeshPayLoadKey, TOptional<UE::Interchange::FMeshPayloadData>> MorphPayloadPerKey;
+	};
+
 private:
 	FEvent* SkeletalMeshLockPropertiesEvent = nullptr;
 
+	
+
+	TMap<int32, FLodPayloads> PayloadsPerLodIndex;
+
 	FImportAssetObjectData ImportAssetObjectData;
 };
 
diff --git a/Engine/Plugins/Interchange/Runtime/Source/Import/Public/Mesh/InterchangeStaticMeshFactory.h b/Engine/Plugins/Interchange/Runtime/Source/Import/Public/Mesh/InterchangeStaticMeshFactory.h
--- a/Engine/Plugins/Interchange/Runtime/Source/Import/Public/Mesh/InterchangeStaticMeshFactory.h
+++ b/Engine/Plugins/Interchange/Runtime/Source/Import/Public/Mesh/InterchangeStaticMeshFactory.h
@@ -5,11 +5,12 @@
 #include "CoreMinimal.h"
 #include "InterchangeFactoryBase.h"
 #include "InterchangeMeshDefinitions.h"
+#include "InterchangeMeshNode.h"
+#include "Mesh/InterchangeMeshPayload.h"
 #include "MeshDescription.h"
 #include "PhysicsEngine/AggregateGeom.h"
 #include "UObject/Object.h"
 #include "UObject/ObjectMacros.h"
-#include "Mesh/InterchangeMeshPayload.h"
 
 #include "InterchangeStaticMeshFactory.generated.h"
 
@@ -20,7 +21,6 @@ class UInterchangeStaticMeshFactoryNode;
 class UInterchangeStaticMeshLodDataNode;
 struct FKAggregateGeom;
 
-
 UCLASS(BlueprintType)
 class INTERCHANGEIMPORT_API UInterchangeStaticMeshFactory : public UInterchangeFactoryBase
 {
@@ -32,6 +32,7 @@ public:
 
 	virtual UClass* GetFactoryClass() const override;
 	virtual EInterchangeFactoryAssetType GetFactoryAssetType() override { return EInterchangeFactoryAssetType::Meshes; }
+	virtual void CreatePayloadTasks(const FImportAssetObjectParams& Arguments, bool bAsync, TArray<TSharedPtr<UE::Interchange::FInterchangeTaskBase>>& PayloadTasks) override;
 	virtual FImportAssetResult BeginImportAsset_GameThread(const FImportAssetObjectParams& Arguments) override;
 	virtual FImportAssetResult ImportAsset_Async(const FImportAssetObjectParams& Arguments) override;
 	virtual FImportAssetResult EndImportAsset_GameThread(const FImportAssetObjectParams& Arguments) override;
@@ -47,15 +48,23 @@ public:
 	//////////////////////////////////////////////////////////////////////////
 
 private:
-
 	struct FMeshPayload
 	{
 		FString MeshName;
-		TFuture<TOptional<UE::Interchange::FMeshPayloadData>> PayloadData;
+		TOptional<UE::Interchange::FMeshPayloadData> PayloadData;
 		FTransform Transform = FTransform::Identity;
 	};
 
-	TArray<FMeshPayload> GetMeshPayloads(const FImportAssetObjectParams& Arguments, const TArray<FString>& MeshUids) const;
+	struct FLodPayloads
+	{
+		TMap<FInterchangeMeshPayLoadKey, FMeshPayload> MeshPayloadPerKey;
+		TMap<FInterchangeMeshPayLoadKey, FMeshPayload> CollisionBoxPayloadPerKey;
+		TMap<FInterchangeMeshPayLoadKey, FMeshPayload> CollisionCapsulePayloadPerKey;
+		TMap<FInterchangeMeshPayLoadKey, FMeshPayload> CollisionSpherePayloadPerKey;
+		TMap<FInterchangeMeshPayLoadKey, FMeshPayload> CollisionConvexPayloadPerKey;
+	};
+
+	TMap<int32, FLodPayloads> PayloadsPerLodIndex;
 
 	bool AddConvexGeomFromVertices(const FImportAssetObjectParams& Arguments, const FMeshDescription& MeshDescription, const FTransform& Transform, FKAggregateGeom& AggGeom);
 	bool DecomposeConvexMesh(const FImportAssetObjectParams& Arguments, const FMeshDescription& MeshDescription, const FTransform& Transform, UBodySetup* BodySetup);
@@ -63,11 +72,10 @@ private:
 	bool AddSphereGeomFromVertices(const FImportAssetObjectParams& Arguments, const FMeshDescription& MeshDescription, const FTransform& Transform, FKAggregateGeom& AggGeom);
 	bool AddCapsuleGeomFromVertices(const FImportAssetObjectParams& Arguments, const FMeshDescription& MeshDescription, const FTransform& Transform, FKAggregateGeom& AggGeom);
 
-	bool ImportBoxCollision(const FImportAssetObjectParams& Arguments, UStaticMesh* StaticMesh, const UInterchangeStaticMeshLodDataNode* LodDataNode);
-	bool ImportCapsuleCollision(const FImportAssetObjectParams& Arguments, UStaticMesh* StaticMesh, const UInterchangeStaticMeshLodDataNode* LodDataNode);
-	bool ImportSphereCollision(const FImportAssetObjectParams& Arguments, UStaticMesh* StaticMesh, const UInterchangeStaticMeshLodDataNode* LodDataNode);
+	bool ImportBoxCollision(const FImportAssetObjectParams& Arguments, UStaticMesh* StaticMesh);
+	bool ImportCapsuleCollision(const FImportAssetObjectParams& Arguments, UStaticMesh* StaticMesh);
+	bool ImportSphereCollision(const FImportAssetObjectParams& Arguments, UStaticMesh* StaticMesh);
 	bool ImportConvexCollision(const FImportAssetObjectParams& Arguments, UStaticMesh* StaticMesh, const UInterchangeStaticMeshLodDataNode* LodDataNode);
-	bool GenerateKDopCollision(UStaticMesh* StaticMesh);
 	bool ImportSockets(const FImportAssetObjectParams& Arguments, UStaticMesh* StaticMesh, const UInterchangeStaticMeshFactoryNode* FactoryNode);
 
 	void CommitMeshDescriptions(UStaticMesh& StaticMesh);
diff --git a/Engine/Plugins/Interchange/Runtime/Source/Import/Public/Scene/InterchangeVariantSetPayloadInterface.h b/Engine/Plugins/Interchange/Runtime/Source/Import/Public/Scene/InterchangeVariantSetPayloadInterface.h
--- a/Engine/Plugins/Interchange/Runtime/Source/Import/Public/Scene/InterchangeVariantSetPayloadInterface.h
+++ b/Engine/Plugins/Interchange/Runtime/Source/Import/Public/Scene/InterchangeVariantSetPayloadInterface.h
@@ -74,7 +74,7 @@ public:
 	 * @param PayloadKey - The key to retrieve the a particular payload contain into the specified source data.
 	 * @return a PayloadData containing the the data ask with the key.
 	 */
-	virtual TFuture<TOptional<UE::Interchange::FVariantSetPayloadData>> GetVariantSetPayloadData(const FString& PayloadKey) const = 0;
+	virtual TOptional<UE::Interchange::FVariantSetPayloadData> GetVariantSetPayloadData(const FString& PayloadKey) const = 0;
 };
 
 
diff --git a/Engine/Plugins/Interchange/Runtime/Source/Import/Public/Usd/InterchangeUsdTranslator.h b/Engine/Plugins/Interchange/Runtime/Source/Import/Public/Usd/InterchangeUsdTranslator.h
--- a/Engine/Plugins/Interchange/Runtime/Source/Import/Public/Usd/InterchangeUsdTranslator.h
+++ b/Engine/Plugins/Interchange/Runtime/Source/Import/Public/Usd/InterchangeUsdTranslator.h
@@ -97,7 +97,7 @@ public:
 	) const;
 
 	/** Begin Interchange payload interfaces */
-	virtual TFuture<TOptional<UE::Interchange::FMeshPayloadData>> GetMeshPayloadData(
+	virtual TOptional<UE::Interchange::FMeshPayloadData> GetMeshPayloadData(
 		const FInterchangeMeshPayLoadKey& PayLoadKey,
 		const FTransform& MeshGlobalTransform
 	) const override;
diff --git a/Engine/Plugins/Interchange/Runtime/Source/Nodes/Public/InterchangeAnimationTrackSetNode.h b/Engine/Plugins/Interchange/Runtime/Source/Nodes/Public/InterchangeAnimationTrackSetNode.h
--- a/Engine/Plugins/Interchange/Runtime/Source/Nodes/Public/InterchangeAnimationTrackSetNode.h
+++ b/Engine/Plugins/Interchange/Runtime/Source/Nodes/Public/InterchangeAnimationTrackSetNode.h
@@ -46,6 +46,16 @@ public:
 		, Type(InType)
 	{
 	}
+
+	bool operator==(const FInterchangeAnimationPayLoadKey& Other) const
+	{
+		return UniqueId.Equals(Other.UniqueId) && Type == Other.Type;
+	}
+
+	friend uint32 GetTypeHash(const FInterchangeAnimationPayLoadKey& InterchangeAnimationPayLoadKey)
+	{
+		return GetTypeHash(InterchangeAnimationPayLoadKey.UniqueId + FString::FromInt(static_cast<int32>(InterchangeAnimationPayLoadKey.Type)));
+	}
 };
 
 //Interchange namespace
diff --git a/Engine/Plugins/Interchange/Runtime/Source/Nodes/Public/InterchangeMeshNode.h b/Engine/Plugins/Interchange/Runtime/Source/Nodes/Public/InterchangeMeshNode.h
--- a/Engine/Plugins/Interchange/Runtime/Source/Nodes/Public/InterchangeMeshNode.h
+++ b/Engine/Plugins/Interchange/Runtime/Source/Nodes/Public/InterchangeMeshNode.h
@@ -59,6 +59,16 @@ public:
 		, Type(InType)
 	{
 	}
+
+	bool operator==(const FInterchangeMeshPayLoadKey& Other) const
+	{
+		return UniqueId.Equals(Other.UniqueId) && Type == Other.Type;
+	}
+
+	friend uint32 GetTypeHash(const FInterchangeMeshPayLoadKey& InterchangeMeshPayLoadKey)
+	{
+		return GetTypeHash(InterchangeMeshPayLoadKey.UniqueId + FString::FromInt(static_cast<int32>(InterchangeMeshPayLoadKey.Type)));
+	}
 };
 
 
diff --git a/Engine/Source/Runtime/Interchange/Core/Private/InterchangeTaskSystem.cpp b/Engine/Source/Runtime/Interchange/Core/Private/InterchangeTaskSystem.cpp
new file mode 100644
--- /dev/null
+++ b/Engine/Source/Runtime/Interchange/Core/Private/InterchangeTaskSystem.cpp
@@ -0,0 +1,360 @@
+// Copyright Epic Games, Inc. All Rights Reserved.
+
+#include "InterchangeTaskSystem.h"
+
+#include "Async/Async.h"
+#include "Async/TaskGraphInterfaces.h"
+#include "Containers/Array.h"
+#include "Containers/Map.h"
+#include "Containers/Ticker.h"
+#include "CoreGlobals.h"
+#include "HAL/CriticalSection.h"
+#include "HAL/Thread.h"
+#include "HAL/ThreadManager.h"
+#include "InterchangeLogPrivate.h"
+#include "Misc/CoreDelegates.h"
+#include "Templates/SharedPointer.h"
+#include "UObject/GarbageCollection.h"
+
+bool UE::Interchange::FInterchangeTaskSystem::bIsCreatingSingleton = false;
+
+bool UE::Interchange::FInterchangeTaskBase::SetTaskStatus(EInterchangeTaskStatus NewTaskStatus)
+{
+	FScopeLock TaskStatusLock(&TaskStatusCriticalSection);
+	if (TaskStatus == NewTaskStatus)
+	{
+		return true;
+	}
+	//If status is done, we cannot change it
+	if (!ensure(TaskStatus != EInterchangeTaskStatus::Done))
+	{
+		return false;
+	}
+	TaskStatus = NewTaskStatus;
+	return true;
+}
+
+void UE::Interchange::FInterchangeTaskBase::SetPrerequisites(const TArray<uint64>& InPrerequisiteTasks)
+{
+	check(GetTaskStatus() == EInterchangeTaskStatus::Waiting);
+	PrerequisiteTasks = InPrerequisiteTasks;
+}
+
+void UE::Interchange::FInterchangeTaskBase::Wait() const
+{
+	FInterchangeTaskSystem::Get().WaitUntilTasksComplete({ TaskId });
+}
+
+UE::Interchange::FInterchangeTaskSystem::FInterchangeTaskSystem()
+{
+	if (!bIsCreatingSingleton)
+	{
+		ensureMsgf(bIsCreatingSingleton, TEXT("Interchange task system is a singleton you must call UE::Interchange::FInterchangeTaskSystem::Get() to access it."));
+	}
+}
+
+UE::Interchange::FInterchangeTaskSystem& UE::Interchange::FInterchangeTaskSystem::Get()
+{
+	static TSharedPtr<FInterchangeTaskSystem, ESPMode::ThreadSafe> TaskSystemPtr;
+	static bool bEndOfSingletonLife = false;
+
+	if (!TaskSystemPtr.IsValid())
+	{
+		//Cannot create the singleton outside of the game thread
+		check(IsInGameThread());
+
+		//Ensure before crashing in case this function is call after the engine exit
+		if (bEndOfSingletonLife)
+		{
+			ensure(!bEndOfSingletonLife);
+		}
+
+		bIsCreatingSingleton = true;
+		TaskSystemPtr = MakeShared<FInterchangeTaskSystem, ESPMode::ThreadSafe>().ToSharedPtr();
+		bIsCreatingSingleton = false;
+
+
+
+		TaskSystemPtr->TickTickerHandle = FTSTicker::GetCoreTicker().AddTicker(TEXT("InterchangeTaskSystemTickHandle"), 0.0f, [](float DeltaTime)
+			{
+				TaskSystemPtr->Tick();
+				return true; //Return true will keep the tick active
+			});
+			
+		FCoreDelegates::OnEnginePreExit.AddLambda([]()
+			{
+				if (!ensure(TaskSystemPtr.IsValid()))
+				{
+					return;
+				}
+
+				//We cancel any running task when we pre exit the engine
+				for (TPair<uint64, TSharedPtr<FInterchangeTaskBase, ESPMode::ThreadSafe>> TaskInfo : TaskSystemPtr->TaskPerIdMap)
+				{
+					//Cancel any task
+					TaskSystemPtr->CancelTask(TaskInfo.Key, false);
+				}
+
+				if (TaskSystemPtr->TickTickerHandle.IsValid())
+				{
+					FTSTicker::GetCoreTicker().RemoveTicker(TaskSystemPtr->TickTickerHandle);
+					TaskSystemPtr->TickTickerHandle.Reset();
+				}
+
+				//Free the singleton
+				bEndOfSingletonLife = true;
+				TaskSystemPtr = nullptr;
+			});
+
+	}
+	//Return the singleton
+	check(TaskSystemPtr.IsValid());
+	return *TaskSystemPtr.Get();
+}
+
+uint64 UE::Interchange::FInterchangeTaskSystem::AddTask(TSharedPtr<FInterchangeTaskBase, ESPMode::ThreadSafe> Task, TArray<uint64>& TaskPrerequisites)
+{
+	FScopeLock TaskPerIdMapLock(&TaskPerIdMapCriticalSection);
+	FScopeLock PriorityValueLock(&PriorityValueCriticalSection);
+
+	Task->SetTaskStatus(EInterchangeTaskStatus::Waiting);
+	Task->TaskId = PriorityValue++;
+	Task->SetPrerequisites(TaskPrerequisites);
+	TaskPerIdMap.Add(Task->TaskId, Task);
+
+	return Task->TaskId;
+}
+
+uint64 UE::Interchange::FInterchangeTaskSystem::AddTask(TSharedPtr<FInterchangeTaskBase, ESPMode::ThreadSafe> Task)
+{
+	TArray<uint64> EmptyArray;
+	return AddTask(Task, EmptyArray);
+}
+
+UE::Interchange::EInterchangeTaskStatus UE::Interchange::FInterchangeTaskSystem::GetTaskStatus(const uint64 TaskId) const
+{
+	FScopeLock TaskPerIdMapLock(&TaskPerIdMapCriticalSection);
+	if (TaskPerIdMap.Contains(TaskId))
+	{
+		TSharedPtr<FInterchangeTaskBase, ESPMode::ThreadSafe> Task = TaskPerIdMap.FindChecked(TaskId);
+		return Task->GetTaskStatus();
+	}
+
+	//If task doesn't exist we assume its done
+	return EInterchangeTaskStatus::Done;
+}
+
+void UE::Interchange::FInterchangeTaskSystem::InternalAddCancelRequestNoLock(const uint64 TaskId, bool bCancelPrerequisites)
+{
+	if (!TaskPerIdMap.Contains(TaskId))
+	{
+		//Nothing to cancel
+		return;
+	}
+
+	TSharedPtr<FInterchangeTaskBase, ESPMode::ThreadSafe> Task = TaskPerIdMap.FindChecked(TaskId);
+	uint64& CanceltaskIdRequest = CancelTaskRequests.AddDefaulted_GetRef();
+	CanceltaskIdRequest = Task->GetTaskId();
+	if (bCancelPrerequisites)
+	{
+		for (uint64 PrerequisiteTaskId : Task->GetPrerequisiteTasks())
+		{
+			InternalAddCancelRequestNoLock(PrerequisiteTaskId, bCancelPrerequisites);
+		}
+	}
+}
+
+void UE::Interchange::FInterchangeTaskSystem::CancelTask(const uint64 TaskId, bool bCancelPrerequisites)
+{
+	FScopeLock TaskPerIdMapLock(&TaskPerIdMapCriticalSection);
+	FScopeLock CancelTaskRequestsLock(&CancelTaskRequestsCriticalSection);
+	InternalAddCancelRequestNoLock(TaskId, bCancelPrerequisites);
+}
+
+void UE::Interchange::FInterchangeTaskSystem::WaitUntilTasksComplete(const TArray<uint64>& TasksToComplete)
+{
+	check(IsInGameThread());
+	const FTimespan MicroSecondTimespan = FTimespan::FromMicroseconds(1);
+	const float DeltaSecond = static_cast<float>(MicroSecondTimespan.GetTotalSeconds());
+	//Loop
+	FPlatformProcess::ConditionalSleep([this, &TasksToComplete, &MicroSecondTimespan]()
+		{
+			LLM_SCOPE_BYNAME(TEXT("Interchange"));
+			bool bTasksCompleted = true;
+			//Look if all tasks are completed
+			{
+				FScopeLock TaskPerIdMapLock(&TaskPerIdMapCriticalSection);
+				for (uint64 TaskId : TasksToComplete)
+				{
+					if (TaskPerIdMap.Contains(TaskId))
+					{
+						TSharedPtr<FInterchangeTaskBase, ESPMode::ThreadSafe> Task = TaskPerIdMap.FindChecked(TaskId);
+						if (Task->GetTaskStatus() != EInterchangeTaskStatus::Done)
+						{
+							bTasksCompleted = false;
+							break;
+						}
+					}
+				}
+			}
+			if (!bTasksCompleted)
+			{
+				//Tick all systems our tasks can depends to avoid stalling the game thread
+				
+				//We cannot tick the "core ticker" since we are already inside a tick and it will assert.
+				//We cannot tick "runnable thread" it create some issue with the garbage collector tasks
+
+				//Tick Interchange task system
+				Tick();
+
+				OnTaskSystemTick.Broadcast();
+			}
+			return bTasksCompleted;
+		}, DeltaSecond);
+}
+
+void UE::Interchange::FInterchangeTaskSystem::Tick()
+{
+	//Tick must be execute on the game thread
+	check(IsInGameThread());
+	
+	//Tick is suppose to be a safe place where we can create and manipulate UObjects
+	ensure(!IsGarbageCollecting());
+#if WITH_EDITOR
+	ensure(!UE::GetIsEditorLoadingPackage());
+#endif
+
+	FScopeLock TaskPerIdMapLock(&TaskPerIdMapCriticalSection);
+
+	//Execute the cancel task
+	{
+		FScopeLock CancelTaskRequestsLock(&CancelTaskRequestsCriticalSection);
+		for (uint64 CancelTaskId : CancelTaskRequests)
+		{
+			if (!TaskPerIdMap.Contains(CancelTaskId))
+			{
+				continue;
+			}
+			TSharedPtr<FInterchangeTaskBase, ESPMode::ThreadSafe> Task = TaskPerIdMap.FindChecked(CancelTaskId);
+			if (Task->GetTaskStatus() == EInterchangeTaskStatus::Executing)
+			{
+				//Wait until the task is done
+				Task->Wait();
+			}
+			Task->SetTaskStatus(EInterchangeTaskStatus::Done);
+		}
+	}
+	if (TaskPerIdMap.Num() == 0)
+	{
+		//We do not have any task waiting or running
+		ReleaseAndDoneTasks.Reset();
+		return;
+	}
+
+	TArray<uint64> WaitingGameThreadTick;
+	TArray<uint64> WaitingTaskGraph;
+	TArray<uint64> DoneTasks;
+	for (TPair<uint64, TSharedPtr<FInterchangeTaskBase, ESPMode::ThreadSafe>>& TaskInfo : TaskPerIdMap)
+	{
+		TSharedPtr<FInterchangeTaskBase, ESPMode::ThreadSafe> Task = TaskInfo.Value;
+		if (Task->GetTaskStatus() == EInterchangeTaskStatus::Waiting)
+		{
+			if (Task->GetTaskThread() == EInterchangeTaskThread::GameThread)
+			{
+				WaitingGameThreadTick.Add(TaskInfo.Key);
+			}
+			else if (Task->GetTaskThread() == EInterchangeTaskThread::AsyncThread)
+			{
+				WaitingTaskGraph.Add(TaskInfo.Key);
+			}
+		}
+		else if (Task->GetTaskStatus() == EInterchangeTaskStatus::Done)
+		{
+			DoneTasks.Add(TaskInfo.Key);
+		}
+	}
+
+	//Sort all waiting task by priority order, prerequisite task are always added before.
+	//The priority is the task id, lower mean must be execute before higher task id
+	auto SortLambda = [](const uint64& A, const uint64& B)
+		{
+			return A < B;
+		};
+
+	WaitingGameThreadTick.Sort(SortLambda);
+	WaitingTaskGraph.Sort(SortLambda);
+
+	auto UpdateWaitingTask = [this, &DoneTasks](const TArray<uint64>& TaskIds)
+		{
+			for (uint64 TaskId : TaskIds)
+			{
+				bool bAllPrerequisitesCompleted = true;
+				TSharedPtr<FInterchangeTaskBase, ESPMode::ThreadSafe> Task = TaskPerIdMap.FindChecked(TaskId);
+				const TArray<uint64>& TaskPrerequisites = Task->GetPrerequisiteTasks();
+				for (uint64 PrerequisiteTaskId : TaskPrerequisites)
+				{
+					if (ReleaseAndDoneTasks.Contains(PrerequisiteTaskId)
+						|| DoneTasks.Contains(PrerequisiteTaskId))
+					{
+						continue;
+					}
+					//if we have one waiting or executing task, prerequisites are not completed
+					if (TaskPerIdMap.Contains(PrerequisiteTaskId))
+					{
+						TSharedPtr<FInterchangeTaskBase, ESPMode::ThreadSafe> PrerequisiteTask = TaskPerIdMap.FindChecked(PrerequisiteTaskId);
+						if (PrerequisiteTask->GetTaskStatus() != EInterchangeTaskStatus::Done)
+						{
+							bAllPrerequisitesCompleted = false;
+							break;
+						}
+					}
+				}
+				//Skip this task if not all prerequisite are completed
+				if (!bAllPrerequisitesCompleted)
+				{
+					continue;
+				}
+
+				auto TaskCompletionHandler = [Task]()
+					{
+						//Only set the task to done if the status is executing
+						if (Task->GetTaskStatus() == EInterchangeTaskStatus::Executing)
+						{
+							Task->SetTaskStatus(EInterchangeTaskStatus::Done);
+						}
+					};
+
+				if (Task->GetTaskThread() == EInterchangeTaskThread::GameThread)
+				{
+					//Prerequisites completed, we can execute the task
+					Task->SetTaskStatus(EInterchangeTaskStatus::Executing);
+					Task->Execute();
+					TaskCompletionHandler();
+				}
+				else if (Task->GetTaskThread() == EInterchangeTaskThread::AsyncThread)
+				{
+					Task->SetTaskStatus(EInterchangeTaskStatus::Executing);
+					Task->Future = Async(EAsyncExecution::TaskGraph
+					, [Task]()->uint64 //Task callback
+					{
+						Task->Execute();
+						return Task->GetTaskId();
+					}
+					, TaskCompletionHandler); //Completion callback
+				}
+			}
+		};
+
+	//Update GameThread
+	UpdateWaitingTask(WaitingGameThreadTick);
+
+	//Update TaskGraph task
+	UpdateWaitingTask(WaitingTaskGraph);
+
+	//Release Done task
+	for (uint64 TaskId : DoneTasks)
+	{
+		TaskPerIdMap.Remove(TaskId);
+		ReleaseAndDoneTasks.Add(TaskId);
+	}
+}
diff --git a/Engine/Source/Runtime/Interchange/Core/Public/InterchangeFactoryBase.h b/Engine/Source/Runtime/Interchange/Core/Public/InterchangeFactoryBase.h
--- a/Engine/Source/Runtime/Interchange/Core/Public/InterchangeFactoryBase.h
+++ b/Engine/Source/Runtime/Interchange/Core/Public/InterchangeFactoryBase.h
@@ -4,6 +4,7 @@
 
 #include "CoreMinimal.h"
 #include "InterchangeResultsContainer.h"
+#include "InterchangeTaskSystem.h"
 #include "Nodes/InterchangeFactoryBaseNode.h"
 #include "UObject/Class.h"
 #include "UObject/Object.h"
@@ -126,6 +127,11 @@ public:
 		return nullptr;
 	}
 
+	virtual void CreatePayloadTasks(const FImportAssetObjectParams& Arguments, bool bAsync, TArray<TSharedPtr<UE::Interchange::FInterchangeTaskBase>>& PayloadTasks)
+	{
+		return;
+	}
+
 	/**
 	 * Override this function to import/reimport source data and configure the properties synchronously.
 	 * Create the asset package on the game thread because it's not thread-safe.
diff --git a/Engine/Source/Runtime/Interchange/Core/Public/InterchangeTaskSystem.h b/Engine/Source/Runtime/Interchange/Core/Public/InterchangeTaskSystem.h
new file mode 100644
--- /dev/null
+++ b/Engine/Source/Runtime/Interchange/Core/Public/InterchangeTaskSystem.h
@@ -0,0 +1,236 @@
+// Copyright Epic Games, Inc. All Rights Reserved.
+
+#pragma once
+
+#include "Async/Future.h"
+#include "Containers/Array.h"
+#include "Containers/Map.h"
+#include "Containers/Ticker.h"
+#include "HAL/CriticalSection.h"
+#include "Misc/ScopeLock.h"
+
+#define INTERCHANGE_INVALID_TASK_ID 0xFFFFFFFFFFFFFFFF
+
+DECLARE_MULTICAST_DELEGATE(FOnInterchangeTaskSystemTick)
+
+namespace UE::Interchange
+{
+	//Forward declare to add it has friend
+	class FInterchangeTaskSystem;
+
+	enum class EInterchangeTaskStatus : uint8
+	{
+		Waiting, //Task is queue for execution when prerequisite are terminate and resource are available
+		Executing, //Task is being executing
+		Done //Task was execute and is now terminate
+	};
+
+	enum class EInterchangeTaskThread : uint8
+	{
+		GameThread, //Task will be tick in a ticker in the engine tick on the game thread. This is a safe place to manipulate UObject
+		AsyncThread //Task will be call on asynchronous thread, code running inside those task must be thread safe.
+	};
+
+	class FInterchangeTaskBase
+	{
+	public:
+		FInterchangeTaskBase()
+		{}
+
+		virtual ~FInterchangeTaskBase()
+		{
+			if (GetTaskStatus() == EInterchangeTaskStatus::Executing)
+			{
+				if (Future.IsValid())
+				{
+					Future.Get();
+				}
+			}
+		}
+
+		INTERCHANGECORE_API uint64 GetTaskId() const
+		{
+			return TaskId;
+		}
+
+		INTERCHANGECORE_API virtual EInterchangeTaskThread GetTaskThread() const
+		{
+			return EInterchangeTaskThread::GameThread;
+		}
+
+		INTERCHANGECORE_API const TArray<uint64>& GetPrerequisiteTasks() const
+		{
+			return PrerequisiteTasks;
+		}
+
+		INTERCHANGECORE_API EInterchangeTaskStatus GetTaskStatus() const
+		{
+			FScopeLock TaskStatusLock(&TaskStatusCriticalSection);
+			return TaskStatus;
+		}
+
+		/**
+		 * This function must return the task id to fill the promise of the execution.
+		 * 
+		 * If the caller want to re-enqueue the task and not run the subsequent tasks after
+		 * this function, It must set the task status to Waiting and the task will be
+		 * re-execute later on another tick. This is the best way to implement a async wait.
+		 */
+		INTERCHANGECORE_API virtual void Execute() = 0;
+
+		/**
+		 * Wait can be call only on game thread. Async thread wait can create dead lock by using all the thread resources.
+		 * 
+		 * The caller have some restriction to follow
+		 *    - This must be a safe place for any game thread task that will be executing during the wait.
+		 *        - Wait should be call from a engine tick and no lock (like the global UObject locks) should be taken before calling wait.
+		 *        - Slate ui callback are safe
+		 *        - Engine tick is safe, use FTSTicker::GetCoreTicker() when you can
+		 *    - Wait should not happen from a package/asset load call stack
+		 *    - Wait should not happen from a save package/asset call stack
+		 *    - Wait should not happen from a garbage collect call stack
+		 * 
+		 * If your system need to wait and you already have some locks or you are under a unsafe place you need to redesign your system to
+		 * move the wait before you get to a unsafe call stack place to wait.
+		 */
+		INTERCHANGECORE_API void Wait() const;
+	protected:
+
+		//Only friend FInterchangeTaskSystem can set a new status
+		INTERCHANGECORE_API bool SetTaskStatus(EInterchangeTaskStatus NewTaskStatus);
+
+		INTERCHANGECORE_API void SetPrerequisites(const TArray<uint64>& InPrerequisiteTasks);
+
+	private:
+
+		mutable FCriticalSection TaskStatusCriticalSection;
+		uint64 TaskId = INTERCHANGE_INVALID_TASK_ID;
+		EInterchangeTaskStatus TaskStatus = EInterchangeTaskStatus::Waiting;
+		TArray<uint64> PrerequisiteTasks;
+
+		//If a task is asynchronous the future will be set.
+		//The future will contain the task id when it will be ready.
+		TFuture<uint64> Future;
+
+		friend FInterchangeTaskSystem;
+	};
+
+	/**
+	 * This class is use to run a lambda with the interchange task system.
+	 */
+	class FInterchangeTaskLambda : public FInterchangeTaskBase
+	{
+	private:
+		EInterchangeTaskThread TaskThread = EInterchangeTaskThread::GameThread;
+		TFunction<void()> ExecuteLambda;
+		
+
+	public:
+		FInterchangeTaskLambda(EInterchangeTaskThread InTaskThread, TFunction<void()> InExecuteLambda)
+			: TaskThread(InTaskThread)
+			, ExecuteLambda(MoveTemp(InExecuteLambda))
+		{
+		}
+
+		virtual EInterchangeTaskThread GetTaskThread() const override
+		{
+			return TaskThread;
+		}
+
+		virtual void Execute()
+		{
+			ExecuteLambda();
+		}
+	};
+
+	class FInterchangeTaskSystem
+	{
+	public:
+		/**
+		 * The constructor should not be called directly, an ensure will be set if its the case
+		 * Use the static Get function to acces the singleton.
+		 */
+		FInterchangeTaskSystem();
+
+		//The interchange task system is a singleton
+		static INTERCHANGECORE_API FInterchangeTaskSystem& Get();
+
+		/**
+		 * Add a task with some prerequisites. Return the task id.
+		 * If there is an issue the returned task id will be INTERCHANGE_INVALID_TASK_ID
+		 * The task will start when all prerequisites are done and some thread resource will be available.
+		 */
+		INTERCHANGECORE_API uint64 AddTask(TSharedPtr<FInterchangeTaskBase, ESPMode::ThreadSafe> Task, TArray<uint64>& TaskPrerequisites);
+
+		/**
+		 * Add a task without prerequisite. Return the task id.
+		 * If there is an issue the returned task id will be INTERCHANGE_INVALID_TASK_ID
+		 * The task will start when some thread resource will be available.
+		 */
+		INTERCHANGECORE_API uint64 AddTask(TSharedPtr<FInterchangeTaskBase, ESPMode::ThreadSafe> Task);
+
+		/** Return the current task status */
+		INTERCHANGECORE_API EInterchangeTaskStatus GetTaskStatus(const uint64 TaskId) const;
+
+		/**
+		 * Cancel a task, you can control if you cancel also the prerequisites, and subsequents.
+		 * Canceling a task that is waiting will simply push it state to Done which will allow subsequent to be executed.
+		 */
+		INTERCHANGECORE_API void CancelTask(const uint64 TaskId, const bool bCancelPrerequisites);
+
+		/**
+		 * Wait can be call only on game thread. Async thread wait can create dead lock by using all the thread resources.
+		 *
+		 * The caller have some restriction to follow
+		 *    - This must be a safe place for any game thread task that will be executing during the wait.
+		 *        - Wait should be call from a engine tick and no lock (like the global UObject locks) should be taken before calling wait.
+		 *        - Slate ui callback are safe
+		 *        - Engine tick is safe, use FTSTicker::GetCoreTicker() when you can
+		 *    - Wait should not happen from a package/asset load call stack
+		 *    - Wait should not happen from a save package/asset call stack
+		 *    - Wait should not happen from a garbage collect call stack
+		 *
+		 * If your system need to wait and you already have some locks or you are under a unsafe place you need to redesign your system to
+		 * move the wait before you get to a unsafe call stack place to wait.
+		 */
+		INTERCHANGECORE_API void WaitUntilTasksComplete(const TArray<uint64>& TasksToComplete);
+
+		INTERCHANGECORE_API FOnInterchangeTaskSystemTick& OnTaskSystemTickDelegate()
+		{
+			return OnTaskSystemTick;
+		}
+
+	private:
+
+		static bool bIsCreatingSingleton;
+
+		FOnInterchangeTaskSystemTick OnTaskSystemTick;
+
+		//Internal cancel request with no lock, lock is handle by the caller (CancelTask)
+		void InternalAddCancelRequestNoLock(const uint64 TaskId, bool bCancelPrerequisites);
+
+		/**
+		 * The tick is where we start task execution and where we update the tasks status
+		 */
+		void Tick();
+
+		//The ticker we use to tick in the engine tick
+		FTSTicker::FDelegateHandle	TickTickerHandle;
+
+		// Store the tasks per id
+		TMap<uint64, TSharedPtr<FInterchangeTaskBase, ESPMode::ThreadSafe>> TaskPerIdMap;
+		mutable FCriticalSection TaskPerIdMapCriticalSection;
+
+		//When a task is done, we add it here and release the task. Prerequisite are search in the list and here.
+		//A clean up is done when the TaskPerIdMap is empty.
+		TArray<uint64> ReleaseAndDoneTasks;
+
+		//The priority value is a counter that is increment for each task added to the system to create the task id
+		//The lower task id are execute earlier if possible.
+		uint64 PriorityValue = 0;
+		FCriticalSection PriorityValueCriticalSection;
+
+		TArray<uint64> CancelTaskRequests;
+		FCriticalSection CancelTaskRequestsCriticalSection;
+	};
+} //namespace UE::Interchange
diff --git a/Engine/Source/Runtime/Interchange/Engine/Private/InterchangeManager.cpp b/Engine/Source/Runtime/Interchange/Engine/Private/InterchangeManager.cpp
--- a/Engine/Source/Runtime/Interchange/Engine/Private/InterchangeManager.cpp
+++ b/Engine/Source/Runtime/Interchange/Engine/Private/InterchangeManager.cpp
@@ -2,6 +2,7 @@
 
 #include "InterchangeManager.h"
 
+#include "AssetCompilingManager.h"
 #include "AssetRegistry/AssetData.h"
 #include "AssetRegistry/AssetDataTagMap.h"
 #include "AssetRegistry/AssetRegistryModule.h"
@@ -651,6 +652,7 @@ TArray<uint64> UE::Interchange::FImportAsyncHelper::GetCompletionTaskGraphEvent(
 	UE::Interchange::FInterchangeTaskSystem::Get().WaitUntilTasksComplete(TasksToComplete);
 	TasksToComplete.Reset();
 
+	TasksToComplete.Append(ImportObjectQueryPayloadsTasks);
 	TasksToComplete.Append(BeginImportObjectTasks);
 	TasksToComplete.Append(ImportObjectTasks);
 	TasksToComplete.Append(FinalizeImportObjectTasks);
@@ -1022,8 +1024,15 @@ UInterchangeManager& UInterchangeManager::GetInterchangeManager()
 		//Make sure we have a valid editor utilities
 		InterchangeManager->SetEditorUtilities(UInterchangeEditorUtilitiesBase::StaticClass());
 
+		FDelegateHandle OnTaskSystemTickDelegateHandle = UE::Interchange::FInterchangeTaskSystem::Get().OnTaskSystemTickDelegate().AddLambda([]()
+			{
+				//Tick task graph game thread tasks
+				FTaskGraphInterface::Get().ProcessThreadUntilIdle(ENamedThreads::GameThread);
+				//Tick the all compiling manager
+				FAssetCompilingManager::Get().ProcessAsyncTasks(true);
+			});
 		//We cancel any running task when we pre exit the engine
-		FCoreDelegates::OnEnginePreExit.AddLambda([]()
+		FCoreDelegates::OnEnginePreExit.AddLambda([OnTaskSystemTickDelegateHandle]()
 		{
 			//InterchangeManager should be valid at this point since this lambda is where the strong reference pointer get reset.
 			if (!ensure(InterchangeManager.IsValid()))
@@ -1032,6 +1041,8 @@ UInterchangeManager& UInterchangeManager::GetInterchangeManager()
 				return;
 			}
 
+			UE::Interchange::FInterchangeTaskSystem::Get().OnTaskSystemTickDelegate().Remove(OnTaskSystemTickDelegateHandle);
+
 			if (IsInterchangeImportEnabled())
 			{
 				//If a user run a editor commandlet, we want to finish the import before the editor close.
diff --git a/Engine/Source/Runtime/Interchange/Engine/Private/InterchangeTaskSystem.cpp b/Engine/Source/Runtime/Interchange/Engine/Private/InterchangeTaskSystem.cpp
deleted file mode 100644
diff --git a/Engine/Source/Runtime/Interchange/Engine/Private/Tasks/InterchangeTaskImportObject.cpp b/Engine/Source/Runtime/Interchange/Engine/Private/Tasks/InterchangeTaskImportObject.cpp
--- a/Engine/Source/Runtime/Interchange/Engine/Private/Tasks/InterchangeTaskImportObject.cpp
+++ b/Engine/Source/Runtime/Interchange/Engine/Private/Tasks/InterchangeTaskImportObject.cpp
@@ -274,10 +274,11 @@ void UE::Interchange::FTaskImportObject_GameThread::Execute()
 	check(IsInGameThread());
 
 	// Create factory
-	UInterchangeFactoryBase* Factory = NewObject<UInterchangeFactoryBase>(GetTransientPackage(), FactoryClass);
-	Factory->SetResultsContainer(AsyncHelper->AssetImportResult->GetResults());
-
-	AsyncHelper->AddCreatedFactory(FactoryNode->GetUniqueID(), Factory);
+	UInterchangeFactoryBase* Factory = AsyncHelper->GetCreatedFactory(FactoryNode->GetUniqueID());
+	if (!ensure(Factory))
+	{
+		return;
+	}
 
 	UPackage* Pkg = nullptr;
 	FString PackageName;
diff --git a/Engine/Source/Runtime/Interchange/Engine/Private/Tasks/InterchangeTaskParsing.cpp b/Engine/Source/Runtime/Interchange/Engine/Private/Tasks/InterchangeTaskParsing.cpp
--- a/Engine/Source/Runtime/Interchange/Engine/Private/Tasks/InterchangeTaskParsing.cpp
+++ b/Engine/Source/Runtime/Interchange/Engine/Private/Tasks/InterchangeTaskParsing.cpp
@@ -291,13 +291,50 @@ void UE::Interchange::FTaskParsing::Execute()
 				TEXT("Found multiple task data with the same asset name (%s). Only one will be executed."), *AssetFullPath))
 			{
 				LLM_SCOPE_BYNAME(TEXT("Interchange"));
+				
+				// Create factory
+				UInterchangeFactoryBase* Factory = NewObject<UInterchangeFactoryBase>(GetTransientPackage(), FactoryClass);
+				Factory->SetResultsContainer(AsyncHelper->AssetImportResult->GetResults());
+				AsyncHelper->AddCreatedFactory(FactoryNode->GetUniqueID(), Factory);
+
+				//If the factory wants to get async payloads it can add the task here so the payloads are in the dependency graph
+				TArray<uint64> ImportObjectQueryPayloadsTasksPrerequistes;
+				{
+					TArray<TSharedPtr<FInterchangeTaskBase>> PayloadTasks;
+					UInterchangeFactoryBase::FImportAssetObjectParams ImportAssetObjectParams;
+					ImportAssetObjectParams.AssetName = FactoryNode->GetDisplayLabel();
+					ImportAssetObjectParams.AssetNode = FactoryNode;
+					ImportAssetObjectParams.NodeContainer = AsyncHelper->BaseNodeContainers[SourceIndex].Get();
+					ImportAssetObjectParams.SourceData = AsyncHelper->SourceDatas[SourceIndex];
+					ImportAssetObjectParams.Translator = AsyncHelper->Translators[SourceIndex];
+					bool bAsyncPayloads = !AsyncHelper->bRunSynchronous;
+					if (!ImportAssetObjectParams.Translator || !ImportAssetObjectParams.Translator->IsThreadSafe())
+					{
+						//Translator not thread safe cannot get payloads asynchronously
+						bAsyncPayloads = false;
+					}
+					Factory->CreatePayloadTasks(ImportAssetObjectParams, bAsyncPayloads, PayloadTasks);
+					if (!PayloadTasks.IsEmpty())
+					{
+						for (TSharedPtr<FInterchangeTaskBase> PayloadTask : PayloadTasks)
+						{
+							int32 ImportObjectQueryPayloadsTaskIndex = AsyncHelper->ImportObjectQueryPayloadsTasks.Add(FInterchangeTaskSystem::Get().AddTask(PayloadTask, TaskData.Prerequisites));
+							ImportObjectQueryPayloadsTasksPrerequistes.Add(AsyncHelper->ImportObjectQueryPayloadsTasks[ImportObjectQueryPayloadsTaskIndex]);
+						}
+					}
+					else
+					{
+						ImportObjectQueryPayloadsTasksPrerequistes = TaskData.Prerequisites;
+					}
+				}
+
 				TArray<uint64> ImportObjectTasksPrerequistes;
 				TSharedPtr<FTaskImportObject_GameThread, ESPMode::ThreadSafe> TaskImportObject_GameThread = MakeShared<FTaskImportObject_GameThread, ESPMode::ThreadSafe>(AsyncHelper->ContentBasePath
 					, SourceIndex
 					, WeakAsyncHelper
 					, FactoryNode
 					, FactoryClass);
-				int32 BeginImportObjectTaskIndex = AsyncHelper->BeginImportObjectTasks.Add(FInterchangeTaskSystem::Get().AddTask(TaskImportObject_GameThread, TaskData.Prerequisites));
+				int32 BeginImportObjectTaskIndex = AsyncHelper->BeginImportObjectTasks.Add(FInterchangeTaskSystem::Get().AddTask(TaskImportObject_GameThread, ImportObjectQueryPayloadsTasksPrerequistes));
 				ImportObjectTasksPrerequistes.Add(AsyncHelper->BeginImportObjectTasks[BeginImportObjectTaskIndex]);
 
 				TSharedPtr<FTaskImportObject_Async, ESPMode::ThreadSafe> TaskImportObject_Async = MakeShared<FTaskImportObject_Async, ESPMode::ThreadSafe>(AsyncHelper->ContentBasePath
diff --git a/Engine/Source/Runtime/Interchange/Engine/Public/InterchangeManager.h b/Engine/Source/Runtime/Interchange/Engine/Public/InterchangeManager.h
--- a/Engine/Source/Runtime/Interchange/Engine/Public/InterchangeManager.h
+++ b/Engine/Source/Runtime/Interchange/Engine/Public/InterchangeManager.h
@@ -247,6 +247,7 @@ namespace UE
 			TArray<uint64> WaitAssetCompilationTasks;
 			TArray<uint64> PostImportTasks;
 			uint64 ParsingTask;
+			TArray<uint64> ImportObjectQueryPayloadsTasks;
 			TArray<uint64> BeginImportObjectTasks;
 			TArray<uint64> ImportObjectTasks;
 			TArray<uint64> FinalizeImportObjectTasks;
diff --git a/Engine/Source/Runtime/Interchange/Engine/Public/InterchangeTaskSystem.h b/Engine/Source/Runtime/Interchange/Engine/Public/InterchangeTaskSystem.h
deleted file mode 100644
