Interchange: Optimizations
- Task system now have a limit time to tick gamethread task to not make the editor unresponsive.

#rnx

diff --git a/Engine/Plugins/Interchange/Runtime/Source/Import/Private/Animation/InterchangeAnimSequenceFactory.cpp b/Engine/Plugins/Interchange/Runtime/Source/Import/Private/Animation/InterchangeAnimSequenceFactory.cpp
--- a/Engine/Plugins/Interchange/Runtime/Source/Import/Private/Animation/InterchangeAnimSequenceFactory.cpp
+++ b/Engine/Plugins/Interchange/Runtime/Source/Import/Private/Animation/InterchangeAnimSequenceFactory.cpp
@@ -932,6 +932,7 @@ UClass* UInterchangeAnimSequenceFactory::GetFactoryClass() const
 
 void UInterchangeAnimSequenceFactory::CreatePayloadTasks(const FImportAssetObjectParams& Arguments, bool bAsync, TArray<TSharedPtr<UE::Interchange::FInterchangeTaskBase>>& PayloadTasks)
 {
+	TRACE_CPUPROFILER_EVENT_SCOPE(UInterchangeAnimSequenceFactory::CreatePayloadTasks);
 	const IInterchangeAnimationPayloadInterface* AnimSequenceTranslatorPayloadInterface = Cast<IInterchangeAnimationPayloadInterface>(Arguments.Translator);
 	if (!AnimSequenceTranslatorPayloadInterface)
 	{
@@ -978,6 +979,7 @@ void UInterchangeAnimSequenceFactory::CreatePayloadTasks(const FImportAssetObjec
 			TSharedPtr<UE::Interchange::FInterchangeTaskLambda, ESPMode::ThreadSafe> TaskGetAnimationPayload = MakeShared<UE::Interchange::FInterchangeTaskLambda, ESPMode::ThreadSafe>(bAsync ? UE::Interchange::EInterchangeTaskThread::AsyncThread : UE::Interchange::EInterchangeTaskThread::GameThread
 				, [this, AnimSequenceTranslatorPayloadInterface]()
 				{
+					TRACE_CPUPROFILER_EVENT_SCOPE(UInterchangeAnimSequenceFactory::GetAnimationPayloadMultipleQueries);
 					TArray<UE::Interchange::FAnimationPayloadQuery> TmpPayloadQueries;
 					BoneAnimationPayloadQueries.GenerateValueArray(TmpPayloadQueries);
 					TArray<UE::Interchange::FAnimationPayloadData> AnimationPayloads = AnimSequenceTranslatorPayloadInterface->GetAnimationPayloadData(TmpPayloadQueries);
@@ -1009,6 +1011,7 @@ void UInterchangeAnimSequenceFactory::CreatePayloadTasks(const FImportAssetObjec
 				TSharedPtr<UE::Interchange::FInterchangeTaskLambda, ESPMode::ThreadSafe> TaskGetAnimationPayload = MakeShared<UE::Interchange::FInterchangeTaskLambda, ESPMode::ThreadSafe>(bAsync ? UE::Interchange::EInterchangeTaskThread::AsyncThread : UE::Interchange::EInterchangeTaskThread::GameThread
 					, [this, SceneNodeUidAndPayloadKey, AnimSequenceTranslatorPayloadInterface]()
 					{
+						TRACE_CPUPROFILER_EVENT_SCOPE(UInterchangeAnimSequenceFactory::GetAnimationPayloadSingleQuery);
 						UE::Interchange::FAnimationPayloadQuery& AnimationQuery = BoneAnimationPayloadQueries.FindChecked(SceneNodeUidAndPayloadKey);
 						TArray<UE::Interchange::FAnimationPayloadData> AnimationPayloads = AnimSequenceTranslatorPayloadInterface->GetAnimationPayloadData({ AnimationQuery });
 						if (ensure(AnimationPayloads.IsValidIndex(0) && BoneAnimationPayloadResults.Contains(SceneNodeUidAndPayloadKey)))
diff --git a/Engine/Plugins/Interchange/Runtime/Source/Import/Private/Mesh/InterchangeSkeletalMeshFactory.cpp b/Engine/Plugins/Interchange/Runtime/Source/Import/Private/Mesh/InterchangeSkeletalMeshFactory.cpp
--- a/Engine/Plugins/Interchange/Runtime/Source/Import/Private/Mesh/InterchangeSkeletalMeshFactory.cpp
+++ b/Engine/Plugins/Interchange/Runtime/Source/Import/Private/Mesh/InterchangeSkeletalMeshFactory.cpp
@@ -1249,6 +1249,7 @@ void UInterchangeSkeletalMeshFactory::CreatePayloadTasks(const FImportAssetObjec
 			TSharedPtr<FInterchangeTaskLambda, ESPMode::ThreadSafe> TaskGetMeshPayload = MakeShared<FInterchangeTaskLambda, ESPMode::ThreadSafe>(bAsync ? EInterchangeTaskThread::AsyncThread : EInterchangeTaskThread::GameThread
 				, [&MeshPayload, ApplyTransformWhenFetchPayload, MeshTranslatorPayloadInterface, PayLoadKey = MeshNodeContext.TranslatorPayloadKey]()
 				{
+					TRACE_CPUPROFILER_EVENT_SCOPE(UInterchangeSkeletalMeshFactory::GetMeshPayloadDataTask);
 					MeshPayload = MeshTranslatorPayloadInterface->GetMeshPayloadData(PayLoadKey, ApplyTransformWhenFetchPayload);
 				});
 			PayloadTasks.Add(TaskGetMeshPayload);
@@ -1274,6 +1275,7 @@ void UInterchangeSkeletalMeshFactory::CreatePayloadTasks(const FImportAssetObjec
 						TSharedPtr<FInterchangeTaskLambda, ESPMode::ThreadSafe> TaskGetMorphPayload = MakeShared<FInterchangeTaskLambda, ESPMode::ThreadSafe>(bAsync ? EInterchangeTaskThread::AsyncThread : EInterchangeTaskThread::GameThread
 							, [&MorphPayload, ApplyTransformWhenFetchPayload, MeshTranslatorPayloadInterface, MorphPayLoadKey]()
 							{
+								TRACE_CPUPROFILER_EVENT_SCOPE(UInterchangeSkeletalMeshFactory::GetMeshMorphTargetPayloadDataTask);
 								MorphPayload = MeshTranslatorPayloadInterface->GetMeshPayloadData(MorphPayLoadKey, ApplyTransformWhenFetchPayload);
 							});
 						PayloadTasks.Add(TaskGetMorphPayload);
diff --git a/Engine/Plugins/Interchange/Runtime/Source/Import/Private/Mesh/InterchangeStaticMeshFactory.cpp b/Engine/Plugins/Interchange/Runtime/Source/Import/Private/Mesh/InterchangeStaticMeshFactory.cpp
--- a/Engine/Plugins/Interchange/Runtime/Source/Import/Private/Mesh/InterchangeStaticMeshFactory.cpp
+++ b/Engine/Plugins/Interchange/Runtime/Source/Import/Private/Mesh/InterchangeStaticMeshFactory.cpp
@@ -263,6 +263,7 @@ void UInterchangeStaticMeshFactory::CreatePayloadTasks(const FImportAssetObjectP
 						TSharedPtr<UE::Interchange::FInterchangeTaskLambda, ESPMode::ThreadSafe> TaskGetMeshPayload = MakeShared<UE::Interchange::FInterchangeTaskLambda, ESPMode::ThreadSafe>(bAsync ? UE::Interchange::EInterchangeTaskThread::AsyncThread : UE::Interchange::EInterchangeTaskThread::GameThread
 							, [&Payload, GlobalMeshTransform, MeshTranslatorPayloadInterface, PayLoadKey]()
 							{
+								TRACE_CPUPROFILER_EVENT_SCOPE(UInterchangeStaticMeshFactory::GetMeshPayloadDataTask);
 								if (ensure(!Payload.PayloadData.IsSet()))
 								{
 									Payload.PayloadData = MeshTranslatorPayloadInterface->GetMeshPayloadData(PayLoadKey, GlobalMeshTransform);
diff --git a/Engine/Source/Runtime/Interchange/Core/Private/InterchangeTaskSystem.cpp b/Engine/Source/Runtime/Interchange/Core/Private/InterchangeTaskSystem.cpp
--- a/Engine/Source/Runtime/Interchange/Core/Private/InterchangeTaskSystem.cpp
+++ b/Engine/Source/Runtime/Interchange/Core/Private/InterchangeTaskSystem.cpp
@@ -215,6 +215,7 @@ void UE::Interchange::FInterchangeTaskSystem::WaitUntilTasksComplete(const TArra
 
 void UE::Interchange::FInterchangeTaskSystem::Tick()
 {
+	TRACE_CPUPROFILER_EVENT_SCOPE(UE::Interchange::FInterchangeTaskSystem::Tick)
 	//Tick must be execute on the game thread
 	check(IsInGameThread());
 	
@@ -284,8 +285,9 @@ void UE::Interchange::FInterchangeTaskSystem::Tick()
 	WaitingGameThreadTick.Sort(SortLambda);
 	WaitingTaskGraph.Sort(SortLambda);
 
-	auto UpdateWaitingTask = [this, &DoneTasks](const TArray<uint64>& TaskIds)
+	auto UpdateWaitingTask = [this, &DoneTasks](const TArray<uint64>& TaskIds, bool bTimeLimit)
 		{
+			double StartTime = bTimeLimit ? FPlatformTime::Seconds() : 0.0;
 			for (uint64 TaskId : TaskIds)
 			{
 				bool bAllPrerequisitesCompleted = true;
@@ -342,14 +344,21 @@ void UE::Interchange::FInterchangeTaskSystem::Tick()
 					}
 					, TaskCompletionHandler); //Completion callback
 				}
+
+				if (bTimeLimit && (FPlatformTime::Seconds() - StartTime > 0.03333))
+				{
+					break;
+				}
 			}
 		};
 
-	//Update GameThread
-	UpdateWaitingTask(WaitingGameThreadTick);
+	//Update GameThread tasks with some time limits to not stall the game thread
+	constexpr bool bTimeLimitTrue = true;
+	UpdateWaitingTask(WaitingGameThreadTick, bTimeLimitTrue);
 
-	//Update TaskGraph task
-	UpdateWaitingTask(WaitingTaskGraph);
+	//Update TaskGraph tasks without time limits since they are asynchronous.
+	constexpr bool bTimeLimitFalse = false;
+	UpdateWaitingTask(WaitingTaskGraph, bTimeLimitFalse);
 
 	//Release Done task
 	for (uint64 TaskId : DoneTasks)
