USD: Detect when user tries using an USD module without the correct setup (for some scenarios).

diff --git a/Engine/Plugins/Experimental/ChaosCachingUSD/Source/ChaosCachingUSD/Public/ChaosCachingUSD/Operations.h b/Engine/Plugins/Experimental/ChaosCachingUSD/Source/ChaosCachingUSD/Public/ChaosCachingUSD/Operations.h
--- a/Engine/Plugins/Experimental/ChaosCachingUSD/Source/ChaosCachingUSD/Public/ChaosCachingUSD/Operations.h
+++ b/Engine/Plugins/Experimental/ChaosCachingUSD/Source/ChaosCachingUSD/Public/ChaosCachingUSD/Operations.h
@@ -16,15 +16,6 @@
 
 namespace UE::ChaosCachingUSD
 {
-	template <class TV>
-	pxr::VtArray<pxr::GfVec3f> ToVtVec3Array(const TArray<TV>& V3Array)
-	{
-		FScopedUsdAllocs UsdAllocs; // Use USD memory allocator
-		pxr::VtArray<pxr::GfVec3f> RetVal(static_cast<size_t>(V3Array.Num()));
-		for (int32 i = 0; i < V3Array.Num(); i++) RetVal[i].Set(V3Array[i][0], V3Array[i][1], V3Array[i][2]);
-		return RetVal;
-	}
-
 	bool ValuesDiffer(const pxr::VtArray<pxr::GfVec3f>& A, const pxr::VtArray<pxr::GfVec3f>& B, const float Tolerance = 1.0e-8, const uint64 stride = 1)
 	{
 		if (A.size() != B.size()) return true;
diff --git a/Engine/Plugins/Experimental/ChaosFlesh/Source/ChaosFleshEngine/ChaosFleshEngine.Build.cs b/Engine/Plugins/Experimental/ChaosFlesh/Source/ChaosFleshEngine/ChaosFleshEngine.Build.cs
--- a/Engine/Plugins/Experimental/ChaosFlesh/Source/ChaosFleshEngine/ChaosFleshEngine.Build.cs
+++ b/Engine/Plugins/Experimental/ChaosFlesh/Source/ChaosFleshEngine/ChaosFleshEngine.Build.cs
@@ -87,6 +87,7 @@ namespace UnrealBuildTool.Rules
 			{
 				bUseRTTI = true;
 				PublicDefinitions.Add("DO_USD_CACHING=1");
+				PrivateDefinitions.Add("SUPPRESS_PER_MODULE_INLINE_FILE"); // This module does not use core's standard operator new/delete overloads
 			}
 			else
 			{
diff --git a/Engine/Plugins/Experimental/ChaosFlesh/Source/ChaosFleshEngine/Private/ChaosCache/FleshComponentCacheAdapter.cpp b/Engine/Plugins/Experimental/ChaosFlesh/Source/ChaosFleshEngine/Private/ChaosCache/FleshComponentCacheAdapter.cpp
--- a/Engine/Plugins/Experimental/ChaosFlesh/Source/ChaosFleshEngine/Private/ChaosCache/FleshComponentCacheAdapter.cpp
+++ b/Engine/Plugins/Experimental/ChaosFlesh/Source/ChaosFleshEngine/Private/ChaosCache/FleshComponentCacheAdapter.cpp
@@ -16,6 +16,7 @@
 
 #include "ChaosCachingUSD/Operations.h"
 #include "USDConversionUtils.h" // for GetPrimPathForObject()
+#include "USDMemory.h"
 
 #endif // USE_USD_SDK && DO_USD_CACHING
 
diff --git a/Engine/Plugins/Experimental/ChaosFlesh/Source/ChaosFleshEngine/Private/ChaosFlesh/ChaosFleshEnginePlugin.cpp b/Engine/Plugins/Experimental/ChaosFlesh/Source/ChaosFleshEngine/Private/ChaosFlesh/ChaosFleshEnginePlugin.cpp
--- a/Engine/Plugins/Experimental/ChaosFlesh/Source/ChaosFleshEngine/Private/ChaosFlesh/ChaosFleshEnginePlugin.cpp
+++ b/Engine/Plugins/Experimental/ChaosFlesh/Source/ChaosFleshEngine/Private/ChaosFlesh/ChaosFleshEnginePlugin.cpp
@@ -8,6 +8,9 @@
 #include "IOptimusCoreModule.h"
 #include "ShaderCore.h"
 
+#if USE_USD_SDK && DO_USD_CACHING
+#include "USDMemory.h"
+#endif
 
 class FChaosFleshEnginePlugin : public IChaosFleshEnginePlugin
 {
@@ -46,5 +49,8 @@ void FChaosFleshEnginePlugin::ShutdownModule()
 }
 
 
+#if USE_USD_SDK && DO_USD_CACHING
+IMPLEMENT_MODULE_USD(FChaosFleshEnginePlugin, ChaosFleshEngine)
+#else
 IMPLEMENT_MODULE(FChaosFleshEnginePlugin, ChaosFleshEngine)
-
+#endif
diff --git a/Engine/Plugins/Importers/USDImporter/Source/USDSchemas/Private/Custom/MDLUSDShadeMaterialTranslator.cpp b/Engine/Plugins/Importers/USDImporter/Source/USDSchemas/Private/Custom/MDLUSDShadeMaterialTranslator.cpp
--- a/Engine/Plugins/Importers/USDImporter/Source/USDSchemas/Private/Custom/MDLUSDShadeMaterialTranslator.cpp
+++ b/Engine/Plugins/Importers/USDImporter/Source/USDSchemas/Private/Custom/MDLUSDShadeMaterialTranslator.cpp
@@ -10,6 +10,7 @@
 #include "USDConversionUtils.h"
 #include "USDErrorUtils.h"
 #include "USDLog.h"
+#include "USDMemory.h"
 #include "USDShadeConversion.h"
 #include "USDTypesConversion.h"
 
@@ -192,11 +193,7 @@ void FMdlUsdShadeMaterialTranslator::CreateAssets()
 
 				// Rename the UMaterialInterface into the target UPackage the asset cache created for us.
 				// SanitizedName will already match it.
-				const bool bRenamed = ReferenceMaterial->Rename(
-					*SanitizedName.ToString(),
-					Outer,
-					REN_NonTransactional | REN_DontCreateRedirectors
-				);
+				const bool bRenamed = ReferenceMaterial->Rename(*SanitizedName.ToString(), Outer, REN_NonTransactional | REN_DontCreateRedirectors);
 				ensure(bRenamed);
 
 				// Let's not trust the flags the MDLImporter used and just use our own instead
@@ -306,11 +303,7 @@ void FMdlUsdShadeMaterialTranslator::CreateAssets()
 					{
 						// Rename the asset into the target UPackage the asset cache created for us.
 						// SanitizedName will already match it.
-						const bool bRenamed = Texture->Rename(
-							*SanitizedName.ToString(),
-							Outer,
-							REN_NonTransactional | REN_DontCreateRedirectors
-						);
+						const bool bRenamed = Texture->Rename(*SanitizedName.ToString(), Outer, REN_NonTransactional | REN_DontCreateRedirectors);
 						ensure(bRenamed);
 
 						// Let's not trust the flags the MDLImporter used and just use our own instead
diff --git a/Engine/Plugins/Importers/USDImporter/Source/USDSchemas/Private/Custom/MaterialXUSDShadeMaterialTranslator.cpp b/Engine/Plugins/Importers/USDImporter/Source/USDSchemas/Private/Custom/MaterialXUSDShadeMaterialTranslator.cpp
--- a/Engine/Plugins/Importers/USDImporter/Source/USDSchemas/Private/Custom/MaterialXUSDShadeMaterialTranslator.cpp
+++ b/Engine/Plugins/Importers/USDImporter/Source/USDSchemas/Private/Custom/MaterialXUSDShadeMaterialTranslator.cpp
@@ -8,6 +8,7 @@
 #include "USDConversionUtils.h"
 #include "USDErrorUtils.h"
 #include "USDLog.h"
+#include "USDMemory.h"
 #include "USDObjectUtils.h"
 #include "USDShadeConversion.h"
 #include "USDTypesConversion.h"
diff --git a/Engine/Plugins/Importers/USDImporter/Source/USDSchemas/Private/USDGeomMeshTranslator.cpp b/Engine/Plugins/Importers/USDImporter/Source/USDSchemas/Private/USDGeomMeshTranslator.cpp
--- a/Engine/Plugins/Importers/USDImporter/Source/USDSchemas/Private/USDGeomMeshTranslator.cpp
+++ b/Engine/Plugins/Importers/USDImporter/Source/USDSchemas/Private/USDGeomMeshTranslator.cpp
@@ -14,6 +14,7 @@
 #include "USDGeomMeshConversion.h"
 #include "USDInfoCache.h"
 #include "USDLog.h"
+#include "USDMemory.h"
 #include "USDObjectUtils.h"
 #include "USDPrimConversion.h"
 #include "USDTranslatorUtils.h"
diff --git a/Engine/Plugins/Importers/USDImporter/Source/USDSchemas/Private/USDGeomPointInstancerTranslator.cpp b/Engine/Plugins/Importers/USDImporter/Source/USDSchemas/Private/USDGeomPointInstancerTranslator.cpp
--- a/Engine/Plugins/Importers/USDImporter/Source/USDSchemas/Private/USDGeomPointInstancerTranslator.cpp
+++ b/Engine/Plugins/Importers/USDImporter/Source/USDSchemas/Private/USDGeomPointInstancerTranslator.cpp
@@ -10,6 +10,7 @@
 #include "USDGeomMeshConversion.h"
 #include "USDGeomMeshTranslator.h"
 #include "USDLog.h"
+#include "USDMemory.h"
 #include "USDPrimConversion.h"
 #include "USDSchemasModule.h"
 #include "USDTypesConversion.h"
diff --git a/Engine/Plugins/Importers/USDImporter/Source/USDSchemas/Private/USDGeomPrimitiveTranslator.cpp b/Engine/Plugins/Importers/USDImporter/Source/USDSchemas/Private/USDGeomPrimitiveTranslator.cpp
--- a/Engine/Plugins/Importers/USDImporter/Source/USDSchemas/Private/USDGeomPrimitiveTranslator.cpp
+++ b/Engine/Plugins/Importers/USDImporter/Source/USDSchemas/Private/USDGeomPrimitiveTranslator.cpp
@@ -10,6 +10,7 @@
 #include "USDGeomMeshConversion.h"
 #include "USDGeomMeshTranslator.h"
 #include "USDInfoCache.h"
+#include "USDMemory.h"
 #include "USDTypesConversion.h"
 
 #include "UsdWrappers/SdfPath.h"
diff --git a/Engine/Plugins/Importers/USDImporter/Source/USDSchemas/Private/USDGeometryCacheTranslator.cpp b/Engine/Plugins/Importers/USDImporter/Source/USDSchemas/Private/USDGeometryCacheTranslator.cpp
--- a/Engine/Plugins/Importers/USDImporter/Source/USDSchemas/Private/USDGeometryCacheTranslator.cpp
+++ b/Engine/Plugins/Importers/USDImporter/Source/USDSchemas/Private/USDGeometryCacheTranslator.cpp
@@ -14,6 +14,7 @@
 #include "USDInfoCache.h"
 #include "USDIntegrationUtils.h"
 #include "USDLog.h"
+#include "USDMemory.h"
 #include "USDObjectUtils.h"
 #include "USDPrimConversion.h"
 #include "USDTypesConversion.h"
@@ -448,12 +449,8 @@ namespace UsdGeometryCacheTranslatorImpl
 			ObjectFlags |= RF_Transient;
 		}
 
-		UGeometryCache* GeometryCache = Context->UsdAssetCache->GetOrCreateCachedAsset<UGeometryCache>(
-			PrefixedGeoCacheHash,
-			DesiredName,
-			ObjectFlags,
-			&bOutIsNew
-		);
+		UGeometryCache* GeometryCache = Context->UsdAssetCache
+											->GetOrCreateCachedAsset<UGeometryCache>(PrefixedGeoCacheHash, DesiredName, ObjectFlags, &bOutIsNew);
 
 		if (GeometryCache && bOutIsNew)
 		{
@@ -502,7 +499,8 @@ namespace UsdGeometryCacheTranslatorImpl
 						FillGeometryCacheTracks(RootPrimPath, MeshPaths, MaterialOffsets, Context, GeometryCache);
 
 						FinalizeGeometryCache(GeometryCache);
-					});
+					}
+				);
 			}
 			if (!bIsImporting)
 			{
diff --git a/Engine/Plugins/Importers/USDImporter/Source/USDSchemas/Private/USDGroomTranslator.cpp b/Engine/Plugins/Importers/USDImporter/Source/USDSchemas/Private/USDGroomTranslator.cpp
--- a/Engine/Plugins/Importers/USDImporter/Source/USDSchemas/Private/USDGroomTranslator.cpp
+++ b/Engine/Plugins/Importers/USDImporter/Source/USDSchemas/Private/USDGroomTranslator.cpp
@@ -25,6 +25,7 @@
 #include "USDGroomTranslatorUtils.h"
 #include "USDIntegrationUtils.h"
 #include "USDLog.h"
+#include "USDMemory.h"
 #include "USDObjectUtils.h"
 #include "USDPrimConversion.h"
 #include "USDTranslatorUtils.h"
diff --git a/Engine/Plugins/Importers/USDImporter/Source/USDSchemas/Private/USDGroomTranslatorUtils.cpp b/Engine/Plugins/Importers/USDImporter/Source/USDSchemas/Private/USDGroomTranslatorUtils.cpp
--- a/Engine/Plugins/Importers/USDImporter/Source/USDSchemas/Private/USDGroomTranslatorUtils.cpp
+++ b/Engine/Plugins/Importers/USDImporter/Source/USDSchemas/Private/USDGroomTranslatorUtils.cpp
@@ -9,6 +9,7 @@
 #include "USDClassesModule.h"
 #include "USDInfoCache.h"
 #include "USDIntegrationUtils.h"
+#include "USDMemory.h"
 #include "USDPrimLinkCache.h"
 #include "USDTypesConversion.h"
 #include "UsdWrappers/SdfPath.h"
diff --git a/Engine/Plugins/Importers/USDImporter/Source/USDSchemas/Private/USDInfoCache.cpp b/Engine/Plugins/Importers/USDImporter/Source/USDSchemas/Private/USDInfoCache.cpp
--- a/Engine/Plugins/Importers/USDImporter/Source/USDSchemas/Private/USDInfoCache.cpp
+++ b/Engine/Plugins/Importers/USDImporter/Source/USDSchemas/Private/USDInfoCache.cpp
@@ -4,6 +4,7 @@
 
 #include "USDGeomMeshConversion.h"
 #include "USDLog.h"
+#include "USDMemory.h"
 #include "USDSchemasModule.h"
 #include "USDSchemaTranslator.h"
 #include "USDTypesConversion.h"
diff --git a/Engine/Plugins/Importers/USDImporter/Source/USDSchemas/Private/USDLuxLightTranslator.cpp b/Engine/Plugins/Importers/USDImporter/Source/USDSchemas/Private/USDLuxLightTranslator.cpp
--- a/Engine/Plugins/Importers/USDImporter/Source/USDSchemas/Private/USDLuxLightTranslator.cpp
+++ b/Engine/Plugins/Importers/USDImporter/Source/USDSchemas/Private/USDLuxLightTranslator.cpp
@@ -7,6 +7,7 @@
 #include "USDDrawModeComponent.h"
 #include "USDLightConversion.h"
 #include "USDLog.h"
+#include "USDMemory.h"
 #include "USDObjectUtils.h"
 #include "USDShadeConversion.h"
 #include "USDTypesConversion.h"
diff --git a/Engine/Plugins/Importers/USDImporter/Source/USDSchemas/Private/USDMediaSpatialAudioTranslator.cpp b/Engine/Plugins/Importers/USDImporter/Source/USDSchemas/Private/USDMediaSpatialAudioTranslator.cpp
--- a/Engine/Plugins/Importers/USDImporter/Source/USDSchemas/Private/USDMediaSpatialAudioTranslator.cpp
+++ b/Engine/Plugins/Importers/USDImporter/Source/USDSchemas/Private/USDMediaSpatialAudioTranslator.cpp
@@ -6,6 +6,7 @@
 #include "USDConversionUtils.h"
 #include "USDDrawModeComponent.h"
 #include "USDLog.h"
+#include "USDMemory.h"
 #include "USDObjectUtils.h"
 #include "USDPrimConversion.h"
 #include "USDShadeConversion.h"
diff --git a/Engine/Plugins/Importers/USDImporter/Source/USDSchemas/Private/USDSchemaTranslator.cpp b/Engine/Plugins/Importers/USDImporter/Source/USDSchemas/Private/USDSchemaTranslator.cpp
--- a/Engine/Plugins/Importers/USDImporter/Source/USDSchemas/Private/USDSchemaTranslator.cpp
+++ b/Engine/Plugins/Importers/USDImporter/Source/USDSchemas/Private/USDSchemaTranslator.cpp
@@ -3,9 +3,9 @@
 #include "USDSchemaTranslator.h"
 
 #include "USDInfoCache.h"
+#include "USDMemory.h"
 #include "USDSchemasModule.h"
 #include "USDTypesConversion.h"
-
 #include "UsdWrappers/UsdPrim.h"
 #include "UsdWrappers/UsdTyped.h"
 
@@ -246,9 +246,9 @@ bool FUsdSchemaTranslator::IsCollapsed(ECollapsingType CollapsingType) const
 
 	// This is merely a fallback, and we should never need this
 	return CanBeCollapsed(CollapsingType);
-#else	// #if USE_USD_SDK
+#else	  // #if USE_USD_SDK
 	return false;
-#endif	// #if USE_USD_SDK
+#endif	  // #if USE_USD_SDK
 }
 
 UE::FSdfPath FUsdSchemaTranslator::GetPrototypePrimPath() const
diff --git a/Engine/Plugins/Importers/USDImporter/Source/USDSchemas/Private/USDShadeMaterialTranslator.cpp b/Engine/Plugins/Importers/USDImporter/Source/USDSchemas/Private/USDShadeMaterialTranslator.cpp
--- a/Engine/Plugins/Importers/USDImporter/Source/USDSchemas/Private/USDShadeMaterialTranslator.cpp
+++ b/Engine/Plugins/Importers/USDImporter/Source/USDSchemas/Private/USDShadeMaterialTranslator.cpp
@@ -6,6 +6,7 @@
 #include "USDAssetUserData.h"
 #include "USDClassesModule.h"
 #include "USDLog.h"
+#include "USDMemory.h"
 #include "USDObjectUtils.h"
 #include "USDPrimConversion.h"
 #include "USDProjectSettings.h"
diff --git a/Engine/Plugins/Importers/USDImporter/Source/USDSchemas/Private/USDVolVolumeTranslator.cpp b/Engine/Plugins/Importers/USDImporter/Source/USDSchemas/Private/USDVolVolumeTranslator.cpp
--- a/Engine/Plugins/Importers/USDImporter/Source/USDSchemas/Private/USDVolVolumeTranslator.cpp
+++ b/Engine/Plugins/Importers/USDImporter/Source/USDSchemas/Private/USDVolVolumeTranslator.cpp
@@ -9,6 +9,7 @@
 #include "USDIntegrationUtils.h"
 #include "USDLayerUtils.h"
 #include "USDLog.h"
+#include "USDMemory.h"
 #include "USDObjectUtils.h"
 #include "USDPrimConversion.h"
 #include "USDProjectSettings.h"
diff --git a/Engine/Plugins/Importers/USDImporter/Source/USDSchemas/Public/USDSchemaTranslator.h b/Engine/Plugins/Importers/USDImporter/Source/USDSchemas/Public/USDSchemaTranslator.h
--- a/Engine/Plugins/Importers/USDImporter/Source/USDSchemas/Public/USDSchemaTranslator.h
+++ b/Engine/Plugins/Importers/USDImporter/Source/USDSchemas/Public/USDSchemaTranslator.h
@@ -8,7 +8,9 @@
 #include "USDAssetCache2.h"
 #include "USDAssetCache3.h"
 #include "USDInfoCache.h"
+#if UE_ENABLE_INCLUDE_ORDER_DEPRECATED_IN_5_5
 #include "USDMemory.h"
+#endif	  // UE_ENABLE_INCLUDE_ORDER_DEPRECATED_IN_5_5
 #include "USDMetadataImportOptions.h"
 #include "USDPrimLinkCache.h"
 #include "USDSkeletalDataConversion.h"
diff --git a/Engine/Plugins/Importers/USDImporter/Source/USDStage/Private/USDStageActor.cpp b/Engine/Plugins/Importers/USDImporter/Source/USDStage/Private/USDStageActor.cpp
--- a/Engine/Plugins/Importers/USDImporter/Source/USDStage/Private/USDStageActor.cpp
+++ b/Engine/Plugins/Importers/USDImporter/Source/USDStage/Private/USDStageActor.cpp
@@ -20,6 +20,7 @@
 #include "USDLightConversion.h"
 #include "USDListener.h"
 #include "USDLog.h"
+#include "USDMemory.h"
 #include "USDObjectUtils.h"
 #include "USDPrimConversion.h"
 #include "USDPrimLinkCache.h"
diff --git a/Engine/Plugins/Importers/USDImporter/Source/USDStage/Private/USDTransactor.cpp b/Engine/Plugins/Importers/USDImporter/Source/USDStage/Private/USDTransactor.cpp
--- a/Engine/Plugins/Importers/USDImporter/Source/USDStage/Private/USDTransactor.cpp
+++ b/Engine/Plugins/Importers/USDImporter/Source/USDStage/Private/USDTransactor.cpp
@@ -6,11 +6,11 @@
 #include "USDLayerUtils.h"
 #include "USDListener.h"
 #include "USDLog.h"
+#include "USDMemory.h"
 #include "USDPrimConversion.h"
 #include "USDStageActor.h"
 #include "USDTypesConversion.h"
 #include "USDValueConversion.h"
-
 #include "UsdWrappers/SdfChangeBlock.h"
 #include "UsdWrappers/SdfLayer.h"
 #include "UsdWrappers/SdfPath.h"
diff --git a/Engine/Plugins/Importers/USDImporter/Source/USDStageEditor/Private/SUSDStageTreeView.cpp b/Engine/Plugins/Importers/USDImporter/Source/USDStageEditor/Private/SUSDStageTreeView.cpp
--- a/Engine/Plugins/Importers/USDImporter/Source/USDStageEditor/Private/SUSDStageTreeView.cpp
+++ b/Engine/Plugins/Importers/USDImporter/Source/USDStageEditor/Private/SUSDStageTreeView.cpp
@@ -8,11 +8,11 @@
 #include "USDConversionUtils.h"
 #include "USDDuplicateType.h"
 #include "USDLayerUtils.h"
+#include "USDMemory.h"
 #include "USDOptionsWindow.h"
 #include "USDPrimViewModel.h"
 #include "USDReferenceOptions.h"
 #include "USDTypesConversion.h"
-
 #include "UsdWrappers/SdfChangeBlock.h"
 #include "UsdWrappers/SdfPath.h"
 #include "UsdWrappers/UsdPrim.h"
@@ -1706,7 +1706,10 @@ void SUsdStageTreeView::FillCollapsingSubmenu(FMenuBuilder& MenuBuilder)
 
 	MenuBuilder.AddMenuEntry(
 		LOCTEXT("CollapseOnKind", "Default"),
-		LOCTEXT("CollapseOnKind_ToolTip", "When 'Use prim kinds for collapsing' is enabled, prims are collapsed according to their kind (default). When disabled, prims won't be collapsed by default."),
+		LOCTEXT(
+			"CollapseOnKind_ToolTip",
+			"When 'Use prim kinds for collapsing' is enabled, prims are collapsed according to their kind (default). When disabled, prims won't be collapsed by default."
+		),
 		FSlateIcon(),
 		FUIAction(
 			FExecuteAction::CreateSP(this, &SUsdStageTreeView::OnSetCollapsingPreference, UsdUtils::ECollapsingPreference::Default),
diff --git a/Engine/Plugins/Importers/USDImporter/Source/USDStageEditor/Private/Widgets/SUSDIntegrationsPanel.cpp b/Engine/Plugins/Importers/USDImporter/Source/USDStageEditor/Private/Widgets/SUSDIntegrationsPanel.cpp
--- a/Engine/Plugins/Importers/USDImporter/Source/USDStageEditor/Private/Widgets/SUSDIntegrationsPanel.cpp
+++ b/Engine/Plugins/Importers/USDImporter/Source/USDStageEditor/Private/Widgets/SUSDIntegrationsPanel.cpp
@@ -7,6 +7,7 @@
 #include "SUSDStageEditorStyle.h"
 #include "USDAttributeUtils.h"
 #include "USDIntegrationUtils.h"
+#include "USDMemory.h"
 #include "USDTypesConversion.h"
 #include "USDValueConversion.h"
 
diff --git a/Engine/Plugins/Importers/USDImporter/Source/USDStageEditorViewModels/Private/USDStageViewModel.cpp b/Engine/Plugins/Importers/USDImporter/Source/USDStageEditorViewModels/Private/USDStageViewModel.cpp
--- a/Engine/Plugins/Importers/USDImporter/Source/USDStageEditorViewModels/Private/USDStageViewModel.cpp
+++ b/Engine/Plugins/Importers/USDImporter/Source/USDStageEditorViewModels/Private/USDStageViewModel.cpp
@@ -9,6 +9,7 @@
 #include "USDErrorUtils.h"
 #include "USDLayerUtils.h"
 #include "USDLog.h"
+#include "USDMemory.h"
 #include "USDStageActor.h"
 #include "USDStageImportContext.h"
 #include "USDStageImporter.h"
@@ -16,7 +17,6 @@
 #include "USDStageImportOptions.h"
 #include "USDStageModule.h"
 #include "USDTypesConversion.h"
-
 #include "UsdWrappers/SdfLayer.h"
 #include "UsdWrappers/UsdStage.h"
 
diff --git a/Engine/Plugins/Interchange/Runtime/Source/Import/Private/Usd/InterchangeUsdTranslator.cpp b/Engine/Plugins/Interchange/Runtime/Source/Import/Private/Usd/InterchangeUsdTranslator.cpp
--- a/Engine/Plugins/Interchange/Runtime/Source/Import/Private/Usd/InterchangeUsdTranslator.cpp
+++ b/Engine/Plugins/Interchange/Runtime/Source/Import/Private/Usd/InterchangeUsdTranslator.cpp
@@ -31,9 +31,9 @@
 #include "Async/Async.h"
 #include "Async/ParallelFor.h"
 #include "HAL/IConsoleManager.h"
-#include "InterchangeManager.h"
 #include "InterchangeCameraNode.h"
 #include "InterchangeLightNode.h"
+#include "InterchangeManager.h"
 #include "InterchangeMaterialInstanceNode.h"
 #include "InterchangeMeshNode.h"
 #include "InterchangeSceneNode.h"
@@ -130,8 +130,13 @@ namespace UE::InterchangeUsdTranslator::Private
 	{
 	public:
 
-		/** Add a material instance to the node container, otherwise it will add a material if it comes from a Translator (for example coming from MaterialX which cannot handle material instances) */
-		void AddMaterialNode(const UE::FUsdPrim& Prim, UInterchangeUsdTranslatorSettings* TranslatorSettings, UInterchangeBaseNodeContainer& NodeContainer);
+		/** Add a material instance to the node container, otherwise it will add a material if it comes from a Translator (for example coming from
+		 * MaterialX which cannot handle material instances) */
+		void AddMaterialNode(
+			const UE::FUsdPrim& Prim,
+			UInterchangeUsdTranslatorSettings* TranslatorSettings,
+			UInterchangeBaseNodeContainer& NodeContainer
+		);
 
 		void AddMeshNode(const UE::FUsdPrim& Prim, UInterchangeBaseNodeContainer& NodeContainer, const FTraversalInfo& Info);
 
@@ -158,7 +163,8 @@ namespace UE::InterchangeUsdTranslator::Private
 		// here for easy access when parsing the tracks
 		UInterchangeAnimationTrackSetNode* CurrentTrackSet = nullptr;
 
-		// Array of translators that we call in the GetTexturePayload, the key has no real meaning, it's just here to avoid having duplicates and calling several times the Translate function
+		// Array of translators that we call in the GetTexturePayload, the key has no real meaning, it's just here to avoid having duplicates and
+		// calling several times the Translate function
 		TMap<FString, UInterchangeTranslatorBase*> Translators;
 
 	private:
@@ -168,7 +174,7 @@ namespace UE::InterchangeUsdTranslator::Private
 			FString MaterialSlotName;
 			UInterchangeMeshNode* MeshNode;
 		};
-		TMap<FString, FString> PrimPathToMaterialPath; 
+		TMap<FString, FString> PrimPathToMaterialPath;
 		TMap<FString, TArray<FMaterialSlotMesh>> PrimPathToSlotMeshNodes;
 	};
 
@@ -715,13 +721,13 @@ namespace UE::InterchangeUsdTranslator::Private
 		auto SetMaterialSlotDependencies = [this, &NodeUid]()
 		{
 			// Now we need to check if we have to set the slot of the mesh nodes here
-			if(TArray<FMaterialSlotMesh>* SlotMeshes = PrimPathToSlotMeshNodes.Find(NodeUid))
+			if (TArray<FMaterialSlotMesh>* SlotMeshes = PrimPathToSlotMeshNodes.Find(NodeUid))
 			{
-				if(FString* NewMaterialUID = PrimPathToMaterialPath.Find(NodeUid))
+				if (FString* NewMaterialUID = PrimPathToMaterialPath.Find(NodeUid))
 				{
-					for(const FMaterialSlotMesh& MaterialSlotMesh : *SlotMeshes)
+					for (const FMaterialSlotMesh& MaterialSlotMesh : *SlotMeshes)
 					{
-						if(!MaterialSlotMesh.MeshNode->GetSlotMaterialDependencyUid(MaterialSlotMesh.MaterialSlotName, *NewMaterialUID))
+						if (!MaterialSlotMesh.MeshNode->GetSlotMaterialDependencyUid(MaterialSlotMesh.MaterialSlotName, *NewMaterialUID))
 						{
 							MaterialSlotMesh.MeshNode->SetSlotMaterialDependencyUid(MaterialSlotMesh.MaterialSlotName, *NewMaterialUID);
 						}
@@ -737,18 +743,17 @@ namespace UE::InterchangeUsdTranslator::Private
 		if(RenderContext == UnrealIdentifiers::MaterialXRenderContext)
 		{
 			TArray<FString> FilePaths = UsdUtils::GetMaterialXFilePaths(Prim);
-			for(const FString& File : FilePaths)
+			for (const FString& File : FilePaths)
 			{
 				// the file has already been handled no need to do a Translate again
-				if(!Translators.Find(File))
+				if (!Translators.Find(File))
 				{
-
 					UInterchangeManager& InterchangeManager = UInterchangeManager::GetInterchangeManager();
 					UInterchangeSourceData* SourceData = UInterchangeManager::CreateSourceData(File);
 
 					UInterchangeTranslatorBase* Translator = InterchangeManager.GetTranslatorForSourceData(SourceData);
-					//check on the Translator, it might return nullptr in case of reimport
-					if(Translator)
+					// check on the Translator, it might return nullptr in case of reimport
+					if (Translator)
 					{
 						Translator->Translate(NodeContainer);
 						Translators.Add(File, Translator);
@@ -761,7 +766,7 @@ namespace UE::InterchangeUsdTranslator::Private
 					[this, &NodeName, &NodeUid](const FString&, UInterchangeShaderGraphNode* ShaderGraphNode)
 					{
 						FString UID = ShaderGraphNode->GetUniqueID();
-						if(FPaths::GetBaseFilename(UID) == NodeName)
+						if (FPaths::GetBaseFilename(UID) == NodeName)
 						{
 							PrimPathToMaterialPath.Add(NodeUid, UID);
 							return true;
@@ -770,12 +775,13 @@ namespace UE::InterchangeUsdTranslator::Private
 						{
 							return false;
 						}
-					});
+					}
+				);
 			}
 
 			SetMaterialSlotDependencies();
 
-			if(!FilePaths.IsEmpty())
+			if (!FilePaths.IsEmpty())
 			{
 				return;
 			}
@@ -786,7 +792,7 @@ namespace UE::InterchangeUsdTranslator::Private
 		MaterialNode->InitializeNode(NodeUid, NodeName, EInterchangeNodeContainerType::TranslatedAsset);
 		MaterialNode->SetAssetName(NodeName);
 		NodeContainer.AddNode(MaterialNode);
-		
+
 		// Set the material instance node to the correct mesh nodes
 		PrimPathToMaterialPath.Add(NodeUid, NodeUid);
 		SetMaterialSlotDependencies();
@@ -794,7 +800,6 @@ namespace UE::InterchangeUsdTranslator::Private
 		UsdToUnreal::FUsdPreviewSurfaceMaterialData MaterialData;
 		const bool bSuccess = UsdToUnreal::ConvertMaterial(Prim, MaterialData, TranslatorSettings ? *RenderContext : nullptr);
 
-
 		// Set all the parameter values to the interchange node
 		bool bHasUDIMTexture = false;
 		FParameterValueVisitor Visitor{Prim, NodeContainer, *MaterialNode, MaterialData.PrimvarToUVIndex};
@@ -1200,20 +1205,25 @@ namespace UE::InterchangeUsdTranslator::Private
 				}
 
 				// if we found a match let's set the slot to the corresponding Material (the Material path could come from another translator)
-				if(FString* NewMaterialInstanceUid = PrimPathToMaterialPath.Find(MaterialInstanceUid))
+				if (FString* NewMaterialInstanceUid = PrimPathToMaterialPath.Find(MaterialInstanceUid))
 				{
 					MeshNode->SetSlotMaterialDependencyUid(SlotName, *NewMaterialInstanceUid);
 				}
-				else // otherwise it's up to the Material to attach itself to the mesh
+				else	// otherwise it's up to the Material to attach itself to the mesh
 				{
 					// one material can be attached to several meshes
-					if(TArray<FMaterialSlotMesh>* SlotsMeshes = PrimPathToSlotMeshNodes.Find(MaterialInstanceUid))
+					if (TArray<FMaterialSlotMesh>* SlotsMeshes = PrimPathToSlotMeshNodes.Find(MaterialInstanceUid))
 					{
-						SlotsMeshes->Add({ SlotName, MeshNode });
+						SlotsMeshes->Add({SlotName, MeshNode});
 					}
 					else
 					{
-						PrimPathToSlotMeshNodes.Add(MaterialInstanceUid, { { SlotName, MeshNode } });
+						PrimPathToSlotMeshNodes.Add(
+							MaterialInstanceUid,
+							{
+								{SlotName, MeshNode}
+						}
+						);
 					}
 				}
 			}
@@ -2384,10 +2394,16 @@ bool UInterchangeUSDTranslator::Translate(UInterchangeBaseNodeContainer& NodeCon
 	// Cache these so we don't have to keep converting these tokens over and over during translation
 	FUsdMeshConversionOptions& MeshOptions = ImplPtr->CachedMeshConversionOptions;
 	MeshOptions.PurposesToLoad = (EUsdPurpose)Settings->GeometryPurpose;
+
+	// TODO: Change FUsdMeshConversionOptions to not hold USD types directly, so we don't have to the conversion below everywhere.
+	// We can't use UsdToUnreal::ConvertToken() here because it returns a TUsdStore, and the template instantiation created in this module doesn't
+	// really do anything anyway as the module doesn't use IMPLEMENT_MODULE_USD!
+	// Luckily we can get around this here because pxr::TfToken doesn't allocate on its own: At most USD makes a copy of the string, which it should
+	// allocate/deallocate on its own allocator.
 	MeshOptions.RenderContext = Settings->RenderContext.IsNone() ? pxr::UsdShadeTokens->universalRenderContext
-																 : UnrealToUsd::ConvertToken(*Settings->RenderContext.ToString()).Get();
+																 : pxr::TfToken{TCHAR_TO_ANSI(*Settings->RenderContext.ToString())};
 	MeshOptions.MaterialPurpose = Settings->MaterialPurpose.IsNone() ? pxr::UsdShadeTokens->allPurpose
-																	 : UnrealToUsd::ConvertToken(*Settings->MaterialPurpose.ToString()).Get();
+																	 : pxr::TfToken{TCHAR_TO_ANSI(*Settings->MaterialPurpose.ToString())};
 
 	// Traverse stage and emit translated nodes
 	FTraversalInfo Info;
@@ -2547,14 +2563,13 @@ TOptional<UE::Interchange::FImportImage> UInterchangeUSDTranslator::GetTexturePa
 	FString FilePath;
 	TextureGroup TextureGroup;
 	bool bDecoded = DecodeTexturePayloadKey(PayloadKey, FilePath, TextureGroup);
-	if(bDecoded)
+	if (bDecoded)
 	{
 		// Defer back to another translator to actually parse the texture raw data
 		UE::Interchange::Private::FScopedTranslator ScopedTranslator(FilePath, Results);
 		const IInterchangeTexturePayloadInterface* TextureTranslator = ScopedTranslator.GetPayLoadInterface<IInterchangeTexturePayloadInterface>();
-		if(ensure(TextureTranslator))
+		if (ensure(TextureTranslator))
 		{
-
 			AlternateTexturePath = FilePath;
 
 			// The texture translators don't use the payload key, and read the texture directly from the SourceData's file path
@@ -2572,14 +2587,14 @@ TOptional<UE::Interchange::FImportImage> UInterchangeUSDTranslator::GetTexturePa
 #endif	  // USE_USD_SDK
 
 	// We did not find a suitable Payload in USD Translator, let's find one in one of the Translators (MaterialX for the moment)
-	// The best way would be to have a direct association between the payload and the right Translator, but we don't have a suitable way of knowing which Payload belongs to which Translator
-	// So let's just loop over them all
-	for(const TPair<FString, UInterchangeTranslatorBase*> & Pair : Impl->Translators)
+	// The best way would be to have a direct association between the payload and the right Translator, but we don't have a suitable way of knowing
+	// which Payload belongs to which Translator So let's just loop over them all
+	for (const TPair<FString, UInterchangeTranslatorBase*>& Pair : Impl->Translators)
 	{
-		if(IInterchangeTexturePayloadInterface* TexturePayloadInterface = Cast<IInterchangeTexturePayloadInterface>(Pair.Value))
+		if (IInterchangeTexturePayloadInterface* TexturePayloadInterface = Cast<IInterchangeTexturePayloadInterface>(Pair.Value))
 		{
 			TexturePayloadData = TexturePayloadInterface->GetTexturePayloadData(PayloadKey, AlternateTexturePath);
-			if(TexturePayloadData)
+			if (TexturePayloadData)
 			{
 				break;
 			}
diff --git a/Engine/Plugins/Runtime/USDCore/Source/USDUtilities/Private/USDAttributeUtils.cpp b/Engine/Plugins/Runtime/USDCore/Source/USDUtilities/Private/USDAttributeUtils.cpp
--- a/Engine/Plugins/Runtime/USDCore/Source/USDUtilities/Private/USDAttributeUtils.cpp
+++ b/Engine/Plugins/Runtime/USDCore/Source/USDUtilities/Private/USDAttributeUtils.cpp
@@ -5,6 +5,7 @@
 #include "UnrealUSDWrapper.h"
 #include "USDLayerUtils.h"
 #include "USDLog.h"
+#include "USDMemory.h"
 #include "USDProjectSettings.h"
 #include "USDTypesConversion.h"
 
diff --git a/Engine/Plugins/Runtime/USDCore/Source/USDUtilities/Private/USDConversionUtils.cpp b/Engine/Plugins/Runtime/USDCore/Source/USDUtilities/Private/USDConversionUtils.cpp
--- a/Engine/Plugins/Runtime/USDCore/Source/USDUtilities/Private/USDConversionUtils.cpp
+++ b/Engine/Plugins/Runtime/USDCore/Source/USDUtilities/Private/USDConversionUtils.cpp
@@ -12,12 +12,12 @@
 #include "USDIntegrationUtils.h"
 #include "USDLayerUtils.h"
 #include "USDLog.h"
+#include "USDMemory.h"
 #include "USDObjectUtils.h"
 #include "USDProjectSettings.h"
 #include "USDSkeletalDataConversion.h"
 #include "USDTypesConversion.h"
 #include "USDUnrealAssetInfo.h"
-
 #include "UsdWrappers/SdfPath.h"
 #include "UsdWrappers/UsdPrim.h"
 #include "UsdWrappers/UsdStage.h"
diff --git a/Engine/Plugins/Runtime/USDCore/Source/USDUtilities/Private/USDGroomConversion.cpp b/Engine/Plugins/Runtime/USDCore/Source/USDUtilities/Private/USDGroomConversion.cpp
--- a/Engine/Plugins/Runtime/USDCore/Source/USDUtilities/Private/USDGroomConversion.cpp
+++ b/Engine/Plugins/Runtime/USDCore/Source/USDUtilities/Private/USDGroomConversion.cpp
@@ -6,6 +6,7 @@
 
 #include "GroomCacheData.h"
 #include "HairDescription.h"
+#include "USDMemory.h"
 #include "USDPrimConversion.h"
 #include "USDTypesConversion.h"
 #include "UsdWrappers/UsdPrim.h"
diff --git a/Engine/Plugins/Runtime/USDCore/Source/USDUtilities/Private/USDIntegrationUtils.cpp b/Engine/Plugins/Runtime/USDCore/Source/USDUtilities/Private/USDIntegrationUtils.cpp
--- a/Engine/Plugins/Runtime/USDCore/Source/USDUtilities/Private/USDIntegrationUtils.cpp
+++ b/Engine/Plugins/Runtime/USDCore/Source/USDUtilities/Private/USDIntegrationUtils.cpp
@@ -4,6 +4,7 @@
 
 #include "UnrealUSDWrapper.h"
 #include "USDAttributeUtils.h"
+#include "USDMemory.h"
 #include "USDTypesConversion.h"
 #include "UsdWrappers/SdfChangeBlock.h"
 #include "UsdWrappers/UsdAttribute.h"
diff --git a/Engine/Plugins/Runtime/USDCore/Source/USDUtilities/Private/USDLayerUtils.cpp b/Engine/Plugins/Runtime/USDCore/Source/USDUtilities/Private/USDLayerUtils.cpp
--- a/Engine/Plugins/Runtime/USDCore/Source/USDUtilities/Private/USDLayerUtils.cpp
+++ b/Engine/Plugins/Runtime/USDCore/Source/USDUtilities/Private/USDLayerUtils.cpp
@@ -5,6 +5,7 @@
 #include "UnrealUSDWrapper.h"
 #include "USDErrorUtils.h"
 #include "USDLog.h"
+#include "USDMemory.h"
 #include "USDTypesConversion.h"
 
 #include "UsdWrappers/SdfChangeBlock.h"
diff --git a/Engine/Plugins/Runtime/USDCore/Source/USDUtilities/Private/USDLightConversion.cpp b/Engine/Plugins/Runtime/USDCore/Source/USDUtilities/Private/USDLightConversion.cpp
--- a/Engine/Plugins/Runtime/USDCore/Source/USDUtilities/Private/USDLightConversion.cpp
+++ b/Engine/Plugins/Runtime/USDCore/Source/USDUtilities/Private/USDLightConversion.cpp
@@ -9,6 +9,7 @@
 #include "USDConversionUtils.h"
 #include "USDLayerUtils.h"
 #include "USDLog.h"
+#include "USDMemory.h"
 #include "USDObjectUtils.h"
 #include "USDShadeConversion.h"
 #include "USDTypesConversion.h"
diff --git a/Engine/Plugins/Runtime/USDCore/Source/USDUtilities/Private/USDPrimConversion.cpp b/Engine/Plugins/Runtime/USDCore/Source/USDUtilities/Private/USDPrimConversion.cpp
--- a/Engine/Plugins/Runtime/USDCore/Source/USDUtilities/Private/USDPrimConversion.cpp
+++ b/Engine/Plugins/Runtime/USDCore/Source/USDUtilities/Private/USDPrimConversion.cpp
@@ -10,12 +10,12 @@
 #include "USDLayerUtils.h"
 #include "USDLightConversion.h"
 #include "USDLog.h"
+#include "USDMemory.h"
 #include "USDObjectUtils.h"
 #include "USDShadeConversion.h"
 #include "USDSkeletalDataConversion.h"
 #include "USDTypesConversion.h"
 #include "USDValueConversion.h"
-
 #include "UsdWrappers/UsdAttribute.h"
 #include "UsdWrappers/UsdGeomBBoxCache.h"
 #include "UsdWrappers/UsdGeomXformable.h"
diff --git a/Engine/Plugins/Runtime/USDCore/Source/USDUtilities/Private/USDTypesConversion.cpp b/Engine/Plugins/Runtime/USDCore/Source/USDUtilities/Private/USDTypesConversion.cpp
--- a/Engine/Plugins/Runtime/USDCore/Source/USDUtilities/Private/USDTypesConversion.cpp
+++ b/Engine/Plugins/Runtime/USDCore/Source/USDUtilities/Private/USDTypesConversion.cpp
@@ -3,6 +3,7 @@
 #include "USDTypesConversion.h"
 
 #include "USDConversionUtils.h"
+#include "USDMemory.h"
 
 #include "Containers/StringConv.h"
 
diff --git a/Engine/Plugins/Runtime/USDCore/Source/USDUtilities/Private/USDValueConversion.cpp b/Engine/Plugins/Runtime/USDCore/Source/USDUtilities/Private/USDValueConversion.cpp
--- a/Engine/Plugins/Runtime/USDCore/Source/USDUtilities/Private/USDValueConversion.cpp
+++ b/Engine/Plugins/Runtime/USDCore/Source/USDUtilities/Private/USDValueConversion.cpp
@@ -4,8 +4,8 @@
 
 #include "USDConversionUtils.h"
 #include "USDLog.h"
+#include "USDMemory.h"
 #include "USDTypesConversion.h"
-
 #include "UsdWrappers/UsdStage.h"
 #include "UsdWrappers/VtValue.h"
 
@@ -1464,14 +1464,14 @@ namespace UsdUtils
 		{
 			// If the string contains a double quote, use a single quote as a delimiter, or else the string
 			// will be impossible to tokenize into the source string later when unstringifying
-            if (Element.Contains(&DoubleQuoteTCHAR))
-            {
-                Result += FString::Printf(TEXT("'%s', "), *Element);
-            }
-            else
-            {
-                Result += FString::Printf(TEXT("\"%s\", "), *Element);
-            }
+			if (Element.Contains(&DoubleQuoteTCHAR))
+			{
+				Result += FString::Printf(TEXT("'%s', "), *Element);
+			}
+			else
+			{
+				Result += FString::Printf(TEXT("\"%s\", "), *Element);
+			}
 		}
 		Result.RemoveFromEnd(TEXT(", "));
 		Result += TEXT("]");
@@ -1518,14 +1518,14 @@ namespace UsdUtils
 		{
 			// If the string contains a double quote, use a single quote as a delimiter, or else the string
 			// will be impossible to tokenize into the source string later when unstringifying
-            if (Element.Contains(&DoubleQuoteTCHAR))
-            {
-                Result += FString::Printf(TEXT("'%s', "), *Element);
-            }
-            else
-            {
-                Result += FString::Printf(TEXT("\"%s\", "), *Element);
-            }
+			if (Element.Contains(&DoubleQuoteTCHAR))
+			{
+				Result += FString::Printf(TEXT("'%s', "), *Element);
+			}
+			else
+			{
+				Result += FString::Printf(TEXT("\"%s\", "), *Element);
+			}
 		}
 		Result.RemoveFromEnd(TEXT(", "));
 		Result += TEXT("])");
@@ -2349,7 +2349,7 @@ namespace UsdUtils
 
  // // Exotic types found in some scenarios
 			{TEXT("SdfListOp<TfToken>"), WrapInVtValue<pxr::SdfListOp<pxr::TfToken>, UnstringifyListOpTokens>}, // This is used for apiSchemas
-  // metadata
+	// metadata
 		};
 
 		if (const TFunction<bool(const FString&, pxr::VtValue&)>* FoundUnstringifier = Unstringifiers.Find(TypeName))
diff --git a/Engine/Plugins/Runtime/USDCore/Source/USDUtilities/Public/USDConversionUtils.h b/Engine/Plugins/Runtime/USDCore/Source/USDUtilities/Public/USDConversionUtils.h
--- a/Engine/Plugins/Runtime/USDCore/Source/USDUtilities/Public/USDConversionUtils.h
+++ b/Engine/Plugins/Runtime/USDCore/Source/USDUtilities/Public/USDConversionUtils.h
@@ -8,7 +8,9 @@
 #include "UObject/UObjectGlobals.h"
 
 #include "UnrealUSDWrapper.h"
+#if UE_ENABLE_INCLUDE_ORDER_DEPRECATED_IN_5_5
 #include "USDMemory.h"
+#endif	  // UE_ENABLE_INCLUDE_ORDER_DEPRECATED_IN_5_5
 #include "UsdWrappers/ForwardDeclarations.h"
 #include "UsdWrappers/SdfLayer.h"
 #include "UsdWrappers/SdfPath.h"
@@ -55,6 +57,8 @@ enum class EUsdDrawMode : int32;
 enum class EUsdDuplicateType : uint8;
 enum class EUsdUpAxis : uint8;
 struct FUsdUnrealAssetInfo;
+template<typename UsdObjectType>
+class TUsdStore;
 
 namespace UE
 {
diff --git a/Engine/Plugins/Runtime/USDCore/Source/USDUtilities/Public/USDLayerUtils.h b/Engine/Plugins/Runtime/USDCore/Source/USDUtilities/Public/USDLayerUtils.h
--- a/Engine/Plugins/Runtime/USDCore/Source/USDUtilities/Public/USDLayerUtils.h
+++ b/Engine/Plugins/Runtime/USDCore/Source/USDUtilities/Public/USDLayerUtils.h
@@ -4,7 +4,9 @@
 
 #if USE_USD_SDK
 
+#if UE_ENABLE_INCLUDE_ORDER_DEPRECATED_IN_5_5
 #include "USDMemory.h"
+#endif	  // UE_ENABLE_INCLUDE_ORDER_DEPRECATED_IN_5_5
 #include "UsdWrappers/ForwardDeclarations.h"
 
 #include "USDIncludesStart.h"
@@ -12,6 +14,8 @@
 #include "USDIncludesEnd.h"
 
 class SWidget;
+template<typename UsdObjectType>
+class TUsdStore;
 
 PXR_NAMESPACE_OPEN_SCOPE
 	class SdfLayer;
diff --git a/Engine/Plugins/Runtime/USDCore/Source/USDUtilities/Public/USDTypesConversion.h b/Engine/Plugins/Runtime/USDCore/Source/USDUtilities/Public/USDTypesConversion.h
--- a/Engine/Plugins/Runtime/USDCore/Source/USDUtilities/Public/USDTypesConversion.h
+++ b/Engine/Plugins/Runtime/USDCore/Source/USDUtilities/Public/USDTypesConversion.h
@@ -4,7 +4,9 @@
 
 #if USE_USD_SDK
 #include "UnrealUSDWrapper.h"
+#if UE_ENABLE_INCLUDE_ORDER_DEPRECATED_IN_5_5
 #include "USDMemory.h"
+#endif	  // UE_ENABLE_INCLUDE_ORDER_DEPRECATED_IN_5_5
 #include "USDStageOptions.h"
 
 #include <string>
@@ -41,10 +43,13 @@ PXR_NAMESPACE_OPEN_SCOPE
 
 	using UsdStageRefPtr = TfRefPtr<UsdStage>;
 
-	template <typename T>
+	template<typename T>
 	class VtArray;
 PXR_NAMESPACE_CLOSE_SCOPE
 
+template<typename UsdObjectType>
+class TUsdStore;
+
 struct USDUTILITIES_API FUsdStageInfo
 {
 	EUsdUpAxis UpAxis = EUsdUpAxis::ZAxis;
diff --git a/Engine/Plugins/Runtime/USDCore/Source/UnrealUSDWrapper/Private/UsdWrappers/UsdEditContext.cpp b/Engine/Plugins/Runtime/USDCore/Source/UnrealUSDWrapper/Private/UsdWrappers/UsdEditContext.cpp
--- a/Engine/Plugins/Runtime/USDCore/Source/UnrealUSDWrapper/Private/UsdWrappers/UsdEditContext.cpp
+++ b/Engine/Plugins/Runtime/USDCore/Source/UnrealUSDWrapper/Private/UsdWrappers/UsdEditContext.cpp
@@ -2,6 +2,7 @@
 
 #include "UsdWrappers/UsdEditContext.h"
 
+#include "USDMemory.h"
 #include "UsdWrappers/SdfLayer.h"
 #include "UsdWrappers/UsdStage.h"
 
diff --git a/Engine/Plugins/Runtime/USDCore/Source/UnrealUSDWrapper/Public/USDMemory.h b/Engine/Plugins/Runtime/USDCore/Source/UnrealUSDWrapper/Public/USDMemory.h
--- a/Engine/Plugins/Runtime/USDCore/Source/UnrealUSDWrapper/Public/USDMemory.h
+++ b/Engine/Plugins/Runtime/USDCore/Source/UnrealUSDWrapper/Public/USDMemory.h
@@ -4,6 +4,7 @@
 
 #include "CoreMinimal.h"
 #include "HAL/LowLevelMemTracker.h"
+#include "HAL/Platform.h"
 #include "Misc/Optional.h"
 #include "Misc/ScopeRWLock.h"
 #include "Modules/Boilerplate/ModuleBoilerplate.h"
@@ -20,13 +21,13 @@ LLM_DECLARE_TAG_API(Usd, UNREALUSDWRAPPER_API);
  * The USD memory tools are only supported in non monolithic builds at this time, since it requires overriding new and delete per module.
  *
  * The USD SDK uses the shared C runtime allocator. This means that objects returned by the USD SDK might try to delete objects that were allocated
- *through the CRT. Since UE overrides new and delete per module, USD objects that try to call delete will end up freeing memory with the UE allocator
- *but the malloc was made using the CRT, leading in a crash.
+ * through the CRT. Since UE overrides new and delete per module, USD objects that try to call delete will end up freeing memory with the UE allocator
+ * but the malloc was made using the CRT, leading in a crash.
  *
  * To go around this problem, modules using the USD SDK need special operators for new and delete. Those operators have the ability to redirect the
- *malloc or free calls to either the UE allocator or to the CRT allocator. The choice of the allocator is made in the FUsdMemoryManager.
- *FUsdMemoryManager manages a stack of active allocators per thread. Using ActivateAllocator and DeactivateAllocator, we can push and pop which
- *allocator is active on the calling thread.
+ * malloc or free calls to either the UE allocator or to the CRT allocator. The choice of the allocator is made in the FUsdMemoryManager.
+ * FUsdMemoryManager manages a stack of active allocators per thread. Using ActivateAllocator and DeactivateAllocator, we can push and pop which
+ * allocator is active on the calling thread.
  *
  * To simplify the workflow, TScopedAllocs is provided to make sure a certain block of code is bound to the right allocator.
  * FScopedUsdAllocs is a TScopedAllocs that activates the CRT allocator, while FScopedUnrealAllocs is the one that activates the UE allocator.
@@ -126,6 +127,21 @@ class TScopedAllocs final
 public:
 	TScopedAllocs()
 	{
+		// clang-format off
+#if !FORCE_ANSI_ALLOCATOR && !IS_MONOLITHIC && !USD_MERGED_MODULES	  // If we're in a situation where we need the overriden allocators...
+#ifndef SUPPRESS_PER_MODULE_INLINE_FILE	   // ...but we don't have this define (we can't check for whether IMPLEMENT_MODULE_USD is actually used, but
+										   // hopefully this is a good enough proxy)
+#ifndef DISABLE_USDMEMORY_WARNING	// This is not normally defined anywhere, but can be defined by the user to disable this check and allow
+									// compilation anyway
+#if !UE_ENABLE_INCLUDE_ORDER_DEPRECATED_IN_5_5	  // Ignore the warning when this is defined, because then we'll still be getting transitive
+												  // USDMemory.h includes from other USD files
+COMPILE_WARNING("You should only include USDMemory.h when necessary. In order to use the allocators from USDMemory.h, make sure that your module uses 'IMPLEMENT_MODULE_USD' on your main module .cpp file, and that 'PrivateDefinitions.Add(\"SUPPRESS_PER_MODULE_INLINE_FILE\");' is added to your module's .Build.cs file. For more info, refer to the official USDImporter documentation, or the documentation comments at the start of USDMemory.h, as well as on the end of UnrealUSDWrapper.Build.cs. You can also disable this compile-time check by defining DISABLE_USDMEMORY_WARNING before including USDMemory.h.")
+#endif
+#endif
+#endif
+#endif
+		// clang-format on
+
 		FUsdMemoryManager::ActivateAllocator(AllocatorType);
 	}
 
diff --git a/Engine/Plugins/Runtime/USDCore/Source/UnrealUSDWrapper/Public/UnrealUSDWrapper.h b/Engine/Plugins/Runtime/USDCore/Source/UnrealUSDWrapper/Public/UnrealUSDWrapper.h
--- a/Engine/Plugins/Runtime/USDCore/Source/UnrealUSDWrapper/Public/UnrealUSDWrapper.h
+++ b/Engine/Plugins/Runtime/USDCore/Source/UnrealUSDWrapper/Public/UnrealUSDWrapper.h
@@ -8,7 +8,9 @@
 #include "Modules/ModuleInterface.h"
 #include "Templates/Tuple.h"
 #include "UObject/ObjectMacros.h"
+#if UE_ENABLE_INCLUDE_ORDER_DEPRECATED_IN_5_5
 #include "USDMemory.h"
+#endif	  // UE_ENABLE_INCLUDE_ORDER_DEPRECATED_IN_5_5
 
 #include "UsdWrappers/ForwardDeclarations.h"
 
