Interchange Collision:
- Fix typos, and make them visible when import into level.

#virtualized

diff --git a/Engine/Plugins/Interchange/Runtime/Source/Pipelines/Public/InterchangeGenericMeshPipeline.h b/Engine/Plugins/Interchange/Runtime/Source/Pipelines/Public/InterchangeGenericMeshPipeline.h
--- a/Engine/Plugins/Interchange/Runtime/Source/Pipelines/Public/InterchangeGenericMeshPipeline.h
+++ b/Engine/Plugins/Interchange/Runtime/Source/Pipelines/Public/InterchangeGenericMeshPipeline.h
@@ -70,7 +70,7 @@ public:
 	/** If enabled, custom collision will be imported. If enabled and there is no custom collision, a generic collision will be automatically generated.
 	 * If disabled, no collision will be created or imported.
 	 */
-	UPROPERTY(EditAnywhere, BlueprintReadWrite, Category = "Static Meshes", meta = (DisplayName = "Collision", SubCategory = "Collision"))
+	UPROPERTY(EditAnywhere, BlueprintReadWrite, Category = "Static Meshes", meta = (DisplayName = "Import Collisions", SubCategory = "Collision"))
 	bool bCollision = true;
 
 	/**
@@ -82,15 +82,15 @@ public:
 	 * USP_ Sphere collision
 	 * UCX_ Convex collision
 	 */
-	UPROPERTY(EditAnywhere, BlueprintReadWrite, Category = "Static Meshes", meta = (SubCategory = "Collision", editcondition = "bCollision == true"))
+	UPROPERTY(EditAnywhere, BlueprintReadWrite, Category = "Static Meshes", meta = (DisplayName = "Import Collisions According To Mesh Name", SubCategory = "Collision", editcondition = "bCollision == true"))
 	bool bImportCollisionAccordingToMeshName = true;
 
 	/** If enabled, each UCX collision mesh will be imported as a single convex hull. If disabled, a UCX mesh will be decomposed into its separate pieces and a convex hull generated for each. */
 	UPROPERTY(EditAnywhere, BlueprintReadWrite, Category = "Static Meshes", meta = (SubCategory = "Collision", editcondition = "bCollision == true && bImportCollisionAccordingToMeshName"))
 	bool bOneConvexHullPerUCX = true;
 
-	/** Type used to generate collision when no custom collision provided. */
-	UPROPERTY(EditAnywhere, BlueprintReadWrite, Category = "Static Meshes", meta = (ScriptName = "DefaultCollisionType", DisplayName = "Default Collision Type", SubCategory = "Collision", editcondition = "bCollision == true"))
+	/** Type used to generate a collision when no custom collisions are present in the file. */
+	UPROPERTY(EditAnywhere, BlueprintReadWrite, Category = "Static Meshes", meta = (ScriptName = "FallbackCollisionType", DisplayName = "Fallback Collision Type", SubCategory = "Collision", editcondition = "bCollision == true"))
 	EInterchangeMeshCollision Collision = EInterchangeMeshCollision::Convex18DOP;
 
 	//////	Static Meshes Build settings Properties //////
