Interchange: Fix mesh baking transform when a translated scene node referencing a mesh do not have a parent.

#rnx

diff --git a/Engine/Plugins/Interchange/Runtime/Source/Pipelines/Private/InterchangePipelineMeshesUtilities.cpp b/Engine/Plugins/Interchange/Runtime/Source/Pipelines/Private/InterchangePipelineMeshesUtilities.cpp
--- a/Engine/Plugins/Interchange/Runtime/Source/Pipelines/Private/InterchangePipelineMeshesUtilities.cpp
+++ b/Engine/Plugins/Interchange/Runtime/Source/Pipelines/Private/InterchangePipelineMeshesUtilities.cpp
@@ -325,11 +325,13 @@ UInterchangePipelineMeshesUtilities* UInterchangePipelineMeshesUtilities::Create
 
 						if (MeshNode && MeshNode->IsA<UInterchangeMeshNode>())
 						{
+							FInterchangeMeshGeometry& MeshGeometry = PipelineMeshesUtilities->MeshGeometriesPerMeshUid.FindChecked(MeshUid);
+
 							const UInterchangeSceneNode* ParentMeshSceneNode = Cast<UInterchangeSceneNode>(BaseNodeContainer->GetNode(SceneNode->GetParentUid()));
+							const UInterchangeSceneNode* LodGroupNode = nullptr;
+							int32 LodIndex = 0;
 							if (ParentMeshSceneNode)
 							{
-								const UInterchangeSceneNode* LodGroupNode = nullptr;
-								int32 LodIndex = 0;
 								if (SceneMeshNodeUidsPerLodParentUidMap.Contains(ParentMeshSceneNode->GetUniqueID()))
 								{
 									LodGroupNode = ParentMeshSceneNode;
@@ -368,43 +370,43 @@ UInterchangePipelineMeshesUtilities* UInterchangePipelineMeshesUtilities::Create
 										ParentMeshSceneNode = Cast<const UInterchangeSceneNode>(BaseNodeContainer->GetNode(ParentMeshSceneNode->GetParentUid()));
 									} while (ParentMeshSceneNode);
 								}
-								FInterchangeMeshGeometry& MeshGeometry = PipelineMeshesUtilities->MeshGeometriesPerMeshUid.FindChecked(MeshUid);
-								if (LodGroupNode)
+							}
+
+							if (LodGroupNode)
+							{
+								//We have a LOD
+								FInterchangeMeshInstance& MeshInstance = PipelineMeshesUtilities->MeshInstancesPerMeshInstanceUid.FindOrAdd(LodGroupNode->GetUniqueID());
+								if (MeshInstance.LodGroupNode != nullptr)
 								{
-									//We have a LOD
-									FInterchangeMeshInstance& MeshInstance = PipelineMeshesUtilities->MeshInstancesPerMeshInstanceUid.FindOrAdd(LodGroupNode->GetUniqueID());
-									if (MeshInstance.LodGroupNode != nullptr)
-									{
-										//This LodGroup was already created, verify everything is ok
-										checkSlow(MeshInstance.LodGroupNode == LodGroupNode);
-										checkSlow(MeshInstance.MeshInstanceUid.Equals(LodGroupNode->GetUniqueID()));
-									}
-									else
-									{
-										MeshInstance.LodGroupNode = LodGroupNode;
-										MeshInstance.MeshInstanceUid = LodGroupNode->GetUniqueID();
-									}
-									FInterchangeLodSceneNodeContainer& InstancedSceneNodes = MeshInstance.SceneNodePerLodIndex.FindOrAdd(LodIndex);
-									InstancedSceneNodes.SceneNodes.AddUnique(SceneNode);
-									MeshGeometry.ReferencingMeshInstanceUids.Add(MeshInstance.MeshInstanceUid);
-									MeshInstance.ReferencingMeshGeometryUids.Add(MeshUid);
-									MeshInstance.bReferenceSkinnedMesh |= MeshGeometry.MeshNode->IsSkinnedMesh();
-									MeshInstance.bReferenceMorphTarget |= MeshGeometry.MeshNode->IsMorphTarget();
-									MeshInstance.bHasMorphTargets |= MeshGeometry.MeshNode->GetMorphTargetDependeciesCount() > 0;
+									//This LodGroup was already created, verify everything is ok
+									checkSlow(MeshInstance.LodGroupNode == LodGroupNode);
+									checkSlow(MeshInstance.MeshInstanceUid.Equals(LodGroupNode->GetUniqueID()));
 								}
 								else
 								{
-									FInterchangeMeshInstance& MeshInstance = PipelineMeshesUtilities->MeshInstancesPerMeshInstanceUid.FindOrAdd(NodeUid);
-									MeshInstance.LodGroupNode = nullptr;
-									MeshInstance.MeshInstanceUid = NodeUid;
-									FInterchangeLodSceneNodeContainer& InstancedSceneNodes = MeshInstance.SceneNodePerLodIndex.FindOrAdd(LodIndex);
-									InstancedSceneNodes.SceneNodes.AddUnique(SceneNode);
-									MeshGeometry.ReferencingMeshInstanceUids.Add(MeshInstance.MeshInstanceUid);
-									MeshInstance.ReferencingMeshGeometryUids.Add(MeshUid);
-									MeshInstance.bReferenceSkinnedMesh |= MeshGeometry.MeshNode->IsSkinnedMesh();
-									MeshInstance.bReferenceMorphTarget |= MeshGeometry.MeshNode->IsMorphTarget();
-									MeshInstance.bHasMorphTargets |= MeshGeometry.MeshNode->GetMorphTargetDependeciesCount() > 0;
+									MeshInstance.LodGroupNode = LodGroupNode;
+									MeshInstance.MeshInstanceUid = LodGroupNode->GetUniqueID();
 								}
+								FInterchangeLodSceneNodeContainer& InstancedSceneNodes = MeshInstance.SceneNodePerLodIndex.FindOrAdd(LodIndex);
+								InstancedSceneNodes.SceneNodes.AddUnique(SceneNode);
+								MeshGeometry.ReferencingMeshInstanceUids.Add(MeshInstance.MeshInstanceUid);
+								MeshInstance.ReferencingMeshGeometryUids.Add(MeshUid);
+								MeshInstance.bReferenceSkinnedMesh |= MeshGeometry.MeshNode->IsSkinnedMesh();
+								MeshInstance.bReferenceMorphTarget |= MeshGeometry.MeshNode->IsMorphTarget();
+								MeshInstance.bHasMorphTargets |= MeshGeometry.MeshNode->GetMorphTargetDependeciesCount() > 0;
+							}
+							else
+							{
+								FInterchangeMeshInstance& MeshInstance = PipelineMeshesUtilities->MeshInstancesPerMeshInstanceUid.FindOrAdd(NodeUid);
+								MeshInstance.LodGroupNode = nullptr;
+								MeshInstance.MeshInstanceUid = NodeUid;
+								FInterchangeLodSceneNodeContainer& InstancedSceneNodes = MeshInstance.SceneNodePerLodIndex.FindOrAdd(LodIndex);
+								InstancedSceneNodes.SceneNodes.AddUnique(SceneNode);
+								MeshGeometry.ReferencingMeshInstanceUids.Add(MeshInstance.MeshInstanceUid);
+								MeshInstance.ReferencingMeshGeometryUids.Add(MeshUid);
+								MeshInstance.bReferenceSkinnedMesh |= MeshGeometry.MeshNode->IsSkinnedMesh();
+								MeshInstance.bReferenceMorphTarget |= MeshGeometry.MeshNode->IsMorphTarget();
+								MeshInstance.bHasMorphTargets |= MeshGeometry.MeshNode->GetMorphTargetDependeciesCount() > 0;
 							}
 						}
 					}
