Interchange: Fix collision when importing into level.

#rnx

diff --git a/Engine/Plugins/Interchange/Runtime/Source/Pipelines/Private/InterchangeGenericStaticMeshPipeline.cpp b/Engine/Plugins/Interchange/Runtime/Source/Pipelines/Private/InterchangeGenericStaticMeshPipeline.cpp
--- a/Engine/Plugins/Interchange/Runtime/Source/Pipelines/Private/InterchangeGenericStaticMeshPipeline.cpp
+++ b/Engine/Plugins/Interchange/Runtime/Source/Pipelines/Private/InterchangeGenericStaticMeshPipeline.cpp
@@ -23,7 +23,7 @@
 
 namespace UE::Interchange::Private
 {
-	FString GetNodeName(const UInterchangeBaseNodeContainer& NodeContainer, const FString& NodeUid)
+	FString GetNodeName(TObjectPtr<UInterchangePipelineMeshesUtilities> PipelineMeshesUtilities, const UInterchangeBaseNodeContainer& NodeContainer, const FString& NodeUid)
 	{
 		const UInterchangeBaseNode* BaseNode = NodeContainer.GetNode(NodeUid);
 		if (BaseNode)
@@ -36,7 +36,20 @@ namespace UE::Interchange::Private
 					TArray<FString> LodGroupChildrens = NodeContainer.GetNodeChildrenUids(SceneNode->GetUniqueID());
 					if (LodGroupChildrens.Num() > 0)
 					{
-						return GetNodeName(NodeContainer, LodGroupChildrens[0]);
+						return GetNodeName(PipelineMeshesUtilities, NodeContainer, LodGroupChildrens[0]);
+					}
+				}
+			}
+			else if (const UInterchangeMeshNode* MeshNode = Cast<UInterchangeMeshNode>(BaseNode))
+			{
+				//If this mesh is reference by only one scene node that do not have any children, use the scene node display label
+				const FInterchangeMeshGeometry& MeshGeometry = PipelineMeshesUtilities->GetMeshGeometryByUid(NodeUid);
+				if (MeshGeometry.ReferencingMeshInstanceUids.Num() == 1
+					&& NodeContainer.GetNodeChildrenCount(MeshGeometry.ReferencingMeshInstanceUids[0]) == 0)
+				{
+					if(const UInterchangeBaseNode* InstanceMeshNode = NodeContainer.GetNode(MeshGeometry.ReferencingMeshInstanceUids[0]))
+					{
+						return InstanceMeshNode->GetDisplayLabel();
 					}
 				}
 			}
@@ -50,9 +63,9 @@ namespace UE::Interchange::Private
 	}
 }
 
-static TTuple<EInterchangeMeshCollision, FString> GetCollisionMeshType(const UInterchangeBaseNodeContainer& NodeContainer, const FString& NodeUid, const TArray<FString>& AllNodeUids)
+static TTuple<EInterchangeMeshCollision, FString> GetCollisionMeshType(TObjectPtr<UInterchangePipelineMeshesUtilities> PipelineMeshesUtilities, const UInterchangeBaseNodeContainer& NodeContainer, const FString& NodeUid, const TArray<FString>& AllNodeUids)
 {
-	FString MeshName = UE::Interchange::Private::GetNodeName(NodeContainer, NodeUid);
+	FString MeshName = UE::Interchange::Private::GetNodeName(PipelineMeshesUtilities, NodeContainer, NodeUid);
 	EInterchangeMeshCollision CollisionType = EInterchangeMeshCollision::None;
 
 	// Determine if the mesh name is a potential collision mesh
@@ -88,10 +101,10 @@ static TTuple<EInterchangeMeshCollision, FString> GetCollisionMeshType(const UIn
 	int32 LastUnderscore = INDEX_NONE;
 	verify(MeshName.FindLastChar(TEXT('_'), LastUnderscore));
 
-	auto MatchPredicate = [&NodeContainer](FString Body)
+	auto MatchPredicate = [&PipelineMeshesUtilities , &NodeContainer](FString Body)
 	{
 		// Generate a predicate to be used by the below Finds.
-		return [Body, &NodeContainer](const FString& ToCompare) { return Body == UE::Interchange::Private::GetNodeName(NodeContainer, ToCompare); };
+		return [Body, &PipelineMeshesUtilities, &NodeContainer](const FString& ToCompare) { return Body == UE::Interchange::Private::GetNodeName(PipelineMeshesUtilities, NodeContainer, ToCompare); };
 	};
 
 	// If we find a mesh named the same as the collision mesh (following the collision prefix), we have a match
@@ -121,17 +134,17 @@ static TTuple<EInterchangeMeshCollision, FString> GetCollisionMeshType(const UIn
 }
 
 
-static bool IsCollisionMeshUid(const UInterchangeBaseNodeContainer& NodeContainer, const FString& MeshUid, const TArray<FString>& MeshUids)
+static bool IsCollisionMeshUid(TObjectPtr<UInterchangePipelineMeshesUtilities> PipelineMeshesUtilities, const UInterchangeBaseNodeContainer& NodeContainer, const FString& MeshUid, const TArray<FString>& MeshUids)
 {
-	return GetCollisionMeshType(NodeContainer, MeshUid, MeshUids).Get<0>() != EInterchangeMeshCollision::None;
+	return GetCollisionMeshType(PipelineMeshesUtilities, NodeContainer, MeshUid, MeshUids).Get<0>() != EInterchangeMeshCollision::None;
 }
 
 
-static void BuildMeshToCollisionMeshMap(const UInterchangeBaseNodeContainer& NodeContainer, const TArray<FString>& MeshUids, TMap<FString, TArray<FString>>& MeshToCollisionMeshMap)
+static void BuildMeshToCollisionMeshMap(TObjectPtr<UInterchangePipelineMeshesUtilities> PipelineMeshesUtilities, const UInterchangeBaseNodeContainer& NodeContainer, const TArray<FString>& MeshUids, TMap<FString, TArray<FString>>& MeshToCollisionMeshMap)
 {
 	for (const FString& MeshUid : MeshUids)
 	{
-		TTuple<EInterchangeMeshCollision, FString> CollisionType = GetCollisionMeshType(NodeContainer, MeshUid, MeshUids);
+		TTuple<EInterchangeMeshCollision, FString> CollisionType = GetCollisionMeshType(PipelineMeshesUtilities, NodeContainer, MeshUid, MeshUids);
 		if (CollisionType.Get<0>() != EInterchangeMeshCollision::None)
 		{
 			MeshToCollisionMeshMap.FindOrAdd(FString(CollisionType.Get<1>())).Emplace(MeshUid);
@@ -241,13 +254,13 @@ void UInterchangeGenericMeshPipeline::ExecutePreImportPipelineStaticMesh()
 				TMap<FString, TArray<FString>> MeshToCollisionMeshMap;
 				if (bImportCollisionAccordingToMeshName)
 				{
-					BuildMeshToCollisionMeshMap(*BaseNodeContainer, MeshUids, MeshToCollisionMeshMap);
+					BuildMeshToCollisionMeshMap(PipelineMeshesUtilities, *BaseNodeContainer, MeshUids, MeshToCollisionMeshMap);
 				}
 
 				// Now iterate through each mesh UID, creating a new factory for each one
 				for (const FString& MeshUid : MeshUids)
 				{
-					if (bImportCollisionAccordingToMeshName && IsCollisionMeshUid(*BaseNodeContainer, MeshUid, MeshUids))
+					if (bImportCollisionAccordingToMeshName && IsCollisionMeshUid(PipelineMeshesUtilities, *BaseNodeContainer, MeshUid, MeshUids))
 					{
 						// If this is a collision mesh, don't add a factory; it will be added as part of another factory
 						continue;
@@ -270,7 +283,7 @@ void UInterchangeGenericMeshPipeline::ExecutePreImportPipelineStaticMesh()
 
 					if (MeshUidsPerLodIndex.Num() > 0)
 					{
-						if (bCollision && Collision != EInterchangeMeshCollision::None && bImportCollisionAccordingToMeshName)
+						if (bCollision && bImportCollisionAccordingToMeshName)
 						{
 							if (const TArray<FString>* CorrespondingCollisionMeshes = MeshToCollisionMeshMap.Find(MeshUid))
 							{
@@ -294,12 +307,12 @@ void UInterchangeGenericMeshPipeline::ExecutePreImportPipelineStaticMesh()
 				TMap<FString, TArray<FString>> MeshToCollisionMeshMap;
 				if (bImportCollisionAccordingToMeshName)
 				{
-					BuildMeshToCollisionMeshMap(*BaseNodeContainer, MeshUids, MeshToCollisionMeshMap);
+					BuildMeshToCollisionMeshMap(PipelineMeshesUtilities, *BaseNodeContainer, MeshUids, MeshToCollisionMeshMap);
 				}
 
 				for (const FString& MeshUid : MeshUids)
 				{
-					if (bImportCollisionAccordingToMeshName && IsCollisionMeshUid(*BaseNodeContainer, MeshUid, MeshUids))
+					if (bImportCollisionAccordingToMeshName && IsCollisionMeshUid(PipelineMeshesUtilities, *BaseNodeContainer, MeshUid, MeshUids))
 					{
 						// If this is a collision mesh, don't add a factory; it will be added as part of another factory
 						continue;
@@ -314,7 +327,7 @@ void UInterchangeGenericMeshPipeline::ExecutePreImportPipelineStaticMesh()
 
 					if (MeshUidsPerLodIndex.Num() > 0)
 					{
-						if (bCollision && Collision != EInterchangeMeshCollision::None && bImportCollisionAccordingToMeshName)
+						if (bCollision && bImportCollisionAccordingToMeshName)
 						{
 							if (const TArray<FString>* CorrespondingCollisionMeshes = MeshToCollisionMeshMap.Find(MeshUid))
 							{
@@ -583,7 +596,7 @@ void UInterchangeGenericMeshPipeline::AddLodDataToStaticMesh(UInterchangeStaticM
 
 			if (bImportCollisionAccordingToMeshName)
 			{
-				TTuple<EInterchangeMeshCollision, FStringView> MeshType = GetCollisionMeshType(*BaseNodeContainer, NodeUid, NodeUids);
+				TTuple<EInterchangeMeshCollision, FStringView> MeshType = GetCollisionMeshType(PipelineMeshesUtilities, *BaseNodeContainer, NodeUid, NodeUids);
 				switch (MeshType.Get<0>())
 				{
 				case EInterchangeMeshCollision::None:
diff --git a/Engine/Plugins/Tests/InterchangeTests/Source/InterchangeTests/Private/InterchangeImportTest.cpp b/Engine/Plugins/Tests/InterchangeTests/Source/InterchangeTests/Private/InterchangeImportTest.cpp
--- a/Engine/Plugins/Tests/InterchangeTests/Source/InterchangeTests/Private/InterchangeImportTest.cpp
+++ b/Engine/Plugins/Tests/InterchangeTests/Source/InterchangeTests/Private/InterchangeImportTest.cpp
@@ -309,7 +309,10 @@ bool FInterchangeImportTest::RunTest(const FString& Path)
 			}
 			else
 			{
-				ObjectsToDelete.Add(ResultObject);
+				if (ResultObject && !ResultObject->HasAnyFlags(RF_BeginDestroyed | RF_FinishDestroyed))
+				{
+					ObjectsToDelete.Add(ResultObject);
+				}
 			}
 		}
 	}
