Interchange: fix assertion when getting payloads

#rnx

diff --git a/Engine/Plugins/Interchange/Runtime/Source/Import/Private/Mesh/InterchangeSkeletalMeshFactory.cpp b/Engine/Plugins/Interchange/Runtime/Source/Import/Private/Mesh/InterchangeSkeletalMeshFactory.cpp
--- a/Engine/Plugins/Interchange/Runtime/Source/Import/Private/Mesh/InterchangeSkeletalMeshFactory.cpp
+++ b/Engine/Plugins/Interchange/Runtime/Source/Import/Private/Mesh/InterchangeSkeletalMeshFactory.cpp
@@ -89,7 +89,7 @@ namespace UE
 					FInterchangeMeshPayLoadKey GlobalPayloadKey = MorphTargetKey;
 					if (SceneGlobalTransform.IsSet())
 					{
-						GlobalPayloadKey.UniqueId += SceneGlobalTransform->ToString();
+						GlobalPayloadKey.UniqueId += FInterchangeMeshPayLoadKey::GetTransformString(SceneGlobalTransform.GetValue());
 					}
 					return GlobalPayloadKey;
 				}
@@ -100,7 +100,7 @@ namespace UE
 					FString UniqueId = TranslatorPayloadKey.UniqueId;
 					if (SceneGlobalTransform.IsSet())
 					{
-						UniqueId += SceneGlobalTransform->ToString();
+						UniqueId += FInterchangeMeshPayLoadKey::GetTransformString(SceneGlobalTransform.GetValue());
 					}
 					return UniqueId;
 				}
diff --git a/Engine/Plugins/Interchange/Runtime/Source/Import/Private/Mesh/InterchangeStaticMeshFactory.cpp b/Engine/Plugins/Interchange/Runtime/Source/Import/Private/Mesh/InterchangeStaticMeshFactory.cpp
--- a/Engine/Plugins/Interchange/Runtime/Source/Import/Private/Mesh/InterchangeStaticMeshFactory.cpp
+++ b/Engine/Plugins/Interchange/Runtime/Source/Import/Private/Mesh/InterchangeStaticMeshFactory.cpp
@@ -254,7 +254,7 @@ void UInterchangeStaticMeshFactory::CreatePayloadTasks(const FImportAssetObjectP
 					FInterchangeMeshPayLoadKey& PayLoadKey = OptionalPayLoadKey.GetValue();
 					
 					FInterchangeMeshPayLoadKey GlobalPayLoadKey = PayLoadKey;
-					GlobalPayLoadKey.UniqueId += GlobalMeshTransform.ToString();
+					GlobalPayLoadKey.UniqueId += FInterchangeMeshPayLoadKey::GetTransformString(GlobalMeshTransform);
 					if (!PayloadPerKey.Contains(GlobalPayLoadKey))
 					{
 						FMeshPayload& Payload = PayloadPerKey.FindOrAdd(GlobalPayLoadKey);
diff --git a/Engine/Plugins/Interchange/Runtime/Source/Nodes/Public/InterchangeMeshNode.h b/Engine/Plugins/Interchange/Runtime/Source/Nodes/Public/InterchangeMeshNode.h
--- a/Engine/Plugins/Interchange/Runtime/Source/Nodes/Public/InterchangeMeshNode.h
+++ b/Engine/Plugins/Interchange/Runtime/Source/Nodes/Public/InterchangeMeshNode.h
@@ -65,6 +65,15 @@ public:
 		return UniqueId.Equals(Other.UniqueId) && Type == Other.Type;
 	}
 
+	//Return the translator key merge with the transform
+	static FString GetTransformString(const FTransform& Transform)
+	{
+		const FQuat R(Transform.GetRotation());
+		const FVector TT(Transform.GetTranslation());
+		const FVector S(Transform.GetScale3D());
+		return FString::Printf(TEXT("%.5f,%.5f,%.5f|%.5f,%.5f,%.5f,%.5f|%.5f,%.5f,%.5f"), TT.X, TT.Y, TT.Z, R.X, R.Y, R.Z, R.W, S.X, S.Y, S.Z);
+	}
+
 	friend uint32 GetTypeHash(const FInterchangeMeshPayLoadKey& InterchangeMeshPayLoadKey)
 	{
 		return GetTypeHash(InterchangeMeshPayLoadKey.UniqueId + FString::FromInt(static_cast<int32>(InterchangeMeshPayLoadKey.Type)));
