USD: Fix DisplayColors material no longer working via Interchange USD.

Tactical renaming for clarity;
Remove global captures from lambdas;
Automatic reformatting via clang-format.

diff --git a/Engine/Plugins/Interchange/Runtime/Source/Import/Private/Usd/InterchangeUsdTranslator.cpp b/Engine/Plugins/Interchange/Runtime/Source/Import/Private/Usd/InterchangeUsdTranslator.cpp
--- a/Engine/Plugins/Interchange/Runtime/Source/Import/Private/Usd/InterchangeUsdTranslator.cpp
+++ b/Engine/Plugins/Interchange/Runtime/Source/Import/Private/Usd/InterchangeUsdTranslator.cpp
@@ -174,7 +174,7 @@ namespace UE::InterchangeUsdTranslator::Private
 			FString MaterialSlotName;
 			UInterchangeMeshNode* MeshNode;
 		};
-		TMap<FString, FString> PrimPathToMaterialPath;
+		TMap<FString, FString> MaterialUidToActualNodeUid;
 		TMap<FString, TArray<FMaterialSlotMesh>> PrimPathToSlotMeshNodes;
 	};
 
@@ -709,21 +709,21 @@ namespace UE::InterchangeUsdTranslator::Private
 	)
 	{
 		FString PrimPath = Prim.GetPrimPath().GetString();
-		FString NodeUid = MaterialPrefix + PrimPath;
-		FString NodeName(Prim.GetName().ToString());
+		FString MaterialUid = MaterialPrefix + PrimPath;
+		FString MaterialPrimName(Prim.GetName().ToString());
 
 		// Check if Node already exist with this ID
-		if (const UInterchangeMaterialInstanceNode* Node = Cast<UInterchangeMaterialInstanceNode>(NodeContainer.GetNode(NodeUid)))
+		if (const UInterchangeMaterialInstanceNode* Node = Cast<UInterchangeMaterialInstanceNode>(NodeContainer.GetNode(MaterialUid)))
 		{
 			return;
 		}
 
-		auto SetMaterialSlotDependencies = [this, &NodeUid]()
+		auto SetMaterialSlotDependencies = [this, &MaterialUid]()
 		{
 			// Now we need to check if we have to set the slot of the mesh nodes here
-			if (TArray<FMaterialSlotMesh>* SlotMeshes = PrimPathToSlotMeshNodes.Find(NodeUid))
+			if (TArray<FMaterialSlotMesh>* SlotMeshes = PrimPathToSlotMeshNodes.Find(MaterialUid))
 			{
-				if (FString* NewMaterialUID = PrimPathToMaterialPath.Find(NodeUid))
+				if (FString* NewMaterialUID = MaterialUidToActualNodeUid.Find(MaterialUid))
 				{
 					for (const FMaterialSlotMesh& MaterialSlotMesh : *SlotMeshes)
 					{
@@ -740,7 +740,7 @@ namespace UE::InterchangeUsdTranslator::Private
 
 		// Check for any references of MaterialX
 #if WITH_EDITOR
-		if(RenderContext == UnrealIdentifiers::MaterialXRenderContext)
+		if (RenderContext == UnrealIdentifiers::MaterialXRenderContext)
 		{
 			TArray<FString> FilePaths = UsdUtils::GetMaterialXFilePaths(Prim);
 			for (const FString& File : FilePaths)
@@ -763,12 +763,12 @@ namespace UE::InterchangeUsdTranslator::Private
 				// The material from the MaterialXTranslator doesn't have the same UID, both the PrimPath have the same name but not the same path
 				// We need to retrieve that name (which is the Material name) in the translator, then we can map it to the right mesh
 				NodeContainer.BreakableIterateNodesOfType<UInterchangeShaderGraphNode>(
-					[this, &NodeName, &NodeUid](const FString&, UInterchangeShaderGraphNode* ShaderGraphNode)
+					[this, &MaterialPrimName, &MaterialUid](const FString&, UInterchangeShaderGraphNode* ShaderGraphNode)
 					{
-						FString UID = ShaderGraphNode->GetUniqueID();
-						if (FPaths::GetBaseFilename(UID) == NodeName)
+						FString ShaderGraphUid = ShaderGraphNode->GetUniqueID();
+						if (FPaths::GetBaseFilename(ShaderGraphUid) == MaterialPrimName)
 						{
-							PrimPathToMaterialPath.Add(NodeUid, UID);
+							MaterialUidToActualNodeUid.Add(MaterialUid, ShaderGraphUid);
 							return true;
 						}
 						else
@@ -786,15 +786,15 @@ namespace UE::InterchangeUsdTranslator::Private
 				return;
 			}
 		}
-#endif // WITH_EDITOR
-		
+#endif	  // WITH_EDITOR
+
 		UInterchangeMaterialInstanceNode* MaterialNode = NewObject<UInterchangeMaterialInstanceNode>(&NodeContainer);
-		MaterialNode->InitializeNode(NodeUid, NodeName, EInterchangeNodeContainerType::TranslatedAsset);
-		MaterialNode->SetAssetName(NodeName);
+		MaterialNode->InitializeNode(MaterialUid, MaterialPrimName, EInterchangeNodeContainerType::TranslatedAsset);
+		MaterialNode->SetAssetName(MaterialPrimName);
 		NodeContainer.AddNode(MaterialNode);
 
 		// Set the material instance node to the correct mesh nodes
-		PrimPathToMaterialPath.Add(NodeUid, NodeUid);
+		MaterialUidToActualNodeUid.Add(MaterialUid, MaterialUid);
 		SetMaterialSlotDependencies();
 
 		UsdToUnreal::FUsdPreviewSurfaceMaterialData MaterialData;
@@ -1178,10 +1178,12 @@ namespace UE::InterchangeUsdTranslator::Private
 
 				// Get the Uid of the material instance that we'll end up assigning to this slot
 				FString MaterialInstanceUid;
+				bool bIsDisplayColor = false;
 				switch (Slot.AssignmentType)
 				{
 					case UsdUtils::EPrimAssignmentType::DisplayColor:
 					{
+						bIsDisplayColor = true;
 						AddDisplayColorMaterialInstanceNodeIfNeeded(NodeContainer, Slot.MaterialSource);
 						MaterialInstanceUid = MaterialPrefix + Slot.MaterialSource;	   // This is e.g. "!DisplayColor_0_1"
 						break;
@@ -1204,27 +1206,22 @@ namespace UE::InterchangeUsdTranslator::Private
 					}
 				}
 
-				// if we found a match let's set the slot to the corresponding Material (the Material path could come from another translator)
-				if (FString* NewMaterialInstanceUid = PrimPathToMaterialPath.Find(MaterialInstanceUid))
+				// If we're a DisplayColor material, we know everything we need right now
+				if (bIsDisplayColor)
+				{
+					MeshNode->SetSlotMaterialDependencyUid(SlotName, *MaterialInstanceUid);
+				}
+				// If we found a match let's set the slot to the corresponding Material right away, as we already must have traversed this material
+				else if (FString* ActualMaterialInstanceUid = MaterialUidToActualNodeUid.Find(MaterialInstanceUid))
 				{
-					MeshNode->SetSlotMaterialDependencyUid(SlotName, *NewMaterialInstanceUid);
+					MeshNode->SetSlotMaterialDependencyUid(SlotName, *ActualMaterialInstanceUid);
 				}
-				else	// otherwise it's up to the Material to attach itself to the mesh
+				// Otherwise, we need to wait until the material prim itself is translated, as we may need to defer to another translator
+				// (e.g. MaterialX) for the translation, which could generate an entirely different translated node we can't know about yet
+				else
 				{
 					// one material can be attached to several meshes
-					if (TArray<FMaterialSlotMesh>* SlotsMeshes = PrimPathToSlotMeshNodes.Find(MaterialInstanceUid))
-					{
-						SlotsMeshes->Add({SlotName, MeshNode});
-					}
-					else
-					{
-						PrimPathToSlotMeshNodes.Add(
-							MaterialInstanceUid,
-							{
-								{SlotName, MeshNode}
-						}
-						);
-					}
+					PrimPathToSlotMeshNodes.FindOrAdd(MaterialInstanceUid).Add({SlotName, MeshNode});
 				}
 			}
 		}
@@ -2700,16 +2697,16 @@ TArray<UE::Interchange::FAnimationPayloadData> UInterchangeUSDTranslator::GetAni
 {
 	using namespace UE::Interchange;
 	using namespace UE::InterchangeUsdTranslator::Private;
-	//This is the results we return
+	// This is the results we return
 	TArray<UE::Interchange::FAnimationPayloadData> AnimationPayloads;
 
-	//Maps to help sorting the queries by payload type
+	// Maps to help sorting the queries by payload type
 	TArray<int32> BakeQueryIndexes;
 	TArray<TArray<UE::Interchange::FAnimationPayloadData>> BakeAnimationPayloads;
 	TArray<int32> CurveQueryIndexes;
 	TArray<TArray<UE::Interchange::FAnimationPayloadData>> CurveAnimationPayloads;
-	
-	//Get all curves with a parallel for
+
+	// Get all curves with a parallel for
 	int32 PayloadCount = PayloadQueries.Num();
 	for (int32 PayloadIndex = 0; PayloadIndex < PayloadCount; ++PayloadIndex)
 	{
@@ -2725,17 +2722,16 @@ TArray<UE::Interchange::FAnimationPayloadData> UInterchangeUSDTranslator::GetAni
 		}
 	}
 
-	
 #if USE_USD_SDK
 
-	//Import the Baked curve payloads
-	if(BakeQueryIndexes.Num() > 0)
+	// Import the Baked curve payloads
+	if (BakeQueryIndexes.Num() > 0)
 	{
 		int32 BakePayloadCount = BakeQueryIndexes.Num();
 		TMap<FString, TArray<const UE::Interchange::FAnimationPayloadQuery*>> BatchedBakeQueries;
 		BatchedBakeQueries.Reserve(BakePayloadCount);
-		
-		//Get the BAKED transform synchronously, since there is some interchange task that parallel them
+
+		// Get the BAKED transform synchronously, since there is some interchange task that parallel them
 		for (int32 BakePayloadIndex = 0; BakePayloadIndex < BakePayloadCount; ++BakePayloadIndex)
 		{
 			if (!ensure(BakeQueryIndexes.IsValidIndex(BakePayloadIndex)))
@@ -2775,61 +2771,64 @@ TArray<UE::Interchange::FAnimationPayloadData> UInterchangeUSDTranslator::GetAni
 		}
 	}
 
-	//Import normal curves
-	if(CurveQueryIndexes.Num() > 0)
+	// Import normal curves
+	if (CurveQueryIndexes.Num() > 0)
 	{
-		auto GetAnimPayloadLambda = [&](int32 PayloadIndex)
+		auto GetAnimPayloadLambda = [this, &PayloadQueries, &CurveAnimationPayloads](int32 PayloadIndex)
+		{
+			if (!PayloadQueries.IsValidIndex(PayloadIndex))
 			{
-				if (!PayloadQueries.IsValidIndex(PayloadIndex))
-				{
-					return;
-				}
-				const UE::Interchange::FAnimationPayloadQuery& PayloadQuery = PayloadQueries[PayloadIndex];
-				EInterchangeAnimationPayLoadType PayloadType = PayloadQuery.PayloadKey.Type;
-				if(PayloadType == EInterchangeAnimationPayLoadType::CURVE
-					|| PayloadType == EInterchangeAnimationPayLoadType::STEPCURVE)
+				return;
+			}
+			const UE::Interchange::FAnimationPayloadQuery& PayloadQuery = PayloadQueries[PayloadIndex];
+			EInterchangeAnimationPayLoadType PayloadType = PayloadQuery.PayloadKey.Type;
+			if (PayloadType == EInterchangeAnimationPayLoadType::CURVE || PayloadType == EInterchangeAnimationPayLoadType::STEPCURVE)
+			{
+				// Property track animation queries.
+				//
+				// We're fine handling these in isolation (currently GetAnimationPayloadData is called with
+				// a single query at a time for these): Emit a separate task for each right away
+				FAnimationPayloadData AnimationPayLoadData{PayloadQuery.SceneNodeUniqueID, PayloadQuery.PayloadKey};
+				if (GetPropertyAnimationCurvePayloadData(Impl->UsdStage, PayloadQuery.PayloadKey.UniqueId, AnimationPayLoadData))
 				{
-					// Property track animation queries.
-					//
-					// We're fine handling these in isolation (currently GetAnimationPayloadData is called with
-					// a single query at a time for these): Emit a separate task for each right away
-					FAnimationPayloadData AnimationPayLoadData{ PayloadQuery.SceneNodeUniqueID, PayloadQuery.PayloadKey };
-					if (GetPropertyAnimationCurvePayloadData(Impl->UsdStage, PayloadQuery.PayloadKey.UniqueId, AnimationPayLoadData))
-					{
-						CurveAnimationPayloads[PayloadIndex].Emplace(AnimationPayLoadData);
-					}
+					CurveAnimationPayloads[PayloadIndex].Emplace(AnimationPayLoadData);
 				}
-				else if(PayloadType == EInterchangeAnimationPayLoadType::MORPHTARGETCURVE)
+			}
+			else if (PayloadType == EInterchangeAnimationPayLoadType::MORPHTARGETCURVE)
+			{
+				// Morph target curve queries.
+				FAnimationPayloadData AnimationPayLoadData{PayloadQuery.SceneNodeUniqueID, PayloadQuery.PayloadKey};
+				if (GetMorphTargetAnimationCurvePayloadData(*Impl, PayloadQuery.PayloadKey.UniqueId, AnimationPayLoadData))
 				{
-					// Morph target curve queries.
-					FAnimationPayloadData AnimationPayLoadData{ PayloadQuery.SceneNodeUniqueID, PayloadQuery.PayloadKey };
-					if (GetMorphTargetAnimationCurvePayloadData(*Impl, PayloadQuery.PayloadKey.UniqueId, AnimationPayLoadData))
-					{
-						CurveAnimationPayloads[PayloadIndex].Emplace(AnimationPayLoadData);
-					}
+					CurveAnimationPayloads[PayloadIndex].Emplace(AnimationPayLoadData);
 				}
-			};
+			}
+		};
 
-		//Get all curves with a parallel for if there is many
+		// Get all curves with a parallel for if there is many
 		int32 CurvePayloadCount = CurveQueryIndexes.Num();
 		CurveAnimationPayloads.AddDefaulted(CurvePayloadCount);
 		const int32 BatchSize = 10;
 		if (CurvePayloadCount > BatchSize)
 		{
 			const int32 NumBatches = (CurvePayloadCount / BatchSize) + 1;
-			ParallelFor(NumBatches, [&](int32 BatchIndex)
+			ParallelFor(
+				NumBatches,
+				[&CurveQueryIndexes, &GetAnimPayloadLambda](int32 BatchIndex)
 				{
 					int32 PayloadIndexOffset = BatchIndex * BatchSize;
 					for (int32 PayloadIndex = PayloadIndexOffset; PayloadIndex < PayloadIndexOffset + BatchSize; ++PayloadIndex)
 					{
-						//The last batch can be incomplete
+						// The last batch can be incomplete
 						if (!CurveQueryIndexes.IsValidIndex(PayloadIndex))
 						{
 							break;
 						}
 						GetAnimPayloadLambda(CurveQueryIndexes[PayloadIndex]);
 					}
-				}, EParallelForFlags::BackgroundPriority);// ParallelFor
+				},
+				EParallelForFlags::BackgroundPriority	 // ParallelFor
+			);
 		}
 		else
 		{
