USD: Disable importing into level via USD Interchange.

diff --git a/Engine/Plugins/Interchange/Runtime/Source/Import/Private/Usd/InterchangeUsdTranslator.cpp b/Engine/Plugins/Interchange/Runtime/Source/Import/Private/Usd/InterchangeUsdTranslator.cpp
--- a/Engine/Plugins/Interchange/Runtime/Source/Import/Private/Usd/InterchangeUsdTranslator.cpp
+++ b/Engine/Plugins/Interchange/Runtime/Source/Import/Private/Usd/InterchangeUsdTranslator.cpp
@@ -66,11 +66,12 @@ static FAutoConsoleVariableRef CVarInterchangeEnableUSDImport(
 );
 
 static bool GInterchangeEnableUSDLevelImport = false;
-static FAutoConsoleVariableRef CVarInterchangeEnableUSDLevelImport(
-	TEXT("Interchange.FeatureFlags.Import.USD.ToLevel"),
-	GInterchangeEnableUSDLevelImport,
-	TEXT("Whether support for USD level import is enabled.")
-);
+// Import into level via USD Interchange is disabled for 5.5 as it's still a work in progress
+// static FAutoConsoleVariableRef CVarInterchangeEnableUSDLevelImport(
+// 	TEXT("Interchange.FeatureFlags.Import.USD.ToLevel"),
+// 	GInterchangeEnableUSDLevelImport,
+// 	TEXT("Whether support for USD level import is enabled.")
+// );
 
 namespace UE::InterchangeUsdTranslator::Private
 {
