USD: Fix missing skeleton when Skeleton prim is named after one of its own bones.

diff --git a/Engine/Plugins/Interchange/Runtime/Source/Import/Private/Usd/InterchangeUsdTranslator.cpp b/Engine/Plugins/Interchange/Runtime/Source/Import/Private/Usd/InterchangeUsdTranslator.cpp
--- a/Engine/Plugins/Interchange/Runtime/Source/Import/Private/Usd/InterchangeUsdTranslator.cpp
+++ b/Engine/Plugins/Interchange/Runtime/Source/Import/Private/Usd/InterchangeUsdTranslator.cpp
@@ -7,6 +7,7 @@
 #include "USDGeomMeshConversion.h"
 #include "USDLightConversion.h"
 #include "USDLog.h"
+#include "USDObjectUtils.h"
 #include "USDPrimConversion.h"
 #include "USDShadeConversion.h"
 #include "USDSkeletalDataConversion.h"
@@ -1652,10 +1653,13 @@ namespace UE::InterchangeUsdTranslator::Private
 			}
 		};
 
-		// Start traversing from the root bones (we may have more than one)
+		// Start traversing from the root bones (we may have more than one, so check them all)
+		TSet<FString> UsedBoneNames;
 		for (int32 BoneIndex = 0; BoneIndex < ConvertedData.Bones.Num(); ++BoneIndex)
 		{
 			const UsdToUnreal::FUsdSkeletonData::FBone& Bone = ConvertedData.Bones[BoneIndex];
+			UsedBoneNames.Add(Bone.Name);
+
 			if (Bone.ParentIndex == INDEX_NONE)
 			{
 				const FString BonePathRoot = TEXT("");
@@ -1663,6 +1667,18 @@ namespace UE::InterchangeUsdTranslator::Private
 			}
 		}
 
+		// Interchange will abort parsing skeletons that don't have unique names for each bone. If the user has that
+		// on their actual skeleton, then that's just invalid data and we can just let it fail and emit the error message.
+		// However, we don't want to end up with duplicate bone names and fail to parse when the duplicate "bone" is due to
+		// how we actually use the Skeleton prim itself as the root, as that's our little "trick". In this case, here we
+		// just change the display text of the skeleton prim itself to be unique (which is used for the bone name)
+		const FString SkeletonPrimName = SkeletonPrimNode.GetDisplayLabel();
+		FString NewSkeletonPrimName = UsdUnreal::ObjectUtils::GetUniqueName(SkeletonPrimName, UsedBoneNames);
+		if (NewSkeletonPrimName != SkeletonPrimName)
+		{
+			SkeletonPrimNode.SetDisplayLabel(NewSkeletonPrimName);
+		}
+
 		// Handle SkelAnimation prims, if we have any bound for this Skeleton
 		AddSkeletalAnimationNode(SkelQuery, BoneToUidAndBoneIndex, TranslatorImpl, SkeletonPrimNode, NodeContainer, Info);
 
