USD: Fix materials exported from UE importing via USD interchange as black plastic without any warning or feedback.

diff --git a/Engine/Plugins/Interchange/Runtime/Source/Import/Private/Usd/InterchangeUsdTranslator.cpp b/Engine/Plugins/Interchange/Runtime/Source/Import/Private/Usd/InterchangeUsdTranslator.cpp
--- a/Engine/Plugins/Interchange/Runtime/Source/Import/Private/Usd/InterchangeUsdTranslator.cpp
+++ b/Engine/Plugins/Interchange/Runtime/Source/Import/Private/Usd/InterchangeUsdTranslator.cpp
@@ -790,6 +790,20 @@ namespace UE::InterchangeUsdTranslator::Private
 		}
 #endif	  // WITH_EDITOR
 
+		const static FName UnrealName = *UsdToUnreal::ConvertToken(UnrealIdentifiers::Unreal);
+		if (RenderContext == UnrealName)
+		{
+			UE_LOG(
+				LogUsd,
+				Warning,
+				TEXT(
+					"The 'unreal' render context is not yet supported via USD Interchange: The material '%s' will use the universal render context instead"
+				),
+				*PrimPath
+			);
+			RenderContext = FString{};
+		}
+
 		UInterchangeMaterialInstanceNode* MaterialNode = NewObject<UInterchangeMaterialInstanceNode>(&NodeContainer);
 		MaterialNode->InitializeNode(MaterialUid, MaterialPrimName, EInterchangeNodeContainerType::TranslatedAsset);
 		MaterialNode->SetAssetName(MaterialPrimName);
@@ -2321,7 +2335,7 @@ namespace UE::InterchangeUsdTranslator::Private
 
 UInterchangeUsdTranslatorSettings::UInterchangeUsdTranslatorSettings()
 	: GeometryPurpose((int32)(EUsdPurpose::Default | EUsdPurpose::Proxy | EUsdPurpose::Render | EUsdPurpose::Guide))
-	, RenderContext(TEXT("unreal"))	   // The proper definition of this is on the USDSchemas module, which we can't depend on
+	, RenderContext(NAME_None)	  // Default to the universal render context for now as we don't support 'unreal' yet anyway
 	, MaterialPurpose(*UnrealIdentifiers::MaterialPreviewPurpose)
 	, InterpolationType(EUsdInterpolationType::Linear)
 	, bOverrideStageOptions(false)
