Interchange: animation only import issue. The root of the animation is now adjust to fit the specified skeleton hierarchy instead of importing nothing if there was a difference.
#rnx

diff --git a/Engine/Plugins/Interchange/Runtime/Source/Import/Private/Animation/InterchangeAnimSequenceFactory.cpp b/Engine/Plugins/Interchange/Runtime/Source/Import/Private/Animation/InterchangeAnimSequenceFactory.cpp
--- a/Engine/Plugins/Interchange/Runtime/Source/Import/Private/Animation/InterchangeAnimSequenceFactory.cpp
+++ b/Engine/Plugins/Interchange/Runtime/Source/Import/Private/Animation/InterchangeAnimSequenceFactory.cpp
@@ -521,22 +521,16 @@ namespace UE::Interchange::Private
 				for (int32 BakeIndex = 0; BakeIndex < BakeKeyCountForAnimationPayload; BakeIndex++, CurrentTime += BakeInterval)
 				{
 					FTransform3f AnimKeyTransform = FTransform3f(AnimationTransformPayload.Transforms[BakeIndex]);
-					if (bApplyGlobalOffset)
+					if (bApplyGlobalOffset && bBakeMeshes)
 					{
-						if (bBakeMeshes)
+ 						if (const UInterchangeSceneNode* RootJointNode = Cast<UInterchangeSceneNode>(NodeContainer->GetNode(SkeletonRootUid)))
 						{
-							const UInterchangeSceneNode* RootJointNode = Cast<UInterchangeSceneNode>(NodeContainer->GetNode(SkeletonRootUid));
-							if (RootJointNode)
+							FString RootJointParentNodeUid = RootJointNode->GetParentUid();
+							if (const UInterchangeSceneNode* RootJointParentNode = Cast<UInterchangeSceneNode>(NodeContainer->GetNode(RootJointParentNodeUid)))
 							{
-								FString RootJointParentNodeUid = RootJointNode->GetParentUid();
-
-								const UInterchangeSceneNode* RootJointParentNode = Cast<UInterchangeSceneNode>(NodeContainer->GetNode(RootJointParentNodeUid));
-								if (RootJointParentNode)
-								{
-									FTransform GlobalTransform;
-									RootJointParentNode->GetCustomGlobalTransform(NodeContainer, GlobalOffsetTransform, GlobalTransform);
-									AnimKeyTransform = AnimKeyTransform * FTransform3f(GlobalTransform);
-								}
+								FTransform GlobalTransform;
+								RootJointParentNode->GetCustomGlobalTransform(NodeContainer, GlobalOffsetTransform, GlobalTransform);
+								AnimKeyTransform = AnimKeyTransform * FTransform3f(GlobalTransform);
 							}
 						}
 					}
diff --git a/Engine/Plugins/Interchange/Runtime/Source/Pipelines/Private/InterchangeGenericAnimationPipeline.cpp b/Engine/Plugins/Interchange/Runtime/Source/Pipelines/Private/InterchangeGenericAnimationPipeline.cpp
--- a/Engine/Plugins/Interchange/Runtime/Source/Pipelines/Private/InterchangeGenericAnimationPipeline.cpp
+++ b/Engine/Plugins/Interchange/Runtime/Source/Pipelines/Private/InterchangeGenericAnimationPipeline.cpp
@@ -517,30 +517,30 @@ void UInterchangeGenericAnimationPipeline::CreateAnimSequenceFactoryNode(UInterc
 	//base on the specified skeleton
 	if(bImportOnlyAnimation && !SkeletonFactoryNode && CommonSkeletalMeshesAndAnimationsProperties->Skeleton.IsValid())
 	{
-		const FReferenceSkeleton& ReferenceSkeleton = CommonSkeletalMeshesAndAnimationsProperties->Skeleton->GetReferenceSkeleton();
-		TArray<FString> SkeletonRootNodeUids;
-		BaseNodeContainer->IterateNodesOfType<UInterchangeSceneNode>([&SkeletonRootNodeUids, BaseNodeContainerClosure = BaseNodeContainer, &ReferenceSkeleton](const FString& NodeUid, UInterchangeSceneNode* Node)
+		TWeakObjectPtr<USkeleton> Skeleton = CommonSkeletalMeshesAndAnimationsProperties->Skeleton;
+		TPair<int32, FString> SkeletonRootNodeUidAndBoneIndex = TPair<int32, FString>(INDEX_NONE, FString());
+		BaseNodeContainer->IterateNodesOfType<UInterchangeSceneNode>([&SkeletonRootNodeUidAndBoneIndex, BaseNodeContainerClosure = BaseNodeContainer, Skeleton](const FString& NodeUid, UInterchangeSceneNode* Node)
 		{
-			if (Node->IsSpecializedTypeContains(UE::Interchange::FSceneNodeStaticData::GetJointSpecializeTypeString()))
+			if (Skeleton.IsValid() && Node->IsSpecializedTypeContains(UE::Interchange::FSceneNodeStaticData::GetJointSpecializeTypeString()))
 			{
-				if(ReferenceSkeleton.FindBoneIndex(FName(*Node->GetDisplayLabel())) != INDEX_NONE)
+				const FReferenceSkeleton& ReferenceSkeleton = Skeleton->GetReferenceSkeleton();
+				int32 RefBoneIndex = ReferenceSkeleton.FindBoneIndex(FName(*Node->GetDisplayLabel()));
+				if(RefBoneIndex != INDEX_NONE)
 				{
-					const FString ParentUid = Node->GetParentUid();
-					if (const UInterchangeSceneNode* ParentNode = Cast<UInterchangeSceneNode>(BaseNodeContainerClosure->GetNode(ParentUid)))
+					if (SkeletonRootNodeUidAndBoneIndex.Key == INDEX_NONE || RefBoneIndex < SkeletonRootNodeUidAndBoneIndex.Key)
 					{
-						if (!ParentNode->IsSpecializedTypeContains(UE::Interchange::FSceneNodeStaticData::GetJointSpecializeTypeString()))
-						{
-							SkeletonRootNodeUids.Add(NodeUid);
-						}
+						SkeletonRootNodeUidAndBoneIndex = TPair<int32, FString>(RefBoneIndex, NodeUid);
 					}
 				}
 			}
 		});
 		FString SkeletonRootUid;
-		if (SkeletonRootNodeUids.Num() > 0)
+		//Use the lower uid we found
+		if (SkeletonRootNodeUidAndBoneIndex.Key != INDEX_NONE && !SkeletonRootNodeUidAndBoneIndex.Value.IsEmpty())
 		{
-			SkeletonRootUid = SkeletonRootNodeUids[0];
+			SkeletonRootUid = SkeletonRootNodeUidAndBoneIndex.Value;
 		}
+
 		if(!SkeletonRootUid.IsEmpty())
 		{
 			//Create a skeleton node from all the joint in the translated nodes
