[Interchange] Runtime/PIE SkeletalMesh re-import fixes.

#rnx

diff --git a/Engine/Plugins/Interchange/Editor/Source/Utilities/Private/InterchangeEditorUtilities.cpp b/Engine/Plugins/Interchange/Editor/Source/Utilities/Private/InterchangeEditorUtilities.cpp
--- a/Engine/Plugins/Interchange/Editor/Source/Utilities/Private/InterchangeEditorUtilities.cpp
+++ b/Engine/Plugins/Interchange/Editor/Source/Utilities/Private/InterchangeEditorUtilities.cpp
@@ -4,6 +4,13 @@
 #include "FileHelpers.h"
 #include "UObject/Package.h"
 
+#if WITH_EDITOR
+#include "CoreMinimal.h"
+#include "Editor/EditorEngine.h"
+
+extern UNREALED_API UEditorEngine* GEditor;
+#endif
+
 #include UE_INLINE_GENERATED_CPP_BY_NAME(InterchangeEditorUtilities)
 
 bool UInterchangeEditorUtilities::SaveAsset(UObject* Asset)
@@ -13,3 +20,12 @@ bool UInterchangeEditorUtilities::SaveAsset(UObject* Asset)
 	FEditorFileUtils::EPromptReturnCode ReturnCode = FEditorFileUtils::PromptForCheckoutAndSave(PackagesToSave, false /*bCheckDirty*/, false /*bPromptToSave*/);
 	return (ReturnCode == FEditorFileUtils::PR_Success);
 }
+
+bool UInterchangeEditorUtilities::IsRuntimeOrPIE()
+{
+#if WITH_EDITOR
+	return (GEditor && GEditor->PlayWorld) || GIsPlayInEditorWorld || IsRunningGame();
+#else
+	return true;
+#endif // WITH_EDITOR
+}
\ No newline at end of file
diff --git a/Engine/Plugins/Interchange/Editor/Source/Utilities/Public/InterchangeEditorUtilities.h b/Engine/Plugins/Interchange/Editor/Source/Utilities/Public/InterchangeEditorUtilities.h
--- a/Engine/Plugins/Interchange/Editor/Source/Utilities/Public/InterchangeEditorUtilities.h
+++ b/Engine/Plugins/Interchange/Editor/Source/Utilities/Public/InterchangeEditorUtilities.h
@@ -16,4 +16,6 @@ public:
 protected:
 
 	virtual bool SaveAsset(UObject* Asset) override;
+
+	virtual bool IsRuntimeOrPIE() override;
 };
diff --git a/Engine/Plugins/Interchange/Runtime/Source/FactoryNodes/Public/InterchangeAnimSequenceFactoryNode.h b/Engine/Plugins/Interchange/Runtime/Source/FactoryNodes/Public/InterchangeAnimSequenceFactoryNode.h
--- a/Engine/Plugins/Interchange/Runtime/Source/FactoryNodes/Public/InterchangeAnimSequenceFactoryNode.h
+++ b/Engine/Plugins/Interchange/Runtime/Source/FactoryNodes/Public/InterchangeAnimSequenceFactoryNode.h
@@ -407,6 +407,12 @@ public:
 	UFUNCTION(BlueprintCallable, Category = "Interchange | Node | SkeletalAnimationTrack")
 	void SetAnimationPayloadKeysForMorphTargetNodeUids(const TMap<FString, FString>& MorphTargetAnimationPayloadKeyUids, const TMap<FString, uint8>& MorphTargetAnimationPayloadKeyTypes);
 
+	/** Return if the import of the class is allowed at runtime.*/
+	virtual bool IsRuntimeImportAllowed() const override
+	{
+		return false;
+	}
+
 private:
 	const UE::Interchange::FAttributeKey Macro_CustomSkeletonFactoryNodeUidKey = UE::Interchange::FAttributeKey(TEXT("SkeletonFactoryNodeUid"));
 
diff --git a/Engine/Plugins/Interchange/Runtime/Source/FactoryNodes/Public/InterchangePhysicsAssetFactoryNode.h b/Engine/Plugins/Interchange/Runtime/Source/FactoryNodes/Public/InterchangePhysicsAssetFactoryNode.h
--- a/Engine/Plugins/Interchange/Runtime/Source/FactoryNodes/Public/InterchangePhysicsAssetFactoryNode.h
+++ b/Engine/Plugins/Interchange/Runtime/Source/FactoryNodes/Public/InterchangePhysicsAssetFactoryNode.h
@@ -100,6 +100,12 @@ public:
 		IMPLEMENT_NODE_ATTRIBUTE_SETTER_NODELEGATE(SkeletalMeshUid, FString)
 	}
 
+	/** Return if the import of the class is allowed at runtime.*/
+	virtual bool IsRuntimeImportAllowed() const override
+	{
+		return false;
+	}
+
 private:
 
 	void FillAssetClassFromAttribute()
diff --git a/Engine/Plugins/Interchange/Runtime/Source/FactoryNodes/Public/InterchangeSceneVariantSetsFactoryNode.h b/Engine/Plugins/Interchange/Runtime/Source/FactoryNodes/Public/InterchangeSceneVariantSetsFactoryNode.h
--- a/Engine/Plugins/Interchange/Runtime/Source/FactoryNodes/Public/InterchangeSceneVariantSetsFactoryNode.h
+++ b/Engine/Plugins/Interchange/Runtime/Source/FactoryNodes/Public/InterchangeSceneVariantSetsFactoryNode.h
@@ -63,6 +63,12 @@ public:
 	UFUNCTION(BlueprintCallable, Category = "Interchange | Node | SceneVariantSetFactory")
 	bool RemoveCustomVariantSetUid(const FString& VariantUid);
 
+	/** Return if the import of the class is allowed at runtime.*/
+	virtual bool IsRuntimeImportAllowed() const override
+	{
+		return false;
+	}
+
 private:
 	UE::Interchange::TArrayAttributeHelper<FString> CustomVariantSetUids;
 };
diff --git a/Engine/Plugins/Interchange/Runtime/Source/FactoryNodes/Public/InterchangeSkeletalMeshFactoryNode.h b/Engine/Plugins/Interchange/Runtime/Source/FactoryNodes/Public/InterchangeSkeletalMeshFactoryNode.h
--- a/Engine/Plugins/Interchange/Runtime/Source/FactoryNodes/Public/InterchangeSkeletalMeshFactoryNode.h
+++ b/Engine/Plugins/Interchange/Runtime/Source/FactoryNodes/Public/InterchangeSkeletalMeshFactoryNode.h
@@ -168,6 +168,12 @@ public:
 	virtual void AppendAssetRegistryTags(TArray<FAssetRegistryTag>& OutTags) const override;
 
 	virtual void CopyWithObject(const UInterchangeFactoryBaseNode* SourceNode, UObject* Object) override;
+
+	/** Return if the import of the class is allowed at runtime.*/
+	virtual bool IsRuntimeImportAllowed () const override
+	{
+		return false;
+	}
 private:
 
 	virtual void FillAssetClassFromAttribute() override;
diff --git a/Engine/Plugins/Interchange/Runtime/Source/FactoryNodes/Public/InterchangeSkeletalMeshLodDataNode.h b/Engine/Plugins/Interchange/Runtime/Source/FactoryNodes/Public/InterchangeSkeletalMeshLodDataNode.h
--- a/Engine/Plugins/Interchange/Runtime/Source/FactoryNodes/Public/InterchangeSkeletalMeshLodDataNode.h
+++ b/Engine/Plugins/Interchange/Runtime/Source/FactoryNodes/Public/InterchangeSkeletalMeshLodDataNode.h
@@ -69,6 +69,12 @@ public:
 	UFUNCTION(BlueprintCallable, Category = "Interchange | Node | SkeletalMeshLodData")
 	bool RemoveAllMeshes();
 
+	/** Return if the import of the class is allowed at runtime.*/
+	virtual bool IsRuntimeImportAllowed() const override
+	{
+		return false;
+	}
+
 private:
 
 	bool IsEditorOnlyDataDefined();
diff --git a/Engine/Plugins/Interchange/Runtime/Source/Import/Private/Animation/InterchangeAnimSequenceFactory.cpp b/Engine/Plugins/Interchange/Runtime/Source/Import/Private/Animation/InterchangeAnimSequenceFactory.cpp
--- a/Engine/Plugins/Interchange/Runtime/Source/Import/Private/Animation/InterchangeAnimSequenceFactory.cpp
+++ b/Engine/Plugins/Interchange/Runtime/Source/Import/Private/Animation/InterchangeAnimSequenceFactory.cpp
@@ -1024,11 +1024,7 @@ UInterchangeFactoryBase::FImportAssetResult UInterchangeAnimSequenceFactory::Beg
 	TRACE_CPUPROFILER_EVENT_SCOPE(UInterchangeAnimSequenceFactory::BeginImportAsset_GameThread);
 
 	UInterchangeFactoryBase::FImportAssetResult ImportAssetResult;
-#if !WITH_EDITOR || !WITH_EDITORONLY_DATA
-
-	UE_LOG(LogInterchangeImport, Error, TEXT("Cannot import AnimSequence asset at runtime. This is an editor-only feature."));
-	return ImportAssetResult;
-#else
+#if WITH_EDITOR && WITH_EDITORONLY_DATA
 	UAnimSequence* NewAnimSequence = nullptr;
 	if (!Arguments.AssetNode || !Arguments.AssetNode->GetObjectClass()->IsChildOf(GetFactoryClass()))
 	{
@@ -1142,9 +1138,9 @@ UInterchangeFactoryBase::FImportAssetResult UInterchangeAnimSequenceFactory::Beg
 
 	AnimSequence = NewAnimSequence;
 	ImportAssetResult.ImportedObject = NewAnimSequence;
+#endif
 
 	return ImportAssetResult;
-#endif //else !WITH_EDITOR || !WITH_EDITORONLY_DATA
 }
 
 UInterchangeFactoryBase::FImportAssetResult UInterchangeAnimSequenceFactory::ImportAsset_Async(const FImportAssetObjectParams& Arguments)
@@ -1152,12 +1148,7 @@ UInterchangeFactoryBase::FImportAssetResult UInterchangeAnimSequenceFactory::Imp
 	TRACE_CPUPROFILER_EVENT_SCOPE(UInterchangeAnimSequenceFactory::ImportAsset_Async);
 
 	UInterchangeFactoryBase::FImportAssetResult ImportAssetResult;
-#if !WITH_EDITOR || !WITH_EDITORONLY_DATA
-
-	UE_LOG(LogInterchangeImport, Error, TEXT("Cannot import AnimSequence asset at runtime. This is an editor-only feature."));
-	return ImportAssetResult;
-#else
-
+#if WITH_EDITOR && WITH_EDITORONLY_DATA
 	//The game thread part should have verified all the data, so no need to do extra log
 	if (!AnimSequence)
 	{
@@ -1311,9 +1302,9 @@ UInterchangeFactoryBase::FImportAssetResult UInterchangeAnimSequenceFactory::Imp
 
 	//Getting the file Hash will cache it into the source data
 	Arguments.SourceData->GetFileContentHash();
+#endif
 
 	return ImportAssetResult;
-#endif //else !WITH_EDITOR || !WITH_EDITORONLY_DATA
 }
 
 UInterchangeFactoryBase::FImportAssetResult UInterchangeAnimSequenceFactory::EndImportAsset_GameThread(const FImportAssetObjectParams& Arguments)
@@ -1321,11 +1312,7 @@ UInterchangeFactoryBase::FImportAssetResult UInterchangeAnimSequenceFactory::End
 	TRACE_CPUPROFILER_EVENT_SCOPE(UInterchangeAnimSequenceFactory::ImportAsset_Async);
 
 	UInterchangeFactoryBase::FImportAssetResult ImportAssetResult;
-#if !WITH_EDITOR || !WITH_EDITORONLY_DATA
-
-	UE_LOG(LogInterchangeImport, Error, TEXT("Cannot import AnimSequence asset at runtime. This is an editor-only feature."));
-	return ImportAssetResult;
-#else
+#if WITH_EDITOR && WITH_EDITORONLY_DATA
 	if (!AnimSequence)
 	{
 		return ImportAssetResult;
@@ -1465,9 +1452,9 @@ UInterchangeFactoryBase::FImportAssetResult UInterchangeAnimSequenceFactory::End
 		CurrentNode->FillAllCustomAttributeFromObject(AnimSequence);
 		UE::Interchange::FFactoryCommon::ApplyReimportStrategyToAsset(AnimSequence, PreviousNode, CurrentNode, AnimSequenceFactoryNode);
 	}
+#endif
 
 	return ImportAssetResult;
-#endif //else !WITH_EDITOR || !WITH_EDITORONLY_DATA
 }
 
 /* This function is call in the completion task on the main thread, use it to call main thread post creation step for your assets*/
diff --git a/Engine/Plugins/Interchange/Runtime/Source/Import/Private/Gltf/InterchangeGltfTranslator.cpp b/Engine/Plugins/Interchange/Runtime/Source/Import/Private/Gltf/InterchangeGltfTranslator.cpp
--- a/Engine/Plugins/Interchange/Runtime/Source/Import/Private/Gltf/InterchangeGltfTranslator.cpp
+++ b/Engine/Plugins/Interchange/Runtime/Source/Import/Private/Gltf/InterchangeGltfTranslator.cpp
@@ -792,9 +792,9 @@ bool UInterchangeGLTFTranslator::Translate( UInterchangeBaseNodeContainer& NodeC
 
 	// Variants
 	// Note: Variants are not supported yet in game play mode
-	if ( !FApp::IsGame() && bHasVariants )
+	if (bHasVariants)
 	{
-		HandleGltfVariants( NodeContainer, FileName );
+		HandleGltfVariants(NodeContainer, FileName);
 	}
 
 	// Add glTF errors and warnings to the Interchange results
diff --git a/Engine/Plugins/Interchange/Runtime/Source/Import/Private/Mesh/InterchangePhysicsAssetFactory.cpp b/Engine/Plugins/Interchange/Runtime/Source/Import/Private/Mesh/InterchangePhysicsAssetFactory.cpp
--- a/Engine/Plugins/Interchange/Runtime/Source/Import/Private/Mesh/InterchangePhysicsAssetFactory.cpp
+++ b/Engine/Plugins/Interchange/Runtime/Source/Import/Private/Mesh/InterchangePhysicsAssetFactory.cpp
@@ -73,13 +73,7 @@ UInterchangeFactoryBase::FImportAssetResult UInterchangePhysicsAssetFactory::Imp
 {
 	TRACE_CPUPROFILER_EVENT_SCOPE(UInterchangePhysicsAssetFactory::ImportAsset_Async);
 	FImportAssetResult ImportAssetResult;
-#if !WITH_EDITORONLY_DATA
-
-	UE_LOG(LogInterchangeImport, Error, TEXT("Cannot import PhysicsAsset asset at runtime. This is an editor-only feature."));
-	return ImportAssetResult;
-
-#else
-
+#if WITH_EDITORONLY_DATA
 	if (!Arguments.AssetNode || !Arguments.AssetNode->GetObjectClass()->IsChildOf(GetFactoryClass()))
 	{
 		return ImportAssetResult;
@@ -117,7 +111,7 @@ UInterchangeFactoryBase::FImportAssetResult UInterchangePhysicsAssetFactory::Imp
 	Arguments.SourceData->GetFileContentHash();
 
 	ImportAssetResult.ImportedObject = PhysicsAssetObject;
-	return ImportAssetResult;
 #endif
+	return ImportAssetResult;
 }
 
diff --git a/Engine/Plugins/Interchange/Runtime/Source/Import/Private/Mesh/InterchangeSkeletalMeshFactory.cpp b/Engine/Plugins/Interchange/Runtime/Source/Import/Private/Mesh/InterchangeSkeletalMeshFactory.cpp
--- a/Engine/Plugins/Interchange/Runtime/Source/Import/Private/Mesh/InterchangeSkeletalMeshFactory.cpp
+++ b/Engine/Plugins/Interchange/Runtime/Source/Import/Private/Mesh/InterchangeSkeletalMeshFactory.cpp
@@ -1032,12 +1032,7 @@ void UInterchangeSkeletalMeshFactory::CreatePayloadTasks(const FImportAssetObjec
 {
 	TRACE_CPUPROFILER_EVENT_SCOPE(UInterchangeSkeletalMeshFactory::CreateAsset)
 
-#if !WITH_EDITOR || !WITH_EDITORONLY_DATA
-
-	UE_LOG(LogInterchangeImport, Error, TEXT("Cannot import SkeletalMesh asset at runtime. This is an editor-only feature."));
-	return;
-
-#else
+#if WITH_EDITOR && WITH_EDITORONLY_DATA
 	using namespace UE::Interchange;
 
 	if (!Arguments.AssetNode || !Arguments.AssetNode->GetObjectClass()->IsChildOf(GetFactoryClass()))
@@ -1293,12 +1288,7 @@ UInterchangeFactoryBase::FImportAssetResult UInterchangeSkeletalMeshFactory::Beg
 
 	FImportAssetResult ImportAssetResult;
 
-#if !WITH_EDITOR || !WITH_EDITORONLY_DATA
-
-	UE_LOG(LogInterchangeImport, Error, TEXT("Cannot import SkeletalMesh asset at runtime. This is an editor-only feature."));
-	return ImportAssetResult;
-
-#else
+#if WITH_EDITOR && WITH_EDITORONLY_DATA
 	USkeletalMesh* SkeletalMesh = nullptr;
 	if (!Arguments.AssetNode || !Arguments.AssetNode->GetObjectClass()->IsChildOf(GetFactoryClass()))
 	{
@@ -1538,10 +1528,9 @@ UInterchangeFactoryBase::FImportAssetResult UInterchangeSkeletalMeshFactory::Beg
 
 		CurrentLodIndex++;
 	}
+#endif
 
 	return ImportAssetResult;
-
-#endif //else !WITH_EDITOR || !WITH_EDITORONLY_DATA
 }
 
 UInterchangeFactoryBase::FImportAssetResult UInterchangeSkeletalMeshFactory::ImportAsset_Async(const FImportAssetObjectParams& Arguments)
@@ -1550,12 +1539,7 @@ UInterchangeFactoryBase::FImportAssetResult UInterchangeSkeletalMeshFactory::Imp
 
 	FImportAssetResult ImportAssetResult;
 
-#if !WITH_EDITOR || !WITH_EDITORONLY_DATA
-
-	UE_LOG(LogInterchangeImport, Error, TEXT("Cannot import SkeletalMesh asset at runtime. This is an editor-only feature."));
-	return ImportAssetResult;
-
-#else
+#if WITH_EDITOR && WITH_EDITORONLY_DATA
 	if (!Arguments.AssetNode || !Arguments.AssetNode->GetObjectClass()->IsChildOf(GetFactoryClass()))
 	{
 		return ImportAssetResult;
@@ -1998,9 +1982,9 @@ UInterchangeFactoryBase::FImportAssetResult UInterchangeSkeletalMeshFactory::Imp
 	Arguments.SourceData->GetFileContentHash();
 
 	ImportAssetResult.ImportedObject = SkeletalMeshObject;
-	return ImportAssetResult;
+#endif
 
-#endif //else !WITH_EDITOR || !WITH_EDITORONLY_DATA
+	return ImportAssetResult;
 }
 
 UInterchangeFactoryBase::FImportAssetResult UInterchangeSkeletalMeshFactory::EndImportAsset_GameThread(const FImportAssetObjectParams& Arguments)
@@ -2010,12 +1994,7 @@ UInterchangeFactoryBase::FImportAssetResult UInterchangeSkeletalMeshFactory::End
 	check(IsInGameThread());
 	FImportAssetResult ImportAssetResult;
 
-#if !WITH_EDITOR || !WITH_EDITORONLY_DATA
-
-	UE_LOG(LogInterchangeImport, Error, TEXT("Cannot import SkeletalMesh asset at runtime. This is an editor-only feature."));
-	return ImportAssetResult;
-
-#else
+#if WITH_EDITOR && WITH_EDITORONLY_DATA
 	if (!Arguments.AssetNode || !Arguments.AssetNode->GetObjectClass()->IsChildOf(GetFactoryClass()))
 	{
 		return ImportAssetResult;
@@ -2273,9 +2252,9 @@ UInterchangeFactoryBase::FImportAssetResult UInterchangeSkeletalMeshFactory::End
 	SkeletalMeshFactoryNode->SetCustomReferenceObject(FSoftObjectPath(SkeletalMesh));
 
 	ImportAssetResult.ImportedObject = SkeletalMesh;
-	return ImportAssetResult;
+#endif
 
-#endif //WITH_EDITOR
+	return ImportAssetResult;
 }
 
 void UInterchangeSkeletalMeshFactory::Cancel()
diff --git a/Engine/Plugins/Interchange/Runtime/Source/Import/Private/Mesh/InterchangeStaticMeshFactory.cpp b/Engine/Plugins/Interchange/Runtime/Source/Import/Private/Mesh/InterchangeStaticMeshFactory.cpp
--- a/Engine/Plugins/Interchange/Runtime/Source/Import/Private/Mesh/InterchangeStaticMeshFactory.cpp
+++ b/Engine/Plugins/Interchange/Runtime/Source/Import/Private/Mesh/InterchangeStaticMeshFactory.cpp
@@ -39,6 +39,7 @@
 #include "StaticMeshAttributes.h"
 #include "StaticMeshCompiler.h"
 #include "StaticMeshOperations.h"
+#include "InterchangeManager.h"
 
 #include UE_INLINE_GENERATED_CPP_BY_NAME(InterchangeStaticMeshFactory)
 
@@ -308,9 +309,13 @@ UInterchangeFactoryBase::FImportAssetResult UInterchangeStaticMeshFactory::Begin
 {
 	TRACE_CPUPROFILER_EVENT_SCOPE(UInterchangeStaticMeshFactory::BeginImportAsset_GameThread);
 
-	//We must ensure we use the same settings until the import is finish, FApp::IsGame() can return a different
+	//We must ensure we use the same settings until the import is finish, EditorUtilities->IsRuntimeOrPIE() can return a different
 	//value during an asynchronous import
-	ImportAssetObjectData.bIsAppGame = FApp::IsGame();
+	ImportAssetObjectData.bIsAppGame = false;
+	if (UInterchangeEditorUtilitiesBase* EditorUtilities = UInterchangeManager::GetInterchangeManager().GetEditorUtilities())
+	{
+		ImportAssetObjectData.bIsAppGame = EditorUtilities->IsRuntimeOrPIE();
+	}
 
 	FImportAssetResult ImportAssetResult;
 	UStaticMesh* StaticMesh = nullptr;
diff --git a/Engine/Plugins/Interchange/Runtime/Source/Import/Private/Scene/InterchangeLevelFactory.cpp b/Engine/Plugins/Interchange/Runtime/Source/Import/Private/Scene/InterchangeLevelFactory.cpp
--- a/Engine/Plugins/Interchange/Runtime/Source/Import/Private/Scene/InterchangeLevelFactory.cpp
+++ b/Engine/Plugins/Interchange/Runtime/Source/Import/Private/Scene/InterchangeLevelFactory.cpp
@@ -82,13 +82,18 @@ namespace UE::Interchange::Private::InterchangeLevelFactory
 
 			constexpr bool bEnableWorldPartitionStreaming = false;
 			constexpr bool bInformEngineOfWorld = true;
+			bool bIsRuntimeOrPIE = false;
+			if (UInterchangeEditorUtilitiesBase* EditorUtilities = UInterchangeManager::GetInterchangeManager().GetEditorUtilities())
+			{
+				bIsRuntimeOrPIE = EditorUtilities->IsRuntimeOrPIE();
+			}
 
 			// Those are the init values taken from the default in UWorld::CreateWorld + CreateWorldPartition.
 			UWorld::InitializationValues InitValues = UWorld::InitializationValues()
 				.ShouldSimulatePhysics(false)
 				.EnableTraceCollision(true)
-				.CreateNavigation(!FApp::IsGame())
-				.CreateAISystem(!FApp::IsGame())
+				.CreateNavigation(!bIsRuntimeOrPIE)
+				.CreateAISystem(!bIsRuntimeOrPIE)
 				.CreateWorldPartition(bCreateWorldPartition)
 				.EnableWorldPartitionStreaming(bEnableWorldPartitionStreaming);
 
@@ -138,7 +143,8 @@ UInterchangeFactoryBase::FImportAssetResult UInterchangeLevelFactory::BeginImpor
 	}
 
 #if WITH_EDITOR
-	if (!FApp::IsGame() && WorldAsset)
+	UInterchangeEditorUtilitiesBase* EditorUtilities = UInterchangeManager::GetInterchangeManager().GetEditorUtilities();
+	if ((!EditorUtilities || !EditorUtilities->IsRuntimeOrPIE()) && WorldAsset)
 	{
 		WorldAsset->PreEditChange(nullptr);
 	}
diff --git a/Engine/Plugins/Interchange/Runtime/Source/Import/Private/Scene/InterchangeSceneVariantSetsFactory.cpp b/Engine/Plugins/Interchange/Runtime/Source/Import/Private/Scene/InterchangeSceneVariantSetsFactory.cpp
--- a/Engine/Plugins/Interchange/Runtime/Source/Import/Private/Scene/InterchangeSceneVariantSetsFactory.cpp
+++ b/Engine/Plugins/Interchange/Runtime/Source/Import/Private/Scene/InterchangeSceneVariantSetsFactory.cpp
@@ -298,11 +298,7 @@ UInterchangeFactoryBase::FImportAssetResult UInterchangeSceneVariantSetsFactory:
 {
 	TRACE_CPUPROFILER_EVENT_SCOPE(UInterchangeSceneVariantSetsFactory::BeginImportAsset_GameThread);
 	FImportAssetResult ImportAssetResult;
-#if !WITH_EDITOR || !WITH_EDITORONLY_DATA
-
-	UE_LOG(LogInterchangeImport, Error, TEXT("Cannot import LevelSequence asset at runtime. This is an editor-only feature."));
-	return ImportAssetResult;
-#else
+#if WITH_EDITOR && WITH_EDITORONLY_DATA
 	auto CannotReimportMessage = [&Arguments, this]()
 	{
 		UInterchangeResultError_Generic* Message = AddMessage<UInterchangeResultError_Generic>();
@@ -368,19 +364,16 @@ UInterchangeFactoryBase::FImportAssetResult UInterchangeSceneVariantSetsFactory:
 	}
 
 	ImportAssetResult.ImportedObject = ImportObjectSourceData(Arguments, LevelVariantSets);
+#endif
+
 	return ImportAssetResult;
-#endif //else !WITH_EDITOR || !WITH_EDITORONLY_DATA
 }
 
 UObject* UInterchangeSceneVariantSetsFactory::ImportObjectSourceData(const FImportAssetObjectParams& Arguments, ULevelVariantSets* LevelVariantSets)
 {
 	TRACE_CPUPROFILER_EVENT_SCOPE(UInterchangeSceneVariantSetsFactory::ImportObjectSourceData);
-#if !WITH_EDITOR || !WITH_EDITORONLY_DATA
+#if WITH_EDITOR && WITH_EDITORONLY_DATA
 	// TODO: Can we import ULevelVariantSets at runtime
-	UE_LOG(LogInterchangeImport, Error, TEXT("Cannot import LevelVariantSets asset at runtime. This is an editor-only feature."));
-	return nullptr;
-
-#else
 	using namespace UE::Interchange;
 
 
@@ -414,10 +407,10 @@ UObject* UInterchangeSceneVariantSetsFactory::ImportObjectSourceData(const FImpo
 
 	//Getting the file Hash will cache it into the source data
 	Arguments.SourceData->GetFileContentHash();
+#endif
 
 	//The interchange completion task (call in the GameThread after the factories pass), will call PostEditChange which will trig another asynchronous system that will build all material in parallel
 	return LevelVariantSets;
-#endif //else !WITH_EDITOR || !WITH_EDITORONLY_DATA
 }
 
 void UInterchangeSceneVariantSetsFactory::SetupObject_GameThread(const FSetupObjectParams& Arguments)
diff --git a/Engine/Plugins/Interchange/Runtime/Source/Pipelines/Private/InterchangeGenericMaterialPipeline.cpp b/Engine/Plugins/Interchange/Runtime/Source/Pipelines/Private/InterchangeGenericMaterialPipeline.cpp
--- a/Engine/Plugins/Interchange/Runtime/Source/Pipelines/Private/InterchangeGenericMaterialPipeline.cpp
+++ b/Engine/Plugins/Interchange/Runtime/Source/Pipelines/Private/InterchangeGenericMaterialPipeline.cpp
@@ -61,6 +61,7 @@
 #include "Templates/SubclassOf.h"
 #include "UObject/Object.h"
 #include "UObject/ObjectMacros.h"
+#include "InterchangeManager.h"
 
 #if UE_BUILD_DEBUG
 #include "HAL/PlatformFileManager.h"
@@ -790,7 +791,8 @@ void UInterchangeGenericMaterialPipeline::ExecutePipeline(UInterchangeBaseNodeCo
 	}
 
 	// Can't import materials at runtime, fall back to instances
-	if (FApp::IsGame() && MaterialImport == EInterchangeMaterialImportOption::ImportAsMaterials)
+	UInterchangeEditorUtilitiesBase* EditorUtilities = UInterchangeManager::GetInterchangeManager().GetEditorUtilities();
+	if ((EditorUtilities && EditorUtilities->IsRuntimeOrPIE()) && MaterialImport == EInterchangeMaterialImportOption::ImportAsMaterials)
 	{
 		MaterialImport = EInterchangeMaterialImportOption::ImportAsMaterialInstances;
 	}
@@ -915,7 +917,7 @@ void UInterchangeGenericMaterialPipeline::ExecutePipeline(UInterchangeBaseNodeCo
 		MaterialFactoryNode->SetDisplayLabel(MaterialNode->GetAssetName());
 		MaterialFactoryNode->SetCustomParent(ParentPath);
 
-		const UClass* MaterialClass = FApp::IsGame() ? UMaterialInstanceDynamic::StaticClass() : UMaterialInstanceConstant::StaticClass();
+		const UClass* MaterialClass = (EditorUtilities && EditorUtilities->IsRuntimeOrPIE()) ? UMaterialInstanceDynamic::StaticClass() : UMaterialInstanceConstant::StaticClass();
 		MaterialFactoryNode->SetCustomInstanceClassName(MaterialClass->GetPathName());
 
 		TArray<FString> Inputs;
diff --git a/Engine/Plugins/Interchange/Runtime/Source/Pipelines/Private/InterchangeGenericSkeletalMeshPipeline.cpp b/Engine/Plugins/Interchange/Runtime/Source/Pipelines/Private/InterchangeGenericSkeletalMeshPipeline.cpp
--- a/Engine/Plugins/Interchange/Runtime/Source/Pipelines/Private/InterchangeGenericSkeletalMeshPipeline.cpp
+++ b/Engine/Plugins/Interchange/Runtime/Source/Pipelines/Private/InterchangeGenericSkeletalMeshPipeline.cpp
@@ -30,6 +30,7 @@
 #include "Tasks/Task.h"
 #include "UObject/Object.h"
 #include "UObject/ObjectMacros.h"
+#include "InterchangeManager.h"
 
 namespace UE::Interchange::SkeletalMeshGenericPipeline
 {
@@ -131,73 +132,65 @@ void UInterchangeGenericMeshPipeline::ExecutePreImportPipelineSkeletalMesh()
 		auto CreatePerSkeletonRootUidCombinedSkinnedMesh = [this, &MeshUidsPerSkeletonRootUid, &SetSkeletalMeshDependencies](const bool bUseInstanceMesh)
 		{
 #if !WITH_EDITOR || !WITH_EDITORONLY_DATA
-				if (MeshUidsPerSkeletonRootUid.Num() > 0)
-				{
-					UE_LOG(LogInterchangePipeline, Warning, TEXT("Cannot import skeletalMesh asset in runtime, this is an editor only feature."));
-				}
-				return false;
+			if (MeshUidsPerSkeletonRootUid.Num() > 0)
+			{
+				UE_LOG(LogInterchangePipeline, Warning, TEXT("Cannot import skeletalMesh asset in runtime, this is an editor only feature."));
+			}
+			return false;
 #else
-				if (MeshUidsPerSkeletonRootUid.Num() > 0)
+			if (MeshUidsPerSkeletonRootUid.Num() > 0)
+			{
+
+				bool bFoundInstances = false;
+				for (const TPair<FString, TArray<FString>>& SkeletonRootUidAndMeshUids : MeshUidsPerSkeletonRootUid)
 				{
-					if (FApp::IsGame())
-					{
-						//PIE
-						UE_LOG(LogInterchangePipeline, Warning, TEXT("Cannot import skeletalMesh asset in runtime, this is an editor only feature."));
-						return false;
-					}
-					else
+					const FString& SkeletonRootUid = SkeletonRootUidAndMeshUids.Key;
+					//Every iteration is a skeletalmesh asset that combine all MeshInstances sharing the same skeleton root node
+					UInterchangeSkeletonFactoryNode* SkeletonFactoryNode = CommonSkeletalMeshesAndAnimationsProperties->CreateSkeletonFactoryNode(BaseNodeContainer, SkeletonRootUid);
+					//The MeshUids can represent a SceneNode pointing on a MeshNode or directly a MeshNode;
+					TMap<int32, TArray<FString>> MeshUidsPerLodIndex;
+					const TArray<FString>& MeshUids = SkeletonRootUidAndMeshUids.Value;
+					for (const FString& MeshUid : MeshUids)
 					{
-						bool bFoundInstances = false;
-						for (const TPair<FString, TArray<FString>>& SkeletonRootUidAndMeshUids : MeshUidsPerSkeletonRootUid)
+						if (bUseInstanceMesh)
 						{
-							const FString& SkeletonRootUid = SkeletonRootUidAndMeshUids.Key;
-							//Every iteration is a skeletalmesh asset that combine all MeshInstances sharing the same skeleton root node
-							UInterchangeSkeletonFactoryNode* SkeletonFactoryNode = CommonSkeletalMeshesAndAnimationsProperties->CreateSkeletonFactoryNode(BaseNodeContainer, SkeletonRootUid);
-							//The MeshUids can represent a SceneNode pointing on a MeshNode or directly a MeshNode;
-							TMap<int32, TArray<FString>> MeshUidsPerLodIndex;
-							const TArray<FString>& MeshUids = SkeletonRootUidAndMeshUids.Value;
-							for (const FString& MeshUid : MeshUids)
+							const FInterchangeMeshInstance& MeshInstance = PipelineMeshesUtilities->GetMeshInstanceByUid(MeshUid);
+							for (const TPair<int32, FInterchangeLodSceneNodeContainer>& LodIndexAndSceneNodeContainer : MeshInstance.SceneNodePerLodIndex)
 							{
-								if (bUseInstanceMesh)
-								{
-									const FInterchangeMeshInstance& MeshInstance = PipelineMeshesUtilities->GetMeshInstanceByUid(MeshUid);
-									for (const TPair<int32, FInterchangeLodSceneNodeContainer>& LodIndexAndSceneNodeContainer : MeshInstance.SceneNodePerLodIndex)
-									{
-										const int32 LodIndex = LodIndexAndSceneNodeContainer.Key;
-										const FInterchangeLodSceneNodeContainer& SceneNodeContainer = LodIndexAndSceneNodeContainer.Value;
-										TArray<FString>& TranslatedNodes = MeshUidsPerLodIndex.FindOrAdd(LodIndex);
-										for (const UInterchangeSceneNode* SceneNode : SceneNodeContainer.SceneNodes)
-										{
-											TranslatedNodes.Add(SceneNode->GetUniqueID());
-										}
-									}
-								}
-								else
+								const int32 LodIndex = LodIndexAndSceneNodeContainer.Key;
+								const FInterchangeLodSceneNodeContainer& SceneNodeContainer = LodIndexAndSceneNodeContainer.Value;
+								TArray<FString>& TranslatedNodes = MeshUidsPerLodIndex.FindOrAdd(LodIndex);
+								for (const UInterchangeSceneNode* SceneNode : SceneNodeContainer.SceneNodes)
 								{
-									//MeshGeometry cannot have Lod since LODs are define in the scene node
-									const FInterchangeMeshGeometry& MeshGeometry = PipelineMeshesUtilities->GetMeshGeometryByUid(MeshUid);
-									const int32 LodIndex = 0;
-									TArray<FString>& TranslatedNodes = MeshUidsPerLodIndex.FindOrAdd(LodIndex);
-									TranslatedNodes.Add(MeshGeometry.MeshUid);
+									TranslatedNodes.Add(SceneNode->GetUniqueID());
 								}
 							}
+						}
+						else
+						{
+							//MeshGeometry cannot have Lod since LODs are define in the scene node
+							const FInterchangeMeshGeometry& MeshGeometry = PipelineMeshesUtilities->GetMeshGeometryByUid(MeshUid);
+							const int32 LodIndex = 0;
+							TArray<FString>& TranslatedNodes = MeshUidsPerLodIndex.FindOrAdd(LodIndex);
+							TranslatedNodes.Add(MeshGeometry.MeshUid);
+						}
+					}
 
-							if (MeshUidsPerLodIndex.Num() > 0)
-							{
-								UInterchangeSkeletalMeshFactoryNode* SkeletalMeshFactoryNode = CreateSkeletalMeshFactoryNode(SkeletonRootUid, MeshUidsPerLodIndex);
-								if (SkeletalMeshFactoryNode != nullptr)
-								{
-									SetSkeletalMeshDependencies(SkeletonRootUid, SkeletalMeshFactoryNode);
-									SkeletonFactoryNodes.Add(SkeletonFactoryNode);
-									SkeletalMeshFactoryNodes.Add(SkeletalMeshFactoryNode);
-									bFoundInstances = true;
-								}
-							}
+					if (MeshUidsPerLodIndex.Num() > 0)
+					{
+						UInterchangeSkeletalMeshFactoryNode* SkeletalMeshFactoryNode = CreateSkeletalMeshFactoryNode(SkeletonRootUid, MeshUidsPerLodIndex);
+						if (SkeletalMeshFactoryNode != nullptr)
+						{
+							SetSkeletalMeshDependencies(SkeletonRootUid, SkeletalMeshFactoryNode);
+							SkeletonFactoryNodes.Add(SkeletonFactoryNode);
+							SkeletalMeshFactoryNodes.Add(SkeletalMeshFactoryNode);
+							bFoundInstances = true;
 						}
-						return bFoundInstances;
 					}
 				}
-				return false;
+				return bFoundInstances;
+			}
+			return false;
 #endif
 		};
 
@@ -222,75 +215,66 @@ void UInterchangeGenericMeshPipeline::ExecutePreImportPipelineSkeletalMesh()
 #else
 			if (MeshUids.Num() > 0)
 			{
-				if (FApp::IsGame())
+				bool bFoundInstances = false;
+				for (const FString& MeshUid : MeshUids)
 				{
-					//PIE
-					UE_LOG(LogInterchangePipeline, Warning, TEXT("Cannot import skeletalMesh asset in runtime, this is an editor only feature."));
-					return false;
-				}
-				else
-				{
-					bool bFoundInstances = false;
-					for (const FString& MeshUid : MeshUids)
+					//Every iteration is a skeletalmesh asset that combine all MeshInstances sharing the same skeleton root node
+					//The MeshUids can represent a SceneNode pointing on a MeshNode or directly a MeshNode;
+					TMap<int32, TArray<FString>> MeshUidsPerLodIndex;
+					FString SkeletonRootUid;
+					if (!(bUseInstanceMesh ? PipelineMeshesUtilities->IsValidMeshInstanceUid(MeshUid) : PipelineMeshesUtilities->IsValidMeshGeometryUid(MeshUid)))
+					{
+						continue;
+					}
+					SkeletonRootUid = (bUseInstanceMesh ? PipelineMeshesUtilities->GetMeshInstanceSkeletonRootUid(MeshUid) : PipelineMeshesUtilities->GetMeshGeometrySkeletonRootUid(MeshUid));
+					if (SkeletonRootUid.IsEmpty())
 					{
-						//Every iteration is a skeletalmesh asset that combine all MeshInstances sharing the same skeleton root node
-						//The MeshUids can represent a SceneNode pointing on a MeshNode or directly a MeshNode;
-						TMap<int32, TArray<FString>> MeshUidsPerLodIndex;
-						FString SkeletonRootUid;
-						if (!(bUseInstanceMesh ? PipelineMeshesUtilities->IsValidMeshInstanceUid(MeshUid) : PipelineMeshesUtilities->IsValidMeshGeometryUid(MeshUid)))
+						if (CommonSkeletalMeshesAndAnimationsProperties->bConvertStaticsWithMorphTargetsToSkeletals)
 						{
-							continue;
+							SkeletonRootUid = MeshUid;
 						}
-						SkeletonRootUid = (bUseInstanceMesh ? PipelineMeshesUtilities->GetMeshInstanceSkeletonRootUid(MeshUid) : PipelineMeshesUtilities->GetMeshGeometrySkeletonRootUid(MeshUid));
-						if (SkeletonRootUid.IsEmpty())
+						else
 						{
-							if (CommonSkeletalMeshesAndAnimationsProperties->bConvertStaticsWithMorphTargetsToSkeletals)
-							{
-								SkeletonRootUid = MeshUid;
-							}
-							else
-							{
-								//Log an error
-								continue;
-							}
+							//Log an error
+							continue;
 						}
-						UInterchangeSkeletonFactoryNode* SkeletonFactoryNode = CommonSkeletalMeshesAndAnimationsProperties->CreateSkeletonFactoryNode(BaseNodeContainer, SkeletonRootUid);
-						if (bUseInstanceMesh)
+					}
+					UInterchangeSkeletonFactoryNode* SkeletonFactoryNode = CommonSkeletalMeshesAndAnimationsProperties->CreateSkeletonFactoryNode(BaseNodeContainer, SkeletonRootUid);
+					if (bUseInstanceMesh)
+					{
+						const FInterchangeMeshInstance& MeshInstance = PipelineMeshesUtilities->GetMeshInstanceByUid(MeshUid);
+						for (const TPair<int32, FInterchangeLodSceneNodeContainer>& LodIndexAndSceneNodeContainer : MeshInstance.SceneNodePerLodIndex)
 						{
-							const FInterchangeMeshInstance& MeshInstance = PipelineMeshesUtilities->GetMeshInstanceByUid(MeshUid);
-							for (const TPair<int32, FInterchangeLodSceneNodeContainer>& LodIndexAndSceneNodeContainer : MeshInstance.SceneNodePerLodIndex)
+							const int32 LodIndex = LodIndexAndSceneNodeContainer.Key;
+							const FInterchangeLodSceneNodeContainer& SceneNodeContainer = LodIndexAndSceneNodeContainer.Value;
+							TArray<FString>& TranslatedNodes = MeshUidsPerLodIndex.FindOrAdd(LodIndex);
+							for (const UInterchangeSceneNode* SceneNode : SceneNodeContainer.SceneNodes)
 							{
-								const int32 LodIndex = LodIndexAndSceneNodeContainer.Key;
-								const FInterchangeLodSceneNodeContainer& SceneNodeContainer = LodIndexAndSceneNodeContainer.Value;
-								TArray<FString>& TranslatedNodes = MeshUidsPerLodIndex.FindOrAdd(LodIndex);
-								for (const UInterchangeSceneNode* SceneNode : SceneNodeContainer.SceneNodes)
-								{
-									TranslatedNodes.Add(SceneNode->GetUniqueID());
-								}
+								TranslatedNodes.Add(SceneNode->GetUniqueID());
 							}
 						}
-						else
+					}
+					else
+					{
+						const FInterchangeMeshGeometry& MeshGeometry = PipelineMeshesUtilities->GetMeshGeometryByUid(MeshUid);
+						const int32 LodIndex = 0;
+						TArray<FString>& TranslatedNodes = MeshUidsPerLodIndex.FindOrAdd(LodIndex);
+						TranslatedNodes.Add(MeshGeometry.MeshUid);
+					}
+					if (MeshUidsPerLodIndex.Num() > 0)
+					{
+						UInterchangeSkeletalMeshFactoryNode* SkeletalMeshFactoryNode = CreateSkeletalMeshFactoryNode(SkeletonRootUid, MeshUidsPerLodIndex);
+						if (SkeletalMeshFactoryNode != nullptr)
 						{
-							const FInterchangeMeshGeometry& MeshGeometry = PipelineMeshesUtilities->GetMeshGeometryByUid(MeshUid);
-							const int32 LodIndex = 0;
-							TArray<FString>& TranslatedNodes = MeshUidsPerLodIndex.FindOrAdd(LodIndex);
-							TranslatedNodes.Add(MeshGeometry.MeshUid);
+							SetSkeletalMeshDependencies(SkeletonRootUid, SkeletalMeshFactoryNode);
+							SkeletonFactoryNodes.Add(SkeletonFactoryNode);
+							SkeletalMeshFactoryNodes.Add(SkeletalMeshFactoryNode);
+							bFoundInstances = true;
 						}
-						if (MeshUidsPerLodIndex.Num() > 0)
-						{
-							UInterchangeSkeletalMeshFactoryNode* SkeletalMeshFactoryNode = CreateSkeletalMeshFactoryNode(SkeletonRootUid, MeshUidsPerLodIndex);
-							if (SkeletalMeshFactoryNode != nullptr)
-							{
-								SetSkeletalMeshDependencies(SkeletonRootUid, SkeletalMeshFactoryNode);
-								SkeletonFactoryNodes.Add(SkeletonFactoryNode);
-								SkeletalMeshFactoryNodes.Add(SkeletalMeshFactoryNode);
-								bFoundInstances = true;
-							}
 
-						}
 					}
-					return bFoundInstances;
 				}
+				return bFoundInstances;
 			}
 			return false;
 #endif
diff --git a/Engine/Plugins/Interchange/Runtime/Source/Pipelines/Private/InterchangeglTFPipeline.cpp b/Engine/Plugins/Interchange/Runtime/Source/Pipelines/Private/InterchangeglTFPipeline.cpp
--- a/Engine/Plugins/Interchange/Runtime/Source/Pipelines/Private/InterchangeglTFPipeline.cpp
+++ b/Engine/Plugins/Interchange/Runtime/Source/Pipelines/Private/InterchangeglTFPipeline.cpp
@@ -8,6 +8,7 @@
 #include "InterchangeShaderGraphNode.h"
 #include "InterchangeMaterialInstanceNode.h"
 #include "InterchangeMaterialFactoryNode.h"
+#include "InterchangeManager.h"
 
 #include "Gltf/InterchangeGLTFMaterial.h"
 
@@ -225,7 +226,8 @@ void UGLTFPipelineSettings::BuildMaterialInstance(const UInterchangeShaderGraphN
 
 	MaterialInstanceFactoryNode->SetCustomParent(Parent);
 
-	const UClass* MaterialClass = FApp::IsGame() ? UMaterialInstanceDynamic::StaticClass() : UMaterialInstanceConstant::StaticClass();
+	UInterchangeEditorUtilitiesBase* EditorUtilities = UInterchangeManager::GetInterchangeManager().GetEditorUtilities();
+	const UClass* MaterialClass = (EditorUtilities && EditorUtilities->IsRuntimeOrPIE()) ? UMaterialInstanceDynamic::StaticClass() : UMaterialInstanceConstant::StaticClass();
 	MaterialInstanceFactoryNode->SetCustomInstanceClassName(MaterialClass->GetPathName());
 
 	for (const TPair<FString, UE::Interchange::FAttributeKey>& GltfAttributeKey : GltfAttributeKeys)
diff --git a/Engine/Source/Runtime/Interchange/Core/Public/Nodes/InterchangeFactoryBaseNode.h b/Engine/Source/Runtime/Interchange/Core/Public/Nodes/InterchangeFactoryBaseNode.h
--- a/Engine/Source/Runtime/Interchange/Core/Public/Nodes/InterchangeFactoryBaseNode.h
+++ b/Engine/Source/Runtime/Interchange/Core/Public/Nodes/InterchangeFactoryBaseNode.h
@@ -361,6 +361,16 @@ public:
 	UFUNCTION(BlueprintCallable, Category = "Interchange | Node")
 	INTERCHANGECORE_API bool UnsetForceNodeReimport();
 
+	/**
+	 * Return if the import of the class is allowed at runtime.
+	 *
+	 */
+	UFUNCTION(BlueprintCallable, Category = "Interchange | Node")
+	INTERCHANGECORE_API virtual bool IsRuntimeImportAllowed() const
+	{
+		return true;
+	}
+
 protected:
 	/**
 	 * Those dependencies are used by the Interchange parsing task to make sure the assets are created in the correct order.
diff --git a/Engine/Source/Runtime/Interchange/Engine/Private/InterchangeManager.cpp b/Engine/Source/Runtime/Interchange/Engine/Private/InterchangeManager.cpp
--- a/Engine/Source/Runtime/Interchange/Engine/Private/InterchangeManager.cpp
+++ b/Engine/Source/Runtime/Interchange/Engine/Private/InterchangeManager.cpp
@@ -2542,6 +2542,8 @@ TSharedRef<UE::Interchange::FImportAsyncHelper, ESPMode::ThreadSafe> UInterchang
 
 	AsyncHelper->AssetImportResult->SetInProgress();
 
+	AsyncHelper->bRuntimeOrPIE = GetEditorUtilities() ? GetEditorUtilities()->IsRuntimeOrPIE() : false;
+
 	return AsyncHelper;
 }
 
diff --git a/Engine/Source/Runtime/Interchange/Engine/Private/Tasks/InterchangeTaskImportObject.cpp b/Engine/Source/Runtime/Interchange/Engine/Private/Tasks/InterchangeTaskImportObject.cpp
--- a/Engine/Source/Runtime/Interchange/Engine/Private/Tasks/InterchangeTaskImportObject.cpp
+++ b/Engine/Source/Runtime/Interchange/Engine/Private/Tasks/InterchangeTaskImportObject.cpp
@@ -114,12 +114,22 @@ namespace UE::Interchange::Private
 
 
 
-	bool CanImportClass(UE::Interchange::FImportAsyncHelper& AsyncHelper, UInterchangeFactoryBaseNode& FactoryNode, int32 SourceIndex)
+	bool CanImportClass(UE::Interchange::FImportAsyncHelper& AsyncHelper, UInterchangeFactoryBaseNode& FactoryNode, int32 SourceIndex, bool bLogError = false)
 	{
+		if (AsyncHelper.bRuntimeOrPIE && !FactoryNode.IsRuntimeImportAllowed())
+		{
+			if (bLogError)
+			{
+				UE_LOG(LogInterchangeEngine, Error, TEXT("Cannot import %s asset at runtime. This is an editor-only feature."), *FactoryNode.GetTypeName());
+			}
+			return false;
+		}
+
 		if (UClass* Class = FactoryNode.GetObjectClass())
 		{
 			return AsyncHelper.IsClassImportAllowed(Class);
 		}
+		
 		return false;
 	}
 
@@ -264,8 +274,8 @@ void UE::Interchange::FTaskImportObject_GameThread::Execute()
 		}
 	};
 	
-	//Verify if the task was cancel or class to import is denied
-	if (AsyncHelper->bCancel || !FactoryNode || !Private::CanImportClass(*AsyncHelper, *FactoryNode, SourceIndex))
+	//Verify if the task was cancel
+	if (AsyncHelper->bCancel || !FactoryNode)
 	{
 		return;
 	}
@@ -284,7 +294,6 @@ void UE::Interchange::FTaskImportObject_GameThread::Execute()
 	FString PackageName;
 	FString AssetName;
 	Private::InternalGetPackageName(*AsyncHelper, SourceIndex, PackageBasePath, FactoryNode, PackageName, AssetName);
-
 	UObject* ObjectToReimport = FFactoryCommon::GetObjectToReimport(Factory, AsyncHelper->TaskData.ReimportObject, *FactoryNode, PackageName, AssetName);
 	if (!ensure(!IsGarbageCollecting()))
 	{
@@ -310,6 +319,13 @@ void UE::Interchange::FTaskImportObject_GameThread::Execute()
 		return;
 	}
 	UInterchangeBaseNodeContainer* NodeContainer = AsyncHelper->BaseNodeContainers[SourceIndex].Get();
+	bool bShouldReimportFactoryNode = ObjectToReimport && Private::ShouldReimportFactoryNode(FactoryNode, NodeContainer, ObjectToReimport);
+
+	//Check if class Import is allowed or not.
+	if (!Private::CanImportClass(*AsyncHelper, *FactoryNode, SourceIndex, !ObjectToReimport || bShouldReimportFactoryNode))
+	{
+		return;
+	}
 
 	bool bSkipObjectNoReplace = false;
 	UObject* ExistingAsset = nullptr;
@@ -377,7 +393,7 @@ void UE::Interchange::FTaskImportObject_GameThread::Execute()
 	//If we do a reimport no need to create a package
 	if (ObjectToReimport)
 	{
-		if (Private::ShouldReimportFactoryNode(FactoryNode, NodeContainer, ObjectToReimport))
+		if (bShouldReimportFactoryNode)
 		{
 			FactoryNode->SetDisplayLabel(ObjectToReimport->GetName());
 			FactoryNode->SetAssetName(ObjectToReimport->GetName());
diff --git a/Engine/Source/Runtime/Interchange/Engine/Private/Tasks/InterchangeTaskParsing.cpp b/Engine/Source/Runtime/Interchange/Engine/Private/Tasks/InterchangeTaskParsing.cpp
--- a/Engine/Source/Runtime/Interchange/Engine/Private/Tasks/InterchangeTaskParsing.cpp
+++ b/Engine/Source/Runtime/Interchange/Engine/Private/Tasks/InterchangeTaskParsing.cpp
@@ -313,7 +313,13 @@ void UE::Interchange::FTaskParsing::Execute()
 						//Translator not thread safe cannot get payloads asynchronously
 						bAsyncPayloads = false;
 					}
-					Factory->CreatePayloadTasks(ImportAssetObjectParams, bAsyncPayloads, PayloadTasks);
+
+					if (!AsyncHelper->bRuntimeOrPIE || FactoryNode->IsRuntimeImportAllowed())
+					{
+						//Only create payloads if it's allowed.
+						Factory->CreatePayloadTasks(ImportAssetObjectParams, bAsyncPayloads, PayloadTasks);
+					}
+					
 					if (!PayloadTasks.IsEmpty())
 					{
 						for (TSharedPtr<FInterchangeTaskBase> PayloadTask : PayloadTasks)
diff --git a/Engine/Source/Runtime/Interchange/Engine/Public/InterchangeEditorUtilitiesBase.h b/Engine/Source/Runtime/Interchange/Engine/Public/InterchangeEditorUtilitiesBase.h
--- a/Engine/Source/Runtime/Interchange/Engine/Public/InterchangeEditorUtilitiesBase.h
+++ b/Engine/Source/Runtime/Interchange/Engine/Public/InterchangeEditorUtilitiesBase.h
@@ -17,4 +17,13 @@ public:
 	{
 		return false;
 	}
+
+	INTERCHANGEENGINE_API virtual bool IsRuntimeOrPIE()
+	{
+#if WITH_EDITOR
+		return false;
+#else
+		return true;
+#endif
+	}
 };
diff --git a/Engine/Source/Runtime/Interchange/Engine/Public/InterchangeManager.h b/Engine/Source/Runtime/Interchange/Engine/Public/InterchangeManager.h
--- a/Engine/Source/Runtime/Interchange/Engine/Public/InterchangeManager.h
+++ b/Engine/Source/Runtime/Interchange/Engine/Public/InterchangeManager.h
@@ -224,6 +224,7 @@ namespace UE
 			}
 
 			bool bRunSynchronous = false;
+			bool bRuntimeOrPIE = false;
 
 			/** Unique ID for this async helper. */
 			int32 UniqueId;
