Interchange: When baking the skeleton hierarchy animation, we need to bake all the nodes even the non animated nodes.

#rnx

diff --git a/Engine/Plugins/Interchange/Runtime/Source/Nodes/Private/InterchangeAnimationTrackSetNode.cpp b/Engine/Plugins/Interchange/Runtime/Source/Nodes/Private/InterchangeAnimationTrackSetNode.cpp
--- a/Engine/Plugins/Interchange/Runtime/Source/Nodes/Private/InterchangeAnimationTrackSetNode.cpp
+++ b/Engine/Plugins/Interchange/Runtime/Source/Nodes/Private/InterchangeAnimationTrackSetNode.cpp
@@ -330,3 +330,13 @@ bool UInterchangeSkeletalAnimationTrackNode::SetAnimationPayloadKeyForMorphTarge
 	bSuccess &= MorphTargetPayloadKeyTypeMap.SetKeyValue(MorphTargetNodeUid, (uint8)InType);
 	return bSuccess;
 }
+
+bool UInterchangeSkeletalAnimationTrackNode::IsNodeAnimatedWithBakedCurve(const FString& SceneNodeUid) const
+{
+	uint8 OutValue = 0;
+	if (SceneNodeAnimationPayloadKeyTypeMap.GetValue(SceneNodeUid, OutValue))
+	{
+		return (OutValue == static_cast<uint8>(EInterchangeAnimationPayLoadType::BAKED));
+	}
+	return false;
+}
\ No newline at end of file
diff --git a/Engine/Plugins/Interchange/Runtime/Source/Nodes/Public/InterchangeAnimationTrackSetNode.h b/Engine/Plugins/Interchange/Runtime/Source/Nodes/Public/InterchangeAnimationTrackSetNode.h
--- a/Engine/Plugins/Interchange/Runtime/Source/Nodes/Public/InterchangeAnimationTrackSetNode.h
+++ b/Engine/Plugins/Interchange/Runtime/Source/Nodes/Public/InterchangeAnimationTrackSetNode.h
@@ -529,6 +529,9 @@ public:
 	UFUNCTION(BlueprintCallable, Category = "Interchange | Node | SkeletalAnimationTrack")
 	bool SetAnimationPayloadKeyForMorphTargetNodeUid(const FString& MorphTargetNodeUid, const FString& InUniqueId, const EInterchangeAnimationPayLoadType& InType);
 
+	UFUNCTION(BlueprintCallable, Category = "Interchange | Node | SkeletalAnimationTrack")
+	bool IsNodeAnimatedWithBakedCurve(const FString& SceneNodeUid) const;
+
 private:
 	const UE::Interchange::FAttributeKey Macro_CustomSkeletonNodeUidKey = UE::Interchange::FAttributeKey(TEXT("SkeletonNodeUid"));
 	const UE::Interchange::FAttributeKey Macro_CustomAnimationSampleRateKey = UE::Interchange::FAttributeKey(TEXT("AnimationSampleRate"));
diff --git a/Engine/Plugins/Interchange/Runtime/Source/Parsers/Fbx/Private/FbxAnimation.cpp b/Engine/Plugins/Interchange/Runtime/Source/Parsers/Fbx/Private/FbxAnimation.cpp
--- a/Engine/Plugins/Interchange/Runtime/Source/Parsers/Fbx/Private/FbxAnimation.cpp
+++ b/Engine/Plugins/Interchange/Runtime/Source/Parsers/Fbx/Private/FbxAnimation.cpp
@@ -2,6 +2,7 @@
 
 #include "FbxAnimation.h"
 
+#include "Async/ParallelFor.h"
 #include "CoreMinimal.h"
 #include "FbxAPI.h"
 #include "FbxConvert.h"
@@ -9,6 +10,7 @@
 #include "FbxInclude.h"
 #include "FbxMesh.h"
 #include "Fbx/InterchangeFbxMessages.h"
+#include "InterchangeAnimationTrackSetNode.h"
 #include "InterchangeCommonAnimationPayload.h"
 #include "InterchangeMeshNode.h"
 #include "InterchangeResultsContainer.h"
@@ -16,12 +18,10 @@
 #include "MeshDescription.h"
 #include "Misc/FileHelper.h"
 #include "Misc/ScopeLock.h"
+#include "Nodes/InterchangeBaseNodeContainer.h"
 #include "Serialization/LargeMemoryWriter.h"
 #include "SkeletalMeshAttributes.h"
 #include "StaticMeshAttributes.h"
-#include "InterchangeAnimationTrackSetNode.h"
-#include "Nodes/InterchangeBaseNodeContainer.h"
-#include "Async/ParallelFor.h"
 
 #define LOCTEXT_NAMESPACE "InterchangeFbxMesh"
 
@@ -768,7 +768,8 @@ namespace UE::Interchange::Private
 		return ImportCurve(AnimCurve, ScaleCurve, InterchangeCurves.AddDefaulted_GetRef().Keys);
 	}
 
-	bool FFbxAnimation::AddSkeletalTransformAnimation(FbxScene* SDKScene
+	bool FFbxAnimation::AddSkeletalTransformAnimation(UInterchangeBaseNodeContainer& NodeContainer
+		, FbxScene* SDKScene
 		, FFbxParser& Parser
 		, FbxNode* Node
 		, UInterchangeSceneNode* SceneNode
@@ -778,6 +779,30 @@ namespace UE::Interchange::Private
 	{
 		FGetFbxTransformCurvesParameters Parameters(SDKScene, Node);
 		GetFbxTransformCurves(Parameters, AnimationIndex);
+		if (!Parameters.IsNodeAnimated)
+		{
+			//If we have a joint under the root skeleton and there is some animation in the parent hierarchy
+			//we have to enable IsNodeAnimated so it get bake correctly and generate the appropriate curves
+			FString RootSkeletonNodeUid;
+			if (SkeletalAnimationTrackNode->GetCustomSkeletonNodeUid(RootSkeletonNodeUid))
+			{
+				const UInterchangeSceneNode* ParentNode = Cast<UInterchangeSceneNode>(NodeContainer.GetNode(SceneNode->GetParentUid()));
+				//Search up the hierarchy if we found any animated parent we have to animate the children
+				while (ParentNode)
+				{
+					if (SkeletalAnimationTrackNode->IsNodeAnimatedWithBakedCurve(ParentNode->GetUniqueID()))
+					{
+						Parameters.IsNodeAnimated = true;
+						break;
+					}
+					if (ParentNode->GetUniqueID().Equals(RootSkeletonNodeUid))
+					{
+						break;
+					}
+					ParentNode = Cast<UInterchangeSceneNode>(NodeContainer.GetNode(ParentNode->GetParentUid()));
+				}
+			}
+		}
 
 		if (Parameters.IsNodeAnimated)
 		{
diff --git a/Engine/Plugins/Interchange/Runtime/Source/Parsers/Fbx/Private/FbxAnimation.h b/Engine/Plugins/Interchange/Runtime/Source/Parsers/Fbx/Private/FbxAnimation.h
--- a/Engine/Plugins/Interchange/Runtime/Source/Parsers/Fbx/Private/FbxAnimation.h
+++ b/Engine/Plugins/Interchange/Runtime/Source/Parsers/Fbx/Private/FbxAnimation.h
@@ -10,6 +10,7 @@
 /** Forward declarations */
 struct FInterchangeCurve;
 struct FMeshDescription;
+class UInterchangeBaseNodeContainer;
 class UInterchangeMeshNode;
 class UInterchangeSceneNode;
 class UInterchangeSkeletalAnimationTrackNode;
@@ -98,7 +99,8 @@ namespace UE::Interchange::Private
 	{
 	public:
 		/** This function add the payload key if the scene node transform is animated. */
-		static bool AddSkeletalTransformAnimation(FbxScene* SDKScene
+		static bool AddSkeletalTransformAnimation(UInterchangeBaseNodeContainer& NodeContainer
+			, FbxScene* SDKScene
 			, FFbxParser& Parser
 			, FbxNode* Node
 			, UInterchangeSceneNode* UnrealNode
diff --git a/Engine/Plugins/Interchange/Runtime/Source/Parsers/Fbx/Private/FbxScene.cpp b/Engine/Plugins/Interchange/Runtime/Source/Parsers/Fbx/Private/FbxScene.cpp
--- a/Engine/Plugins/Interchange/Runtime/Source/Parsers/Fbx/Private/FbxScene.cpp
+++ b/Engine/Plugins/Interchange/Runtime/Source/Parsers/Fbx/Private/FbxScene.cpp
@@ -806,7 +806,7 @@ namespace UE
 					else if (SkeletalAnimationTrackNode)
 					{
 						//Scene node transform can be animated, add the transform animation payload key.
-						if (FFbxAnimation::AddSkeletalTransformAnimation(SDKScene, Parser, Node, UnrealNode, PayloadContexts, SkeletalAnimationTrackNode, AnimationIndex)
+						if (FFbxAnimation::AddSkeletalTransformAnimation(NodeContainer, SDKScene, Parser, Node, UnrealNode, PayloadContexts, SkeletalAnimationTrackNode, AnimationIndex)
 							&& !SkeletalAnimationAddedToContainer)
 						{
 							SkeletalAnimationAddedToContainer = true;
