MaterialX
- Adding support of displacementshader, now the surfacematerial node will also translate displacement maps

#virtualized

diff --git a/Engine/Plugins/Interchange/Runtime/Source/Common/Public/MaterialX/InterchangeMaterialXDefinitions.h b/Engine/Plugins/Interchange/Runtime/Source/Common/Public/MaterialX/InterchangeMaterialXDefinitions.h
--- a/Engine/Plugins/Interchange/Runtime/Source/Common/Public/MaterialX/InterchangeMaterialXDefinitions.h
+++ b/Engine/Plugins/Interchange/Runtime/Source/Common/Public/MaterialX/InterchangeMaterialXDefinitions.h
@@ -47,6 +47,9 @@ enum class EInterchangeMaterialXShaders : uint8
 	/** A surface shader constructed from scattering and emission distribution functions. */
 	Surface,
 
+	/** Shader used for displacement. */
+	Displacement,
+
 	MaxShaderCount UMETA(hidden)
 };
 
@@ -610,6 +613,8 @@ MATERIALX_NAMESPACE_BEGIN
 		static constexpr const char* UsdPreviewSurface = "UsdPreviewSurface";
 		// Shader
 		static constexpr const char* SurfaceUnlit = "surface_unlit";
+		// Displacement Shader
+		static constexpr const char* Displacement = "displacement";
 		// Convolution
 		static constexpr const char* Blur = "blur";
 		static constexpr const char* HeightToNormal = "heighttonormal";
@@ -625,6 +630,8 @@ MATERIALX_NAMESPACE_BEGIN
 		static constexpr const char* DirectionalLight = "ND_directional_light";
 		static constexpr const char* SpotLight = "ND_spot_light";
 		static constexpr const char* Surface = "ND_surface";
+		static constexpr const char* DisplacementFloat = "ND_displacement_float";
+		static constexpr const char* DisplacementVector3 = "ND_displacement_vector3";
 	}
 
 	namespace Library
diff --git a/Engine/Plugins/Interchange/Runtime/Source/Import/Private/MaterialX/MaterialXUtils/MaterialXBase.cpp b/Engine/Plugins/Interchange/Runtime/Source/Import/Private/MaterialX/MaterialXUtils/MaterialXBase.cpp
--- a/Engine/Plugins/Interchange/Runtime/Source/Import/Private/MaterialX/MaterialXUtils/MaterialXBase.cpp
+++ b/Engine/Plugins/Interchange/Runtime/Source/Import/Private/MaterialX/MaterialXUtils/MaterialXBase.cpp
@@ -54,6 +54,11 @@ void FMaterialXBase::UpdateDocumentRecursively(MaterialX::GraphElementPtr Graph)
 	}
 }
 
+void FMaterialXBase::SetNodeDefinition(const char* NodeDef)
+{
+	this->NodeDefinition = NodeDef;
+}
+
 
 void FMaterialXBase::AddTexCoordToTiledImageNodes(MaterialX::ElementPtr Graph)
 {
diff --git a/Engine/Plugins/Interchange/Runtime/Source/Import/Private/MaterialX/MaterialXUtils/MaterialXDisplacementShader.cpp b/Engine/Plugins/Interchange/Runtime/Source/Import/Private/MaterialX/MaterialXUtils/MaterialXDisplacementShader.cpp
new file mode 100644
--- /dev/null
+++ b/Engine/Plugins/Interchange/Runtime/Source/Import/Private/MaterialX/MaterialXUtils/MaterialXDisplacementShader.cpp
@@ -0,0 +1,34 @@
+// Copyright Epic Games, Inc. All Rights Reserved.
+
+#if WITH_EDITOR
+
+#include "MaterialXDisplacementShader.h"
+#include "InterchangeImportLog.h"
+
+namespace mx = MaterialX;
+
+FMaterialXDisplacementShader::FMaterialXDisplacementShader(UInterchangeBaseNodeContainer& BaseNodeContainer)
+	: FMaterialXSurfaceShaderAbstract(BaseNodeContainer)
+{
+	NodeDefinition = mx::NodeDefinition::DisplacementFloat; // we set by default a displacement float but the SurfaceMaterial should set the correct nodedef
+}
+
+TSharedRef<FMaterialXBase> FMaterialXDisplacementShader::MakeInstance(UInterchangeBaseNodeContainer& BaseNodeContainer)
+{
+	TSharedRef<FMaterialXDisplacementShader> Result = MakeShared<FMaterialXDisplacementShader>(BaseNodeContainer);
+	Result->RegisterConnectNodeOutputToInputDelegates();
+	return Result;
+}
+
+void FMaterialXDisplacementShader::Translate(MaterialX::NodePtr DisplacementNode)
+{
+	this->SurfaceShaderNode = DisplacementNode;
+
+	UInterchangeShaderNode* DisplacementShaderNode = FMaterialXSurfaceShaderAbstract::Translate(EInterchangeMaterialXShaders::Displacement);
+
+	using namespace UE::Interchange::Materials;
+
+	// Outputs
+	UInterchangeShaderPortsAPI::ConnectOuputToInputByName(ShaderGraphNode, Common::Parameters::Displacement.ToString(), DisplacementShaderNode->GetUniqueID(), Common::Parameters::Displacement.ToString());
+}
+#endif
\ No newline at end of file
diff --git a/Engine/Plugins/Interchange/Runtime/Source/Import/Private/MaterialX/MaterialXUtils/MaterialXDisplacementShader.h b/Engine/Plugins/Interchange/Runtime/Source/Import/Private/MaterialX/MaterialXUtils/MaterialXDisplacementShader.h
new file mode 100644
--- /dev/null
+++ b/Engine/Plugins/Interchange/Runtime/Source/Import/Private/MaterialX/MaterialXUtils/MaterialXDisplacementShader.h
@@ -0,0 +1,25 @@
+// Copyright Epic Games, Inc. All Rights Reserved.
+
+#pragma once
+
+#if WITH_EDITOR
+
+#include "MaterialX/MaterialXUtils/MaterialXSurfaceShaderAbstract.h"
+
+/** Handles <displacement> node from the pbrlib, we derive the class from FMaterialXSurfaceShaderAbstract to avoid issue when we StaticCast in the SurfaceMaterial, and because we also need a shader graph */
+class FMaterialXDisplacementShader : public FMaterialXSurfaceShaderAbstract
+{
+protected:
+
+	template<typename T, ESPMode>
+	friend class SharedPointerInternals::TIntrusiveReferenceController;
+
+	FMaterialXDisplacementShader(UInterchangeBaseNodeContainer& BaseNodeContainer);
+
+public:
+
+	static TSharedRef<FMaterialXBase> MakeInstance(UInterchangeBaseNodeContainer& BaseNodeContainer);
+
+	virtual void Translate(MaterialX::NodePtr DisplacementNode) override;
+};
+#endif
\ No newline at end of file
diff --git a/Engine/Plugins/Interchange/Runtime/Source/Import/Private/MaterialX/MaterialXUtils/MaterialXManager.cpp b/Engine/Plugins/Interchange/Runtime/Source/Import/Private/MaterialX/MaterialXUtils/MaterialXManager.cpp
--- a/Engine/Plugins/Interchange/Runtime/Source/Import/Private/MaterialX/MaterialXUtils/MaterialXManager.cpp
+++ b/Engine/Plugins/Interchange/Runtime/Source/Import/Private/MaterialX/MaterialXUtils/MaterialXManager.cpp
@@ -16,6 +16,7 @@
 #include "MaterialX/MaterialXUtils/MaterialXDirectionalLightShader.h"
 #include "MaterialX/MaterialXUtils/MaterialXSpotLightShader.h"
 #include "MaterialX/MaterialXUtils/MaterialXSurfaceMaterial.h"
+#include "MaterialX/MaterialXUtils/MaterialXDisplacementShader.h"
 
 #include "InterchangeMaterialDefinitions.h"
 #include "InterchangeImportLog.h"
@@ -214,6 +215,7 @@ FMaterialXManager::FMaterialXManager()
 		{mx::Category::SurfaceUnlit, FMaterialXManager::FOnGetMaterialXInstance::CreateStatic(&FMaterialXSurfaceUnlitShader::MakeInstance)},
 		{mx::Category::StandardSurface, FMaterialXManager::FOnGetMaterialXInstance::CreateStatic(&FMaterialXStandardSurfaceShader::MakeInstance)},
 		{mx::Category::UsdPreviewSurface, FMaterialXManager::FOnGetMaterialXInstance::CreateStatic(&FMaterialXUsdPreviewSurfaceShader::MakeInstance)},
+		{mx::Category::Displacement, FMaterialXManager::FOnGetMaterialXInstance::CreateStatic(&FMaterialXDisplacementShader::MakeInstance)},
 		{mx::Category::PointLight, FMaterialXManager::FOnGetMaterialXInstance::CreateStatic(&FMaterialXPointLightShader::MakeInstance)},
 		{mx::Category::DirectionalLight, FMaterialXManager::FOnGetMaterialXInstance::CreateStatic(&FMaterialXDirectionalLightShader::MakeInstance)},
 		{mx::Category::SpotLight, FMaterialXManager::FOnGetMaterialXInstance::CreateStatic(&FMaterialXSpotLightShader::MakeInstance)},
diff --git a/Engine/Plugins/Interchange/Runtime/Source/Import/Private/MaterialX/MaterialXUtils/MaterialXSurfaceMaterial.cpp b/Engine/Plugins/Interchange/Runtime/Source/Import/Private/MaterialX/MaterialXUtils/MaterialXSurfaceMaterial.cpp
--- a/Engine/Plugins/Interchange/Runtime/Source/Import/Private/MaterialX/MaterialXUtils/MaterialXSurfaceMaterial.cpp
+++ b/Engine/Plugins/Interchange/Runtime/Source/Import/Private/MaterialX/MaterialXUtils/MaterialXSurfaceMaterial.cpp
@@ -52,6 +52,7 @@ void FMaterialXSurfaceMaterial::Translate(MaterialX::NodePtr SurfaceMaterialNode
 			if(bHasSurfaceShader)
 			{
 				ShaderTranslator->ShaderGraphNode = ShaderGraphNode;
+				ShaderTranslator->SetNodeDefinition(ConnectedNode->getNodeDef()->getName().c_str()); // we set the nodedef here, because a displacement shader can either be a float or vector3, we want to be sure to retrieve the correct one
 				ShaderTranslator->Translate(ConnectedNode);
 			}
 		}
diff --git a/Engine/Plugins/Interchange/Runtime/Source/Import/Public/MaterialX/MaterialXUtils/MaterialXBase.h b/Engine/Plugins/Interchange/Runtime/Source/Import/Public/MaterialX/MaterialXUtils/MaterialXBase.h
--- a/Engine/Plugins/Interchange/Runtime/Source/Import/Public/MaterialX/MaterialXUtils/MaterialXBase.h
+++ b/Engine/Plugins/Interchange/Runtime/Source/Import/Public/MaterialX/MaterialXUtils/MaterialXBase.h
@@ -28,6 +28,12 @@ public:
 	 */
 	static void UpdateDocumentRecursively(MaterialX::GraphElementPtr Graph);
 
+	/** 
+	 * Set the node definition of a shader, useful when a shader has multiple definition per types and when we know the type only at the top level during a surfacematerial for example and not a construction-time. 
+	 * Make sure that the pointer don't cease to exist, no allocation is performed here, use the definition in MaterialXDefinitions
+	 */
+	void SetNodeDefinition(const char* NodeDefinition);
+
 protected:
 
 	/**
diff --git a/Engine/Plugins/Interchange/Runtime/Source/Pipelines/Private/InterchangeMaterialXPipeline.cpp b/Engine/Plugins/Interchange/Runtime/Source/Pipelines/Private/InterchangeMaterialXPipeline.cpp
--- a/Engine/Plugins/Interchange/Runtime/Source/Pipelines/Private/InterchangeMaterialXPipeline.cpp
+++ b/Engine/Plugins/Interchange/Runtime/Source/Pipelines/Private/InterchangeMaterialXPipeline.cpp
@@ -33,6 +33,8 @@ namespace
 	MATERIALX_FUNCTIONS_SUBSTRATE_PATH(Surface);
 	constexpr const TCHAR* UsdPreviewSurfaceFunctionsPath = TEXT("/InterchangeAssets/Functions/MX_UsdPreviewSurface.MX_UsdPreviewSurface");
 	constexpr const TCHAR* UsdPreviewSurfaceSubstratePath = TEXT("/Engine/Functions/Substrate/MF_Substrate_UsdPreviewSurface.MF_Substrate_UsdPreviewSurface");
+	constexpr const TCHAR* DisplacementFunctionsPath = TEXT("/InterchangeAssets/Functions/MX_Displacement.MX_Displacement");
+	constexpr const TCHAR* DisplacementSubstratePath = TEXT("/InterchangeAssets/Functions/MX_Displacement.MX_Displacement");
 	
 	MATERIALX_FUNCTIONS_SUBSTRATE_PATH(OrenNayarBSDF);
 	MATERIALX_FUNCTIONS_SUBSTRATE_PATH(BurleyDiffuseBSDF);
@@ -404,6 +406,18 @@ UMaterialXPipelineSettings::UMaterialXPipelineSettings()
 					TSet<FName>{Surface::Substrate::Outputs::Surface,Surface::Substrate::Outputs::Opacity}
 				}
 			},
+			{
+				UMaterialXPipelineSettings::ToEnumKey(EInterchangeMaterialXShaders::Displacement),
+				{
+					// Displacement Inputs
+					TSet<FName>{
+						TEXT("displacement"),
+						TEXT("scale"),
+					},
+					// Displacement Outputs
+					TSet<FName>{TEXT("Displacement")}
+				}
+			},
 			// BSDF
 			{
 				UMaterialXPipelineSettings::ToEnumKey(EInterchangeMaterialXBSDF::OrenNayarDiffuse),
@@ -647,6 +661,7 @@ UInterchangeMaterialXPipeline::UInterchangeMaterialXPipeline()
 			{MATERIALX_MATERIALFUNCTION_PATH(SurfaceUnlit),				  UMaterialXPipelineSettings::ToEnumKey(EInterchangeMaterialXShaders::SurfaceUnlit)},
 			{MATERIALX_MATERIALFUNCTION_PATH(Surface),					  UMaterialXPipelineSettings::ToEnumKey(EInterchangeMaterialXShaders::Surface)},
 			{MATERIALX_MATERIALFUNCTION_PATH(UsdPreviewSurface),		  UMaterialXPipelineSettings::ToEnumKey(EInterchangeMaterialXShaders::UsdPreviewSurface)},
+			{MATERIALX_MATERIALFUNCTION_PATH(Displacement),				  UMaterialXPipelineSettings::ToEnumKey(EInterchangeMaterialXShaders::Displacement)},
 																		  
 			{MATERIALX_MATERIALFUNCTION_PATH(OrenNayarBSDF),			  UMaterialXPipelineSettings::ToEnumKey(EInterchangeMaterialXBSDF::OrenNayarDiffuse)},
 			{MATERIALX_MATERIALFUNCTION_PATH(BurleyDiffuseBSDF),		  UMaterialXPipelineSettings::ToEnumKey(EInterchangeMaterialXBSDF::BurleyDiffuse)},
@@ -820,6 +835,7 @@ void UMaterialXPipelineSettings::InitPredefinedAssets()
 			{UMaterialXPipelineSettings::ToEnumKey(EInterchangeMaterialXShaders::StandardSurfaceTransmission), TransmissionSurfaceFunctionsPath, TransmissionSurfaceSubstratePath},
 			{UMaterialXPipelineSettings::ToEnumKey(EInterchangeMaterialXShaders::SurfaceUnlit), SurfaceUnlitFunctionsPath, SurfaceUnlitSubstratePath},
 			{UMaterialXPipelineSettings::ToEnumKey(EInterchangeMaterialXShaders::UsdPreviewSurface), UsdPreviewSurfaceFunctionsPath, UsdPreviewSurfaceSubstratePath},
+			{UMaterialXPipelineSettings::ToEnumKey(EInterchangeMaterialXShaders::Displacement), DisplacementFunctionsPath, DisplacementSubstratePath},
 
 			{UMaterialXPipelineSettings::ToEnumKey(EInterchangeMaterialXBSDF::OrenNayarDiffuse), OrenNayarBSDFFunctionsPath, OrenNayarBSDFSubstratePath},
 			{UMaterialXPipelineSettings::ToEnumKey(EInterchangeMaterialXBSDF::BurleyDiffuse), BurleyDiffuseBSDFFunctionsPath, BurleyDiffuseBSDFSubstratePath},
