MaterialX:
- Fixing a crash when importing a file where a node has no nodedef, by validating the document doing a prepass first to check on their nodedef and typedef

diff --git a/Engine/Plugins/Interchange/Runtime/Source/Import/Private/MaterialX/InterchangeMaterialXTranslator.cpp b/Engine/Plugins/Interchange/Runtime/Source/Import/Private/MaterialX/InterchangeMaterialXTranslator.cpp
--- a/Engine/Plugins/Interchange/Runtime/Source/Import/Private/MaterialX/InterchangeMaterialXTranslator.cpp
+++ b/Engine/Plugins/Interchange/Runtime/Source/Import/Private/MaterialX/InterchangeMaterialXTranslator.cpp
@@ -45,10 +45,60 @@ TArray<FString> UInterchangeMaterialXTranslator::GetSupportedFormats() const
 	return UE::Interchange::MaterialX::AreMaterialFunctionPackagesLoaded() ? TArray<FString>{ TEXT("mtlx;MaterialX File Format") } : TArray<FString>{};
 }
 
+namespace
+{
+#if WITH_EDITOR
+	bool ValidateDocument(MaterialX::DocumentPtr Document, const UInterchangeTranslatorBase* Translator)
+	{
+		namespace mx = MaterialX;
+		
+		if(std::string MaterialXMessage; !Document->validate(&MaterialXMessage))
+		{
+			UInterchangeResultError_Generic* Message = Translator->AddMessage<UInterchangeResultError_Generic>();
+			Message->Text = FText::Format(LOCTEXT("MaterialXDocumentInvalid", "{0}"),
+										  FText::FromString(MaterialXMessage.c_str()));
+			return false;
+		}
+
+		for(mx::ElementPtr Elem : Document->traverseTree())
+		{
+			//make sure to read only the current file otherwise we'll process the entire library
+			if(Elem->getActiveSourceUri() != Document->getActiveSourceUri())
+			{
+				continue;
+			}
+
+			mx::NodePtr Node = Elem->asA<mx::Node>();
+
+			if(Node)
+			{
+				// Validate that all nodes in the file are strictly respecting their node definition
+				if(!Node->getNodeDef())
+				{
+					UInterchangeResultError_Generic* Message = Translator->AddMessage<UInterchangeResultError_Generic>();
+					Message->Text = FText::Format(LOCTEXT("NodeDefNotFound", "<{0}> has no matching NodeDef, aborting import..."),
+												  FText::FromString(Node->getName().c_str()));
+					return false;
+				}
+				
+				if(!Node->getTypeDef())
+				{
+					UInterchangeResultError_Generic* Message = Translator->AddMessage<UInterchangeResultError_Generic>();
+					Message->Text = FText::Format(LOCTEXT("TypeDefNotFound", "<{0}> has no matching TypeDef, aborting import..."),
+												  FText::FromString(Node->getName().c_str()));
+					return false;
+				}
+			}
+		}
+
+		return true;
+	}
+#endif // WITH_EDITOR
+}
+
 bool UInterchangeMaterialXTranslator::Translate(UInterchangeBaseNodeContainer& BaseNodeContainer) const
 {
 	bool bIsDocumentValid = false;
-	bool bIsReferencesValid = true;
 
 #if WITH_EDITOR
 	namespace mx = MaterialX;
@@ -82,19 +132,17 @@ bool UInterchangeMaterialXTranslator::Translate(UInterchangeBaseNodeContainer& B
 		mx::readFromXmlFile(Document, TCHAR_TO_UTF8(*Filename));
 		Document->importLibrary(MaterialXLibrary);
 
-		std::string MaterialXMessage;
-		bIsDocumentValid = Document->validate(&MaterialXMessage);
+		// Read the document to be sure that the file is valid (meaning all nodes have their nodedef and typedef well-defined)
+		bIsDocumentValid = ValidateDocument(Document, this);
 		if(!bIsDocumentValid)
 		{
-			UInterchangeResultError_Generic* Message = AddMessage<UInterchangeResultError_Generic>();
-			Message->Text = FText::Format(LOCTEXT("MaterialXDocumentInvalid", "{0}"),
-										  FText::FromString(MaterialXMessage.c_str()));
 			return false;
 		}
 
 		//Update the document by initializing and reorganizing the different nodes and subgraphs
 		FMaterialXBase::UpdateDocumentRecursively(Document);
 
+		// coming to this point we know for sure that the document is valid
 		for(mx::ElementPtr Elem : Document->traverseTree())
 		{
 			//make sure to read only the current file otherwise we'll process the entire library
@@ -107,31 +155,12 @@ bool UInterchangeMaterialXTranslator::Translate(UInterchangeBaseNodeContainer& B
 
 			if(Node)
 			{
-				// Validate that all nodes in the file are strictly respecting their node definition
-				if(!Node->getNodeDef())
-				{
-					UInterchangeResultError_Generic* Message = AddMessage<UInterchangeResultError_Generic>();
-					Message->Text = FText::Format(LOCTEXT("NodeDefNotFound", "<{0}> has no matching NodeDef, aborting import..."),
-												  FText::FromString(Node->getName().c_str()));
-					bIsReferencesValid = false;
-					break;
-				}
-
 				bool bIsMaterialShader = Node->getType() == mx::Type::Material;
 				bool bIsLightShader = Node->getType() == mx::Type::LightShader;
 
 				//The entry point is only surfacematerial or lightshader
 				if(bIsMaterialShader || bIsLightShader)
 				{
-					if(!Node->getTypeDef())
-					{
-						UInterchangeResultError_Generic* Message = AddMessage<UInterchangeResultError_Generic>();
-						Message->Text = FText::Format(LOCTEXT("TypeDefNotFound", "<{0}> has no matching TypeDef, aborting import..."),
-													  FText::FromString(Node->getName().c_str()));
-						bIsReferencesValid = false;
-						break;
-					}
-
 					TSharedPtr<FMaterialXBase> ShaderTranslator = FMaterialXManager::GetInstance().GetShaderTranslator(Node->getCategory().c_str(), BaseNodeContainer);
 					if(ShaderTranslator)
 					{
@@ -144,7 +173,6 @@ bool UInterchangeMaterialXTranslator::Translate(UInterchangeBaseNodeContainer& B
 	catch(std::exception& Exception)
 	{
 		bIsDocumentValid = false;
-		bIsReferencesValid = false;
 		UInterchangeResultError_Generic* Message = AddMessage<UInterchangeResultError_Generic>();
 		Message->Text = FText::Format(LOCTEXT("MaterialXException", "{0}"),
 									  FText::FromString(Exception.what()));
@@ -152,13 +180,13 @@ bool UInterchangeMaterialXTranslator::Translate(UInterchangeBaseNodeContainer& B
 
 #endif // WITH_EDITOR
 
-	if(bIsDocumentValid && bIsReferencesValid)
+	if(bIsDocumentValid)
 	{
 		UInterchangeSourceNode* SourceNode = UInterchangeSourceNode::FindOrCreateUniqueInstance(&BaseNodeContainer);
 		SourceNode->SetCustomImportUnusedMaterial(true);
 	}
 
-	return bIsDocumentValid && bIsReferencesValid;
+	return bIsDocumentValid;
 }
 
 TOptional<UE::Interchange::FImportImage> UInterchangeMaterialXTranslator::GetTexturePayloadData(const FString& PayloadKey, TOptional<FString>& AlternateTexturePath) const
