USD: Re-enable the cvar for import into level via USD interchange for 5.6

diff --git a/Engine/Plugins/Interchange/Runtime/Source/Import/Private/Usd/InterchangeUsdTranslator.cpp b/Engine/Plugins/Interchange/Runtime/Source/Import/Private/Usd/InterchangeUsdTranslator.cpp
--- a/Engine/Plugins/Interchange/Runtime/Source/Import/Private/Usd/InterchangeUsdTranslator.cpp
+++ b/Engine/Plugins/Interchange/Runtime/Source/Import/Private/Usd/InterchangeUsdTranslator.cpp
@@ -68,11 +68,11 @@ static FAutoConsoleVariableRef CVarInterchangeEnableUSDImport(
 
 static bool GInterchangeEnableUSDLevelImport = false;
 // Import into level via USD Interchange is disabled for 5.5 as it's still a work in progress
-// static FAutoConsoleVariableRef CVarInterchangeEnableUSDLevelImport(
-// 	TEXT("Interchange.FeatureFlags.Import.USD.ToLevel"),
-// 	GInterchangeEnableUSDLevelImport,
-// 	TEXT("Whether support for USD level import is enabled.")
-// );
+static FAutoConsoleVariableRef CVarInterchangeEnableUSDLevelImport(
+	TEXT("Interchange.FeatureFlags.Import.USD.ToLevel"),
+	GInterchangeEnableUSDLevelImport,
+	TEXT("Whether support for USD level import is enabled.")
+);
 
 namespace UE::InterchangeUsdTranslator::Private
 {
