[BUGFIX] UDIM Texture should fail when there is a mismatch in number of channels
- Added a conversion to Grayscale in the Texture Factory to correctly convert to RGBA from single channel formats

diff --git a/Engine/Plugins/Interchange/Runtime/Source/Import/Private/Texture/InterchangeTextureFactory.cpp b/Engine/Plugins/Interchange/Runtime/Source/Import/Private/Texture/InterchangeTextureFactory.cpp
--- a/Engine/Plugins/Interchange/Runtime/Source/Import/Private/Texture/InterchangeTextureFactory.cpp
+++ b/Engine/Plugins/Interchange/Runtime/Source/Import/Private/Texture/InterchangeTextureFactory.cpp
@@ -598,7 +598,26 @@ namespace UE::Interchange::Private::InterchangeTextureFactory
 						
 						ERawImageFormat::Type BlockedImageRawFormat = FImageCoreUtils::ConvertToRawImageFormat(BlockedImage.Format);
 						FImage DestImage(Image.SizeX, Image.SizeY, BlockedImageRawFormat, BlockedImage.bSRGB ? EGammaSpace::sRGB : EGammaSpace::Linear);
-						FImageCore::CopyImage(SourceImage, DestImage);
+
+						{
+							FImage GrayScaleImage;
+							switch (ImageRawFormat)
+							{
+							case ERawImageFormat::R16F:
+								// falls through
+							case ERawImageFormat::R32F:
+							{
+								GrayScaleImage.Init(Image.SizeX, Image.SizeY, ERawImageFormat::G16, Image.bSRGB ? EGammaSpace::sRGB : EGammaSpace::Linear);
+								FImageCore::CopyImage(SourceImage, GrayScaleImage);
+								SourceImage = GrayScaleImage;
+								break;
+							}
+							default:
+								break;
+							}
+
+							FImageCore::CopyImage(SourceImage, DestImage);
+						}
 
 						Image.RawData = MakeUniqueBufferFromArray(MoveTemp(DestImage.RawData));
 						Image.Format = BlockedImage.Format;
