USD:
- Create a shadergraph using MaterialX sdk
- Adding the same behavior from InterchangeMaterialXTranslator in the GetTexturePayload function in case of the textures comes from the MaterialX API

diff --git a/Engine/Plugins/Interchange/Runtime/Source/Import/Private/Usd/InterchangeUsdTranslator.cpp b/Engine/Plugins/Interchange/Runtime/Source/Import/Private/Usd/InterchangeUsdTranslator.cpp
--- a/Engine/Plugins/Interchange/Runtime/Source/Import/Private/Usd/InterchangeUsdTranslator.cpp
+++ b/Engine/Plugins/Interchange/Runtime/Source/Import/Private/Usd/InterchangeUsdTranslator.cpp
@@ -47,6 +47,12 @@
 #include "StaticMeshAttributes.h"
 #include "UDIMUtilities.h"
 
+#include "UObject/GCObjectScopeGuard.h"
+
+#include "MaterialX/MaterialXUtils/MaterialXManager.h"
+#include "MaterialX/MaterialXUtils/MaterialXBase.h"
+#include "MaterialXFormat/XmlIo.h"
+
 #if USE_USD_SDK
 #include "USDIncludesStart.h"
 #include "pxr/usd/usdGeom/tokens.h"
@@ -142,6 +148,19 @@ namespace UE::InterchangeUsdTranslator::Private
 
 		void AddMeshNode(const UE::FUsdPrim& Prim, UInterchangeBaseNodeContainer& NodeContainer, const FTraversalInfo& Info);
 
+		/**
+		 * Search for a MaterialX file embedded in the file, in case there's none, create a ShaderGraph using MaterialX
+		 */
+		bool AddMaterialXShaderGraph(
+			const UE::FUsdPrim& Prim,
+			UInterchangeUsdTranslatorSettings* TranslatorSettings,
+			UInterchangeBaseNodeContainer& NodeContainer);
+
+		/**
+		 * Set the material slot dependencies on the meshes, if bIterateShaderGraphs is true, it will iterate over the shader graphs to rename them according to the MaterialPrimName
+		 */
+		void SetMaterialSlotDependencies(const FString& MaterialPrimName, const FString& MaterialUid, UInterchangeBaseNodeContainer& NodeContainer, bool bIterateShaderGraphs = false);
+
 	public:
 		// We have to keep a stage reference so that we can parse the payloads after Translate() complete.
 		// ReleaseSource() clears this member, once translation is complete.
@@ -704,91 +723,137 @@ namespace UE::InterchangeUsdTranslator::Private
 		const FString* ParameterName = nullptr;
 	};
 
-	void UInterchangeUSDTranslatorImpl::AddMaterialNode(
-		const UE::FUsdPrim& Prim,
-		UInterchangeUsdTranslatorSettings* TranslatorSettings,
-		UInterchangeBaseNodeContainer& NodeContainer
-	)
-	{
-		FString PrimPath = Prim.GetPrimPath().GetString();
-		FString MaterialUid = MaterialPrefix + PrimPath;
-		FString MaterialPrimName(Prim.GetName().ToString());
 
-		// Check if Node already exist with this ID
-		if (const UInterchangeMaterialInstanceNode* Node = Cast<UInterchangeMaterialInstanceNode>(NodeContainer.GetNode(MaterialUid)))
+	void UInterchangeUSDTranslatorImpl::SetMaterialSlotDependencies(const FString& MaterialPrimName, const FString& MaterialUid, UInterchangeBaseNodeContainer& NodeContainer, bool bIterateShaderGraphs)
+	{
+		if(bIterateShaderGraphs)
 		{
-			return;
+			// The material from the MaterialXTranslator doesn't have the same UID, both the PrimPath have the same name but not the same path
+			// We need to retrieve that name (which is the Material name) in the translator, then we can map it to the right mesh
+			NodeContainer.BreakableIterateNodesOfType<UInterchangeShaderGraphNode>(
+				[this, &MaterialPrimName, &MaterialUid](const FString&, UInterchangeShaderGraphNode* ShaderGraphNode)
+				{
+					FString UID = ShaderGraphNode->GetUniqueID();
+					if(FPaths::GetBaseFilename(UID) == MaterialPrimName)
+					{
+						MaterialUidToActualNodeUid.Add(MaterialUid, UID);
+						return true;
+					}
+					else
+					{
+						return false;
+					}
+				}
+			);
 		}
 
-		auto SetMaterialSlotDependencies = [this, &MaterialUid]()
+		// Now we need to check if we have to set the slot of the mesh nodes here
+		if(TArray<FMaterialSlotMesh>* SlotMeshes = PrimPathToSlotMeshNodes.Find(MaterialUid))
 		{
-			// Now we need to check if we have to set the slot of the mesh nodes here
-			if (TArray<FMaterialSlotMesh>* SlotMeshes = PrimPathToSlotMeshNodes.Find(MaterialUid))
+			if(FString* NewMaterialUID = MaterialUidToActualNodeUid.Find(MaterialUid))
 			{
-				if (FString* NewMaterialUID = MaterialUidToActualNodeUid.Find(MaterialUid))
+				for(const FMaterialSlotMesh& MaterialSlotMesh : *SlotMeshes)
 				{
-					for (const FMaterialSlotMesh& MaterialSlotMesh : *SlotMeshes)
+					if(!MaterialSlotMesh.MeshNode->GetSlotMaterialDependencyUid(MaterialSlotMesh.MaterialSlotName, *NewMaterialUID))
 					{
-						if (!MaterialSlotMesh.MeshNode->GetSlotMaterialDependencyUid(MaterialSlotMesh.MaterialSlotName, *NewMaterialUID))
-						{
-							MaterialSlotMesh.MeshNode->SetSlotMaterialDependencyUid(MaterialSlotMesh.MaterialSlotName, *NewMaterialUID);
-						}
+						MaterialSlotMesh.MeshNode->SetSlotMaterialDependencyUid(MaterialSlotMesh.MaterialSlotName, *NewMaterialUID);
 					}
 				}
 			}
-		};
+		}
+	};
+
+	bool UInterchangeUSDTranslatorImpl::AddMaterialXShaderGraph(
+		const UE::FUsdPrim& Prim,
+		UInterchangeUsdTranslatorSettings* TranslatorSettings,
+		UInterchangeBaseNodeContainer& NodeContainer
+	)
+	{
+		FString PrimPath = Prim.GetPrimPath().GetString();
+		FString MaterialUid = MaterialPrefix + PrimPath;
+		FString MaterialPrimName(Prim.GetName().ToString());
 
 		FName RenderContext = TranslatorSettings ? TranslatorSettings->RenderContext : UnrealIdentifiers::UniversalRenderContext;
 
 		// Check for any references of MaterialX
 #if WITH_EDITOR
-		if (RenderContext == UnrealIdentifiers::MaterialXRenderContext)
+		if(RenderContext == UnrealIdentifiers::MaterialXRenderContext)
 		{
 			TArray<FString> FilePaths = UsdUtils::GetMaterialXFilePaths(Prim);
-			for (const FString& File : FilePaths)
+			for(const FString& File : FilePaths)
 			{
 				// the file has already been handled no need to do a Translate again
-				if (!Translators.Find(File))
+				if(!Translators.Find(File))
 				{
 					UInterchangeManager& InterchangeManager = UInterchangeManager::GetInterchangeManager();
 					UInterchangeSourceData* SourceData = UInterchangeManager::CreateSourceData(File);
 
 					UInterchangeTranslatorBase* Translator = InterchangeManager.GetTranslatorForSourceData(SourceData);
 					// check on the Translator, it might return nullptr in case of reimport
-					if (Translator)
+					if(Translator)
 					{
 						Translator->Translate(NodeContainer);
 						Translators.Add(File, Translator);
 					}
 				}
+			}
+
+			if(!FilePaths.IsEmpty())
+			{
+				return true;
+			}
 
-				// The material from the MaterialXTranslator doesn't have the same UID, both the PrimPath have the same name but not the same path
-				// We need to retrieve that name (which is the Material name) in the translator, then we can map it to the right mesh
-				NodeContainer.BreakableIterateNodesOfType<UInterchangeShaderGraphNode>(
-					[this, &MaterialPrimName, &MaterialUid](const FString&, UInterchangeShaderGraphNode* ShaderGraphNode)
+			bool bHasShaderGraph = false;			
+			if(MaterialX::DocumentPtr Document = UsdToUnreal::CreateSurfaceMaterial(Prim, TranslatorSettings ? *RenderContext.ToString() : nullptr))
+			{
+				if(MaterialX::NodePtr SurfaceMaterial = Document->getNode(TCHAR_TO_ANSI(*Prim.GetName().ToString())))
+				{
+					try
 					{
-						FString ShaderGraphUid = ShaderGraphNode->GetUniqueID();
-						if (FPaths::GetBaseFilename(ShaderGraphUid) == MaterialPrimName)
+						TSharedPtr<FMaterialXBase> ShaderTranslator = FMaterialXManager::GetInstance().GetShaderTranslator(SurfaceMaterial->getCategory().c_str(), NodeContainer);
+						bHasShaderGraph = ShaderTranslator.IsValid();
+						if(bHasShaderGraph)
 						{
-							MaterialUidToActualNodeUid.Add(MaterialUid, ShaderGraphUid);
-							return true;
-						}
-						else
-						{
-							return false;
+							ShaderTranslator->Translate(SurfaceMaterial);
 						}
 					}
-				);
+					catch(std::exception& Exception)
+					{
+						UE_LOG(LogUsd, Error, TEXT("%s"), ANSI_TO_TCHAR(Exception.what()));
+					}
+				}
+
+				return bHasShaderGraph;
 			}
+		}
+#endif // WITH_EDITOR
+		return false;
+	}
 
-			SetMaterialSlotDependencies();
+	void UInterchangeUSDTranslatorImpl::AddMaterialNode(
+		const UE::FUsdPrim& Prim,
+		UInterchangeUsdTranslatorSettings* TranslatorSettings,
+		UInterchangeBaseNodeContainer& NodeContainer
+	)
+	{
+		FString PrimPath = Prim.GetPrimPath().GetString();
+		FString MaterialUid = MaterialPrefix + PrimPath;
+		FString MaterialPrimName(Prim.GetName().ToString());
 
-			if (!FilePaths.IsEmpty())
-			{
-				return;
-			}
+		// Check if Node already exist with this ID
+		if (const UInterchangeMaterialInstanceNode* Node = Cast<UInterchangeMaterialInstanceNode>(NodeContainer.GetNode(MaterialUid)))
+		{
+			return;
 		}
-#endif	  // WITH_EDITOR
+
+		// we only create material instances if we didn't find any MaterialX instances (including a shader graph)
+		if(AddMaterialXShaderGraph(Prim, TranslatorSettings, NodeContainer))
+		{
+			SetMaterialSlotDependencies(MaterialPrimName, MaterialUid, NodeContainer, true);
+			return;
+		}
+
+		FName RenderContext = TranslatorSettings ? TranslatorSettings->RenderContext : UnrealIdentifiers::UniversalRenderContext;
 
 		if (RenderContext == UnrealIdentifiers::UnrealRenderContext)
 		{
@@ -810,7 +875,7 @@ namespace UE::InterchangeUsdTranslator::Private
 
 		// Set the material instance node to the correct mesh nodes
 		MaterialUidToActualNodeUid.Add(MaterialUid, MaterialUid);
-		SetMaterialSlotDependencies();
+		SetMaterialSlotDependencies(MaterialPrimName, MaterialUid, NodeContainer);
 
 		UsdToUnreal::FUsdPreviewSurfaceMaterialData MaterialData;
 		const bool bSuccess = UsdToUnreal::ConvertMaterial(Prim, MaterialData, *RenderContext.ToString());
@@ -2578,6 +2643,49 @@ TOptional<UE::Interchange::FMeshPayloadData> UInterchangeUSDTranslator::GetMeshP
 	return Result;
 }
 
+TOptional<UE::Interchange::FImportImage> GetMaterialXTexturePayloadData(const FString& PayloadKey, TOptional<FString>& AlternateTexturePath, const UInterchangeTranslatorBase* Translator)
+{
+	FString Filename = PayloadKey;
+	TextureCompressionSettings CompressionSettings = TextureCompressionSettings::TC_Default;
+
+#if WITH_EDITOR
+	if(int32 IndexTextureCompression; PayloadKey.FindChar(FMaterialXManager::TexturePayloadSeparator, IndexTextureCompression))
+	{
+		Filename = PayloadKey.Mid(0, IndexTextureCompression);
+		CompressionSettings = TextureCompressionSettings(FCString::Atoi(*PayloadKey.Mid(IndexTextureCompression + 1)));
+	}
+#endif
+	UInterchangeSourceData* PayloadSourceData = UInterchangeManager::GetInterchangeManager().CreateSourceData(Filename);
+	FGCObjectScopeGuard ScopedSourceData(PayloadSourceData);
+
+	if(!PayloadSourceData)
+	{
+		return TOptional<UE::Interchange::FImportImage>();
+	}
+
+	UInterchangeTranslatorBase* SourceTranslator = UInterchangeManager::GetInterchangeManager().GetTranslatorForSourceData(PayloadSourceData);
+	FGCObjectScopeGuard ScopedSourceTranslator(SourceTranslator);
+	const IInterchangeTexturePayloadInterface* TextureTranslator = Cast<IInterchangeTexturePayloadInterface>(SourceTranslator);
+
+	if(!ensure(TextureTranslator))
+	{
+		return TOptional<UE::Interchange::FImportImage>();
+	}
+
+	SourceTranslator->SetResultsContainer(Translator->Results);
+
+	AlternateTexturePath = Filename;
+
+	TOptional<UE::Interchange::FImportImage> TexturePayloadData = TextureTranslator->GetTexturePayloadData(PayloadKey, AlternateTexturePath);
+
+	if(TexturePayloadData.IsSet())
+	{
+		TexturePayloadData.GetValue().CompressionSettings = CompressionSettings;
+	}
+
+	return TexturePayloadData;
+}
+
 TOptional<UE::Interchange::FImportImage> UInterchangeUSDTranslator::GetTexturePayloadData(
 	const FString& PayloadKey,
 	TOptional<FString>& AlternateTexturePath
@@ -2596,7 +2704,7 @@ TOptional<UE::Interchange::FImportImage> UInterchangeUSDTranslator::GetTexturePa
 		// Defer back to another translator to actually parse the texture raw data
 		UE::Interchange::Private::FScopedTranslator ScopedTranslator(FilePath, Results);
 		const IInterchangeTexturePayloadInterface* TextureTranslator = ScopedTranslator.GetPayLoadInterface<IInterchangeTexturePayloadInterface>();
-		if (ensure(TextureTranslator))
+		if (TextureTranslator)
 		{
 			AlternateTexturePath = FilePath;
 
@@ -2617,6 +2725,7 @@ TOptional<UE::Interchange::FImportImage> UInterchangeUSDTranslator::GetTexturePa
 	// We did not find a suitable Payload in USD Translator, let's find one in one of the Translators (MaterialX for the moment)
 	// The best way would be to have a direct association between the payload and the right Translator, but we don't have a suitable way of knowing
 	// which Payload belongs to which Translator So let's just loop over them all
+	bool bTranslatorPayload = false;
 	for (const TPair<FString, UInterchangeTranslatorBase*>& Pair : Impl->Translators)
 	{
 		if (IInterchangeTexturePayloadInterface* TexturePayloadInterface = Cast<IInterchangeTexturePayloadInterface>(Pair.Value))
@@ -2624,11 +2733,18 @@ TOptional<UE::Interchange::FImportImage> UInterchangeUSDTranslator::GetTexturePa
 			TexturePayloadData = TexturePayloadInterface->GetTexturePayloadData(PayloadKey, AlternateTexturePath);
 			if (TexturePayloadData)
 			{
+				bTranslatorPayload = true;
 				break;
 			}
 		}
 	}
 
+	//if we couldn't find a texture in either the USD Translator nor the Translators, then it's most likely coming from reading an mtlx in memory, copying the behavior from InterchangeMaterialXTranslator
+	if(!bTranslatorPayload)
+	{
+		TexturePayloadData = GetMaterialXTexturePayloadData(PayloadKey, AlternateTexturePath, this);
+	}
+
 	return TexturePayloadData;
 }
 
diff --git a/Engine/Plugins/Runtime/USDCore/Source/USDUtilities/Private/USDShadeConversion.cpp b/Engine/Plugins/Runtime/USDCore/Source/USDUtilities/Private/USDShadeConversion.cpp
--- a/Engine/Plugins/Runtime/USDCore/Source/USDUtilities/Private/USDShadeConversion.cpp
+++ b/Engine/Plugins/Runtime/USDCore/Source/USDUtilities/Private/USDShadeConversion.cpp
@@ -2,6 +2,8 @@
 
 #include "USDShadeConversion.h"
 
+#include "MaterialXFormat/Util.h"
+
 #if USE_USD_SDK
 
 #include "UnrealUSDWrapper.h"
@@ -68,6 +70,8 @@ class UTextureFactory;
 
 #define LOCTEXT_NAMESPACE "USDShadeConversion"
 
+namespace mx = MaterialX;
+
 namespace UE::UsdShadeConversion::Private
 {
 #if WITH_EDITOR
@@ -2138,6 +2142,306 @@ bool UsdToUnreal::FTextureParameterValue::GetSRGBValue() const
 	return bSRGB.Get(bFallbackSRGB);
 }
 
+namespace UE::UsdShadeConversion::Private
+{
+	std::unordered_set<std::string> MaterialXTypes
+	{
+		"boolean",
+		"integer",
+		"float",
+		"color3",
+		"color4",
+		"vector2",
+		"vector3",
+		"vector4",
+		"matrix33",
+		"matrix44",
+		"string",
+		"filename",
+		"geomname",
+		"surfaceshader",
+		"displacementshader",
+		"volumeshader",
+		"lightshader",
+		"material",
+		"none",
+		"integerarray",
+		"floatarray",
+		"color3array",
+		"color4array",
+		"vector2array",
+		"vector3array",
+		"vector4array",
+		"stringarray",
+		"geomnamearray",
+	};
+
+	std::unordered_map<std::string, std::string> UsdToMaterialXTypes
+	{
+		{"color3f", "color3"},
+		{"color4f", "color4"},
+		{"float2", "vector2"},
+		{"vector3f", "vector3"},
+		{"vector4f", "vector4"},
+		{"int", "integer"},
+	};
+
+	bool ComputeShaderGraph(mx::NodePtr Node, const pxr::UsdShadeShader& Shader)
+	{
+		if (!Node)
+		{
+			return false;
+		}
+
+		mx::DocumentPtr Document = Node->getDocument();
+		TUsdStore<std::vector<pxr::UsdShadeInput>> UsdInputs = Shader.GetInputs();
+
+		for (const pxr::UsdShadeInput& UsdInput : *UsdInputs)
+		{
+			std::string InputName = UsdInput.GetBaseName().GetString();
+			std::string InputType = UsdInput.GetTypeName().GetAsToken().GetString();
+
+			if (MaterialXTypes.find(InputType) == MaterialXTypes.end())
+			{
+				if(auto It = UsdToMaterialXTypes.find(InputType); It != UsdToMaterialXTypes.end())
+				{
+					InputType = It->second;
+				}
+				else if (InputName == "file") // we're dealing with a texture, USD calls the type "asset" for the input
+				{
+					InputType = "filename";
+				}
+				else
+				{
+					UE_LOG(LogUsd, Warning, TEXT("Couldn't find a supported type for (%s)."), ANSI_TO_TCHAR(InputType.c_str()));
+				}
+			}
+
+			mx::InputPtr Input = Node->addInput(InputName, InputType);
+
+			TUsdStore<pxr::UsdShadeInput::SourceInfoVector> ConnectedSources = UsdInput.GetConnectedSources();
+			bool bHasValue = true;
+			for (const pxr::UsdShadeConnectionSourceInfo& Source : *ConnectedSources)
+			{
+				bHasValue = false;
+				pxr::UsdShadeShader ConnectedShader{ Source.source.GetPrim() };
+
+				pxr::TfToken NodeDefTokenId;
+				ConnectedShader.GetShaderId(&NodeDefTokenId);
+				mx::NodeDefPtr NodeDef = Document->getNodeDef(NodeDefTokenId.GetString());
+
+				// if we couldn't find a nodedef just stop the graph here and propagate nullptr
+				if (!NodeDef)
+				{
+					Node = nullptr;
+					UE_LOG(LogUsd, Warning, TEXT("Couldn't find a nodedef for (%s)."), ANSI_TO_TCHAR(NodeDefTokenId.GetString().c_str()));
+					return false;
+				}
+
+				std::string ConnectedShaderName = ConnectedShader.GetPrim().GetName().GetString();
+
+				// Add the node only once
+				if (!Document->getNode(ConnectedShaderName))
+				{
+					mx::NodePtr ConnectedNode = Document->addNodeInstance(NodeDef, ConnectedShaderName);
+					Input->setConnectedNode(ConnectedNode);
+
+					// add the output of the connected node to the current input
+					TUsdStore<std::vector<pxr::UsdShadeOutput>> UsdOutputs = ConnectedShader.GetOutputs();
+					for (const pxr::UsdShadeOutput& UsdOutput : *UsdOutputs)
+					{
+						//no need to add the default output
+						if (std::string OutputName = UsdOutput.GetBaseName().GetString(); OutputName != "out")
+						{
+							Input->setOutputString(OutputName);
+						}
+					}
+
+					ComputeShaderGraph(ConnectedNode, ConnectedShader);
+				}
+			}
+
+			if (bHasValue)
+			{
+				auto SetInputValue = [&UsdInput, Node, Input](auto Type)
+					{
+						using MtlxType = decltype(Type);
+
+						// complex template to just retrieve the right type, basically a mapping between Usd and Mtlx for Vector types, otherwise we just take the type (float, int, string)
+						using UsdType = std::conditional_t<
+							std::is_same_v<MtlxType, mx::Color4> || std::is_same_v<MtlxType, mx::Vector4>,
+							pxr::GfVec4f,
+								std::conditional_t<
+								std::is_same_v<MtlxType, mx::Color3> || std::is_same_v<MtlxType, mx::Vector3>,
+									pxr::GfVec3f,
+									std::conditional_t<
+										std::is_same_v<MtlxType, mx::Vector2>,
+										pxr::GfVec2f,
+							MtlxType>>>;
+
+						if (UsdType Value; UsdInput.Get<UsdType>(&Value))
+						{
+							if constexpr(std::is_same_v<UsdType, pxr::GfVec4f> || std::is_same_v<UsdType, pxr::GfVec3f> || std::is_same_v<UsdType, pxr::GfVec2f>)
+							{
+								Input->setValue(MtlxType{ Value.data(), Value.data() + UsdType::dimension });
+							}
+							else
+							{
+								Input->setValue(Value);
+							}
+						}
+						else
+						{
+							mx::NodeDefPtr NodeDef = Node->getNodeDef();
+							mx::InputPtr ActiveInput = NodeDef->getActiveInput(Input->getName());
+							MtlxType DefaultValue = ActiveInput->getDefaultValue()->asA<MtlxType>();
+							Input->setValue(DefaultValue);
+						}
+					};
+
+				if (InputType == "float")
+				{
+					SetInputValue(float{});
+				}
+				else if (InputType == "integer")
+				{
+					SetInputValue(int32{});
+				}
+				else if (InputType == "vector2")
+				{
+					SetInputValue(mx::Vector2{});
+				}
+				else if (InputType == "vector3")
+				{
+					SetInputValue(mx::Vector3{});
+				}
+				else if (InputType == "vector4")
+				{
+					SetInputValue(mx::Vector4{});
+				}
+				else if (InputType == "color3")
+				{
+					SetInputValue(mx::Color3{});
+				}
+				else if (InputType == "color4")
+				{
+					SetInputValue(mx::Color4{});
+				}
+				else if (InputType == "string")
+				{
+					SetInputValue(std::string{});
+				}
+				else if (InputType == "filename")
+				{
+					pxr::SdfAssetPath AssetPath;
+					UsdInput.Get<pxr::SdfAssetPath>(&AssetPath);
+					TUsdStore<UE::FSdfLayer> Layer = UsdUtils::FindLayerForAttribute(UsdInput.GetAttr(), pxr::UsdTimeCode::Default().GetValue());
+
+					FString PathToResolve =
+						Layer.Get()
+						? UsdToUnreal::ConvertString(TUsdStore<std::string>(pxr::SdfComputeAssetPathRelativeToLayer(Layer.Get(), AssetPath.GetAssetPath())).Get())
+						: AssetPath.GetAssetPath().c_str();
+					Input->setValue(std::string{ TCHAR_TO_ANSI(*PathToResolve) });
+				}
+				else
+				{
+					UE_LOG(LogUsd, Warning, TEXT("Couldn't find a value type for (%s)."), ANSI_TO_TCHAR(InputType.c_str()));
+				}
+			}
+		}
+
+		return true;
+	}
+}
+
+mx::DocumentPtr UsdToUnreal::CreateSurfaceMaterial(const pxr::UsdPrim& UsdShadeMaterialPrim, const TCHAR* RenderContext)
+{
+	pxr::UsdShadeMaterial UsdShadeMaterial{ UsdShadeMaterialPrim };
+
+	if (!UsdShadeMaterial)
+	{
+		return nullptr;
+	}
+
+	pxr::TfToken RenderContextToken = pxr::UsdShadeTokens->universalRenderContext;
+	if (RenderContext)
+	{
+		pxr::TfToken ProvidedRenderContextToken = UnrealToUsd::ConvertToken(RenderContext).Get();
+		RenderContextToken = ProvidedRenderContextToken;
+	}
+
+	pxr::UsdShadeShader SurfaceShader = UsdShadeMaterial.ComputeSurfaceSource(RenderContextToken);
+	pxr::UsdShadeShader DisplacementShader = UsdShadeMaterial.ComputeDisplacementSource(RenderContextToken);
+
+	if (!SurfaceShader)
+	{
+		return nullptr;
+	}
+
+	mx::DocumentPtr Document = mx::createDocument();
+	{
+		mx::DocumentPtr MaterialXLibrary = mx::createDocument();
+		mx::FileSearchPath MaterialXFolder{ TCHAR_TO_UTF8(*FPaths::Combine(
+			FPaths::EngineDir(),
+			TEXT("Binaries"),
+			TEXT("ThirdParty"),
+			TEXT("MaterialX"))) };
+		mx::StringSet LoadedLibs = mx::loadLibraries({ "libraries" }, MaterialXFolder, MaterialXLibrary);
+		Document->importLibrary(MaterialXLibrary);
+	}
+
+	auto CreateShaderGraph = [](const pxr::UsdShadeShader& SurfaceShader, mx::DocumentPtr Document) -> mx::NodePtr
+	{
+		pxr::TfToken NodeDefTokenId;
+		SurfaceShader.GetShaderId(&NodeDefTokenId);
+		mx::NodeDefPtr NodeDef = Document->getNodeDef(NodeDefTokenId.GetString());
+
+		// no node definition, just return
+		if (!NodeDef)
+		{
+			return nullptr;
+		}
+
+		mx::NodePtr MtlxSurfaceShader = Document->addNodeInstance(NodeDef, SurfaceShader.GetPrim().GetName().GetString());
+
+		//if the shadergraph is not valid (probably caused by a missing nodedef or wrong type), just invalide it all
+		if (!UE::UsdShadeConversion::Private::ComputeShaderGraph(MtlxSurfaceShader, SurfaceShader))
+		{
+			MtlxSurfaceShader = nullptr;
+		}
+
+		return MtlxSurfaceShader;
+	};
+
+	//Create the shader graph for the surface shader and displacement shader	
+	if (mx::NodePtr MaterialXSurfaceShader = CreateShaderGraph(SurfaceShader, Document))
+	{
+		mx::NodeDefPtr SurfaceMaterialNodeDef = Document->getNodeDef("ND_surfacematerial");
+		mx::NodePtr SurfaceMaterial = Document->addNodeInstance(SurfaceMaterialNodeDef, UsdShadeMaterial.GetPrim().GetName().GetString());
+
+		SurfaceMaterial->setConnectedNode("surfaceshader", MaterialXSurfaceShader);
+
+		//the displacement shader is not necessarily required
+		if (DisplacementShader)
+		{
+			if(mx::NodePtr MaterialXDisplacementShader = CreateShaderGraph(DisplacementShader, Document); MaterialXDisplacementShader)
+			{
+				SurfaceMaterial->setConnectedNode("displacementshader", MaterialXDisplacementShader);
+			}
+		}
+
+		std::string ErrorMessage;
+		SurfaceMaterial->validate(&ErrorMessage);
+		if (!ErrorMessage.empty())
+		{
+			UE_LOG(LogUsd, Error, TEXT("USD MaterialX: %s."), UTF8_TO_TCHAR(ErrorMessage.c_str()));
+			SurfaceMaterial = nullptr;
+		}
+	}
+
+	return Document;
+}
+
 bool UsdToUnreal::ConvertMaterial(
 	const pxr::UsdPrim& UsdShadeMaterialPrim,
 	FUsdPreviewSurfaceMaterialData& OutMaterialData,
@@ -2162,6 +2466,7 @@ bool UsdToUnreal::ConvertMaterial(
 	}
 
 	pxr::UsdShadeShader SurfaceShader = UsdShadeMaterial.ComputeSurfaceSource(RenderContextToken);
+
 	if (!SurfaceShader)
 	{
 		return false;
diff --git a/Engine/Plugins/Runtime/USDCore/Source/USDUtilities/Public/USDShadeConversion.h b/Engine/Plugins/Runtime/USDCore/Source/USDUtilities/Public/USDShadeConversion.h
--- a/Engine/Plugins/Runtime/USDCore/Source/USDUtilities/Public/USDShadeConversion.h
+++ b/Engine/Plugins/Runtime/USDCore/Source/USDUtilities/Public/USDShadeConversion.h
@@ -20,6 +20,8 @@
 #include "pxr/usd/usdShade/tokens.h"
 #include "USDIncludesEnd.h"
 
+#include "MaterialXCore/Document.h"
+
 PXR_NAMESPACE_OPEN_SCOPE
 	class UsdAttribute;
 	class UsdGeomMesh;
@@ -130,6 +132,14 @@ namespace UsdToUnreal
 		bool bShareAssetsForIdenticalPrims = true
 	);
 
+	/** Create a MaterialX Surface material. If the created shader graph is not valid according to MaterialX spec, return nullptr
+	 * We return the Document in order for the caller to search for the surface material, by just looking for the name will return the correct surface material.
+	 * We need to return the document instead of the node directly otherwise the Document will go out of scope and create an orphaned shader graph
+	 * @param UsdShadeMaterialPrim - The prim used to generate the shader graph, should be a valid pxr::UsdShadeMaterial
+	 * @param RenderContext - Which render context output to read from the UsdShadeMaterial
+	 */
+	USDUTILITIES_API MaterialX::DocumentPtr CreateSurfaceMaterial(const pxr::UsdPrim& UsdShadeMaterialPrim, const TCHAR* RenderContext);
+
 	/**
 	 * Attemps to assign the values of the surface shader inputs to the MaterialInstance parameters by matching the inputs display names to the
 	 * parameters names.
diff --git a/Engine/Plugins/Runtime/USDCore/Source/USDUtilities/USDUtilities.Build.cs b/Engine/Plugins/Runtime/USDCore/Source/USDUtilities/USDUtilities.Build.cs
--- a/Engine/Plugins/Runtime/USDCore/Source/USDUtilities/USDUtilities.Build.cs
+++ b/Engine/Plugins/Runtime/USDCore/Source/USDUtilities/USDUtilities.Build.cs
@@ -17,6 +17,7 @@ namespace UnrealBuildTool.Rules
 					"Boost",
 					"Core",
 					"CoreUObject",
+					"MaterialX",
 					"UnrealUSDWrapper",
 					"USDClasses", // So that consumers can also include IUsdClassesModule for the new definition of FDisplayColorMaterial
 				}
