[BUGFIX] Interchange: Do not alter lod screesize settings during re-import
- LOD Properties are hidden during reimport and are ignored by pipeline and factories internally

diff --git a/Engine/Plugins/Interchange/Runtime/Source/Import/Private/Mesh/InterchangeStaticMeshFactory.cpp b/Engine/Plugins/Interchange/Runtime/Source/Import/Private/Mesh/InterchangeStaticMeshFactory.cpp
--- a/Engine/Plugins/Interchange/Runtime/Source/Import/Private/Mesh/InterchangeStaticMeshFactory.cpp
+++ b/Engine/Plugins/Interchange/Runtime/Source/Import/Private/Mesh/InterchangeStaticMeshFactory.cpp
@@ -1094,18 +1094,22 @@ void UInterchangeStaticMeshFactory::SetupSourceModelsSettings(UStaticMesh& Stati
 	constexpr int32 LODIndex = 0;
 	float PreviousLODScreenSize = UStaticMesh::ComputeLODScreenSize(LODIndex);
 
-	// If no values are provided, then force AutoCompute
-	if (LodScreenSizes.IsEmpty())
+	// No change during reimport
+	if (!bIsAReimport)
 	{
-		bAutoComputeLODScreenSizes = true;
+		// If no values are provided, then force AutoCompute
+		if (LodScreenSizes.IsEmpty())
+		{
+			bAutoComputeLODScreenSizes = true;
+		}
+		StaticMesh.bAutoComputeLODScreenSize = bAutoComputeLODScreenSizes;
 	}
-	StaticMesh.bAutoComputeLODScreenSize = bAutoComputeLODScreenSizes;
 	
 	for (int32 LodIndex = 0; LodIndex < FinalLodCount; ++LodIndex)
 	{
 		FStaticMeshSourceModel& SrcModel = StaticMesh.GetSourceModel(LodIndex);
 
-		if (!bAutoComputeLODScreenSizes)
+		if (!bIsAReimport && !bAutoComputeLODScreenSizes)
 		{
 			if (LodScreenSizes.IsValidIndex(LodIndex))
 			{
diff --git a/Engine/Plugins/Interchange/Runtime/Source/Pipelines/Public/InterchangeGenericMeshPipeline.h b/Engine/Plugins/Interchange/Runtime/Source/Pipelines/Public/InterchangeGenericMeshPipeline.h
--- a/Engine/Plugins/Interchange/Runtime/Source/Pipelines/Public/InterchangeGenericMeshPipeline.h
+++ b/Engine/Plugins/Interchange/Runtime/Source/Pipelines/Public/InterchangeGenericMeshPipeline.h
@@ -53,15 +53,15 @@ public:
 	bool bCombineStaticMeshes = false;
 
 	/** The LOD group that will be assigned to this mesh. */
-	UPROPERTY(EditAnywhere, BlueprintReadWrite, Category = "Static Meshes", meta=(DisplayName = "LOD Group"))
+	UPROPERTY(EditAnywhere, BlueprintReadWrite, Category = "Static Meshes", meta=(DisplayName = "LOD Group"), meta = (ReimportRestrict = "true"))
 	FName LodGroup = NAME_None;
 
 	/** If enabled, LOD Screen Sizes would be auto-computed. */
-	UPROPERTY(EditAnywhere, BlueprintReadWrite, Category = "Static Meshes")
+	UPROPERTY(EditAnywhere, BlueprintReadWrite, Category = "Static Meshes", meta=(ReimportRestrict = "true"))
 	bool bAutoComputeLODScreenSizes = true;
 
 	/** This setting is only used if the Auto Compute LOD Screen Sizes setting is disabled. */
-	UPROPERTY(EditAnywhere, BlueprintReadWrite, Category = "Static Meshes", meta=(EditCondition = "!bAutoComputeLODScreenSizes", DisplayName = "LOD Screen Sizes"))
+	UPROPERTY(EditAnywhere, BlueprintReadWrite, Category = "Static Meshes", meta=(EditCondition = "!bAutoComputeLODScreenSizes", DisplayName = "LOD Screen Sizes", ReimportRestrict="true"))
 	TArray<float> LODScreenSizes;
 
 	UPROPERTY(meta = (DeprecatedProperty, DeprecationMessage = "Use Collision instead."))
