Fix collision stump when the import collision settings is still active when re-importing

#rnx

diff --git a/Engine/Plugins/Interchange/Runtime/Source/Import/Private/Mesh/InterchangeStaticMeshFactory.cpp b/Engine/Plugins/Interchange/Runtime/Source/Import/Private/Mesh/InterchangeStaticMeshFactory.cpp
--- a/Engine/Plugins/Interchange/Runtime/Source/Import/Private/Mesh/InterchangeStaticMeshFactory.cpp
+++ b/Engine/Plugins/Interchange/Runtime/Source/Import/Private/Mesh/InterchangeStaticMeshFactory.cpp
@@ -605,9 +605,6 @@ UInterchangeFactoryBase::FImportAssetResult UInterchangeStaticMeshFactory::Impor
 				{
 					//Let's clean only the imported collisions first in order 
 					// to store the previous editor-generated collisions to re-generate them later in the Game Thread with their properties
-#if WITH_EDITOR
-					StaticMesh->GetBodySetup()->AggGeom.EmptyImportedElements();
-#endif
 					ImportAssetObjectData.AggregateGeom = StaticMesh->GetBodySetup()->AggGeom;
 					StaticMesh->GetBodySetup()->AggGeom.EmptyElements();
 				}
@@ -850,69 +847,53 @@ UInterchangeFactoryBase::FImportAssetResult UInterchangeStaticMeshFactory::EndIm
 	//Let's now re-generate the previous collisions with their properties, only the extents will be updated	
 	if(bReimport)
 	{
-		for(FKBoxElem& BoxElem : ImportAssetObjectData.AggregateGeom.BoxElems)
+		if (!StaticMesh->GetBodySetup())
 		{
-			int32 Index = GenerateBoxAsSimpleCollision(StaticMesh, false);
-			FKBoxElem& NewBoxElem = StaticMesh->GetBodySetup()->AggGeom.BoxElems[Index];
-
-			//Preserve the new extents
-			float X = NewBoxElem.X;
-			float Y = NewBoxElem.Y;
-			float Z = NewBoxElem.Z;
-
-			//Copy the previous settings, we just want the new extents
-			NewBoxElem = BoxElem;
-			NewBoxElem.X = X;
-			NewBoxElem.Y = Y;
-			NewBoxElem.Z = Z;
+			StaticMesh->CreateBodySetup();
 		}
-
-		for(FKSphereElem& SphereElem : ImportAssetObjectData.AggregateGeom.SphereElems)
+		//If we do not have any imported collision, we put back the original collision body setup
+		if (StaticMesh->GetBodySetup()->AggGeom.GetElementCount() == 0)
 		{
-			int32 Index = GenerateSphereAsSimpleCollision(StaticMesh, false);
-			FKSphereElem& NewSphereElem = StaticMesh->GetBodySetup()->AggGeom.SphereElems[Index];
-
-			//Preserve the new extents
-			float Radius = NewSphereElem.Radius;
-
-			//Copy the previous settings, we just want the new extents
-			NewSphereElem = SphereElem;
-			NewSphereElem.Radius = Radius;
+			StaticMesh->GetBodySetup()->AggGeom = ImportAssetObjectData.AggregateGeom;
 		}
-
-		for(FKSphylElem& CapsuleElem: ImportAssetObjectData.AggregateGeom.SphylElems)
+		else
 		{
-			int32 Index = GenerateSphylAsSimpleCollision(StaticMesh, false);
-			FKSphylElem& NewCapsuleElem = StaticMesh->GetBodySetup()->AggGeom.SphylElems[Index];
+			//If there is some collision, we remove the original imported collision and add any editor generated collision
+#if WITH_EDITOR
+			ImportAssetObjectData.AggregateGeom.EmptyImportedElements();
+#endif
+
+			for (FKBoxElem& BoxElem : ImportAssetObjectData.AggregateGeom.BoxElems)
+			{
+				int32 Index = GenerateBoxAsSimpleCollision(StaticMesh, false);
+				FKBoxElem& NewBoxElem = StaticMesh->GetBodySetup()->AggGeom.BoxElems[Index];
+				//Copy element
+				NewBoxElem = BoxElem;
+			}
 
-			//Preserve the new extents
-			float Radius = NewCapsuleElem.Radius;
-			float Length = NewCapsuleElem.Length;
+			for (FKSphereElem& SphereElem : ImportAssetObjectData.AggregateGeom.SphereElems)
+			{
+				int32 Index = GenerateSphereAsSimpleCollision(StaticMesh, false);
+				FKSphereElem& NewSphereElem = StaticMesh->GetBodySetup()->AggGeom.SphereElems[Index];
+				//Copy element
+				NewSphereElem = SphereElem;
+			}
 
-			//Copy the previous settings, we just want the new extents
-			NewCapsuleElem = CapsuleElem;
-			NewCapsuleElem.Radius = Radius;
-			NewCapsuleElem.Length = Length;
-		}
+			for (FKSphylElem& CapsuleElem : ImportAssetObjectData.AggregateGeom.SphylElems)
+			{
+				int32 Index = GenerateSphylAsSimpleCollision(StaticMesh, false);
+				FKSphylElem& NewCapsuleElem = StaticMesh->GetBodySetup()->AggGeom.SphylElems[Index];
+				//Copy element
+				NewCapsuleElem = CapsuleElem;
+			}
 
-		for(FKConvexElem& ConvexElem: ImportAssetObjectData.AggregateGeom.ConvexElems)
-		{
-			int32 Index = GenerateKDopAsSimpleCollision(StaticMesh, TArray<FVector>(KDopDir18, sizeof(KDopDir18) / sizeof(FVector)), false);
-			FKConvexElem& NewConvexElem = StaticMesh->GetBodySetup()->AggGeom.ConvexElems[Index];
-
-			//Preserve the new extents
-			TArray<FVector> VertexData = NewConvexElem.VertexData;
-			TArray<int32> IndexData = NewConvexElem.IndexData;
-			FBox ElemBox = NewConvexElem.ElemBox;
-			Chaos::FConvexPtr ChaosConvex = NewConvexElem.GetChaosConvexMesh();
-			NewConvexElem.ResetChaosConvexMesh();
-
-			//Copy the previous settings, we just want the new extents
-			NewConvexElem = ConvexElem;
-			NewConvexElem.VertexData = VertexData;
-			NewConvexElem.IndexData = IndexData;
-			NewConvexElem.ElemBox = ElemBox;
-			NewConvexElem.SetConvexMeshObject(MoveTemp(ChaosConvex));
+			for (FKConvexElem& ConvexElem : ImportAssetObjectData.AggregateGeom.ConvexElems)
+			{
+				int32 Index = GenerateKDopAsSimpleCollision(StaticMesh, TArray<FVector>(KDopDir18, sizeof(KDopDir18) / sizeof(FVector)), false);
+				FKConvexElem& NewConvexElem = StaticMesh->GetBodySetup()->AggGeom.ConvexElems[Index];
+				//Copy element
+				NewConvexElem = ConvexElem;
+			}
 		}
 	}
 
