Interchange import generic mesh default changes (enabled Nanite by default, disabled generate lightmap UVs and reversed index buffer by default).

diff --git a/Engine/Plugins/Interchange/Runtime/Source/Pipelines/Public/InterchangeGenericMeshPipeline.h b/Engine/Plugins/Interchange/Runtime/Source/Pipelines/Public/InterchangeGenericMeshPipeline.h
--- a/Engine/Plugins/Interchange/Runtime/Source/Pipelines/Public/InterchangeGenericMeshPipeline.h
+++ b/Engine/Plugins/Interchange/Runtime/Source/Pipelines/Public/InterchangeGenericMeshPipeline.h
@@ -97,15 +97,15 @@ public:
 
 	/** If enabled, imported meshes will be rendered by Nanite at runtime. Make sure your meshes and materials meet the requirements for Nanite. */
 	UPROPERTY(EditAnywhere, BlueprintReadWrite, Category = "Static Meshes", meta = (SubCategory = "Build"))
-	bool bBuildNanite = false;
+	bool bBuildNanite = true;
 
 	/** If enabled, builds a reversed index buffer for each static mesh. */
 	UPROPERTY(EditAnywhere, BlueprintReadWrite, Category = "Static Meshes", meta = (SubCategory = "Build"))
-	bool bBuildReversedIndexBuffer = true;
+	bool bBuildReversedIndexBuffer = false;
 	
 	/** If enabled, generates lightmap UVs for each static mesh. */
 	UPROPERTY(EditAnywhere, BlueprintReadWrite, Category = "Static Meshes", meta = (SubCategory = "Build"))
-	bool bGenerateLightmapUVs = true;
+	bool bGenerateLightmapUVs = false;
 	
 	/** 
 	 * Determines whether to generate the distance field treating every triangle hit as a front face.  
