Make sure we can import animsequence when pie is running.

#rnx

diff --git a/Engine/Plugins/Interchange/Runtime/Source/FactoryNodes/Public/InterchangeAnimSequenceFactoryNode.h b/Engine/Plugins/Interchange/Runtime/Source/FactoryNodes/Public/InterchangeAnimSequenceFactoryNode.h
--- a/Engine/Plugins/Interchange/Runtime/Source/FactoryNodes/Public/InterchangeAnimSequenceFactoryNode.h
+++ b/Engine/Plugins/Interchange/Runtime/Source/FactoryNodes/Public/InterchangeAnimSequenceFactoryNode.h
@@ -410,7 +410,7 @@ public:
 	/** Return if the import of the class is allowed at runtime.*/
 	virtual bool IsRuntimeImportAllowed() const override
 	{
-		return false;
+		return true;
 	}
 
 private:
