Interchange MaterialX:
- Add uniqueness on the node that goes into the node container taking the whole hierarchy as a UID
- Caching the UID as attribute in the Element

diff --git a/Engine/Plugins/Interchange/Runtime/Source/Common/Public/MaterialX/InterchangeMaterialXDefinitions.h b/Engine/Plugins/Interchange/Runtime/Source/Common/Public/MaterialX/InterchangeMaterialXDefinitions.h
--- a/Engine/Plugins/Interchange/Runtime/Source/Common/Public/MaterialX/InterchangeMaterialXDefinitions.h
+++ b/Engine/Plugins/Interchange/Runtime/Source/Common/Public/MaterialX/InterchangeMaterialXDefinitions.h
@@ -464,6 +464,7 @@ MATERIALX_NAMESPACE_BEGIN
 		static constexpr const char* IsVisited = "UE:IsVisited";
 		static constexpr const char* NewName = "UE:NewName";
 		static constexpr const char* ParentName = "UE:ParentName";
+		static constexpr const char* UniqueName = "UE:UniqueName";
 	}
 
 	namespace Category
diff --git a/Engine/Plugins/Interchange/Runtime/Source/Import/Private/MaterialX/MaterialXUtils/MaterialXSurfaceShaderAbstract.cpp b/Engine/Plugins/Interchange/Runtime/Source/Import/Private/MaterialX/MaterialXUtils/MaterialXSurfaceShaderAbstract.cpp
--- a/Engine/Plugins/Interchange/Runtime/Source/Import/Private/MaterialX/MaterialXUtils/MaterialXSurfaceShaderAbstract.cpp
+++ b/Engine/Plugins/Interchange/Runtime/Source/Import/Private/MaterialX/MaterialXUtils/MaterialXSurfaceShaderAbstract.cpp
@@ -74,7 +74,7 @@ bool FMaterialXSurfaceShaderAbstract::AddAttributeFromValueOrInterface(MaterialX
 		{
 			using namespace UE::Interchange::Materials::Standard::Nodes;
 
-			UInterchangeShaderNode* SwizzleNode = CreateShaderNode(Input->getParent()->asA<mx::Node>()->getName().c_str() + FString{ TEXT("_Channels_") } + Input->getName().c_str(), Swizzle::Name.ToString());
+			UInterchangeShaderNode* SwizzleNode = CreateShaderNode(Input, Input->getParent()->asA<mx::Node>()->getName().c_str() + FString{ TEXT("_Channels_") } + Input->getName().c_str(), Swizzle::Name.ToString());
 			SwizzleNode->AddStringAttribute(Swizzle::Attributes::Channels.ToString(), Input->getChannels().c_str());
 
 			UInterchangeShaderPortsAPI::ConnectDefaultOuputToInput(ShaderNode, InputChannelName, SwizzleNode->GetUniqueID());
@@ -111,7 +111,7 @@ bool FMaterialXSurfaceShaderAbstract::AddBooleanAttribute(MaterialX::InputPtr In
 		FString NodeName = Node ? Node->getName().c_str() + FString{ TEXT("_") } : FString{};
 		NodeName += Input->getName().c_str();
 
-		UInterchangeShaderNode* StaticBoolParameterNode = CreateShaderNode(NodeName, StaticBoolParameter::Name.ToString());
+		UInterchangeShaderNode* StaticBoolParameterNode = CreateShaderNode(Input, NodeName, StaticBoolParameter::Name.ToString());
 		StaticBoolParameterNode->AddBooleanAttribute(UInterchangeShaderPortsAPI::MakeInputParameterKey(StaticBoolParameter::Attributes::DefaultValue.ToString()), Value);
 		return UInterchangeShaderPortsAPI::ConnectDefaultOuputToInput(ShaderNode, InputChannelName, StaticBoolParameterNode->GetUniqueID());
 	}
@@ -133,7 +133,7 @@ bool FMaterialXSurfaceShaderAbstract::AddFloatAttribute(MaterialX::InputPtr Inpu
 			FString NodeName = Node ? Node->getName().c_str() + FString{ TEXT("_") } : FString{};
 			NodeName += Input->getName().c_str();
 
-			UInterchangeShaderNode* ScalarParameterNode = CreateShaderNode(NodeName, ScalarParameter::Name.ToString());
+			UInterchangeShaderNode* ScalarParameterNode = CreateShaderNode(Input, NodeName, ScalarParameter::Name.ToString());
 			ScalarParameterNode->AddFloatAttribute(UInterchangeShaderPortsAPI::MakeInputParameterKey(ScalarParameter::Attributes::DefaultValue.ToString()), Value);
 			return UInterchangeShaderPortsAPI::ConnectDefaultOuputToInput(ShaderNode, InputChannelName, ScalarParameterNode->GetUniqueID());
 		}
@@ -156,7 +156,7 @@ bool FMaterialXSurfaceShaderAbstract::AddLinearColorAttribute(MaterialX::InputPt
 			FString NodeName = Node ? Node->getName().c_str() + FString{ TEXT("_") } : FString{};
 			NodeName += Input->getName().c_str();
 
-			UInterchangeShaderNode* VectorParameterNode = CreateShaderNode(NodeName, VectorParameter::Name.ToString());
+			UInterchangeShaderNode* VectorParameterNode = CreateShaderNode(Input, NodeName, VectorParameter::Name.ToString());
 			VectorParameterNode->AddLinearColorAttribute(UInterchangeShaderPortsAPI::MakeInputParameterKey(VectorParameter::Attributes::DefaultValue.ToString()), Value);
 			return UInterchangeShaderPortsAPI::ConnectOuputToInputByIndex(ShaderNode, InputChannelName, VectorParameterNode->GetUniqueID(), OutputIndex);
 		}
@@ -179,7 +179,7 @@ bool FMaterialXSurfaceShaderAbstract::AddVectorAttribute(MaterialX::InputPtr Inp
 			FString NodeName = Node ? Node->getName().c_str() + FString{ TEXT("_") } : FString{};
 			NodeName += Input->getName().c_str();
 
-			UInterchangeShaderNode* VectorParameterNode = CreateShaderNode(NodeName, VectorParameter::Name.ToString());
+			UInterchangeShaderNode* VectorParameterNode = CreateShaderNode(Input, NodeName, VectorParameter::Name.ToString());
 			VectorParameterNode->AddLinearColorAttribute(UInterchangeShaderPortsAPI::MakeInputParameterKey(VectorParameter::Attributes::DefaultValue.ToString()), Value);
 			return UInterchangeShaderPortsAPI::ConnectOuputToInputByIndex(ShaderNode, InputChannelName, VectorParameterNode->GetUniqueID(), OutputIndex);
 		}
@@ -251,12 +251,12 @@ bool FMaterialXSurfaceShaderAbstract::ConnectMatchingNodeOutputToInput(const FCo
 			UInterchangeShaderPortsAPI::ConnectOuputToInputByIndex(Connect.ParentShaderNode, Connect.InputChannelName, OperatorNode->GetUniqueID(), IndexOutput);
 	};
 
-	auto ConnectOutputToInput = [&](const FString* ShaderType, auto* (FMaterialXSurfaceShaderAbstract::* CreateFunctionCallOrShaderNode)(const FString&, const FString&, const FString&), bool bFindMatchingInput = true)
+	auto ConnectOutputToInput = [&](const FString* ShaderType, auto* (FMaterialXSurfaceShaderAbstract::* CreateFunctionCallOrShaderNode)(mx::ElementPtr, const FString&, const FString&, const FString&), bool bFindMatchingInput = true)
 	{
 		UInterchangeShaderNode* OperatorNode = nullptr;
 
 		//We don't take the node output here because it would cause the creation of a new node (output is meaningful with ComponentMaskNode/separate where we have to create a new expression
-		OperatorNode = (this->*CreateFunctionCallOrShaderNode)(Connect.UpstreamNode->getName().c_str(), *ShaderType, DefaultOutput);
+		OperatorNode = (this->*CreateFunctionCallOrShaderNode)(Connect.UpstreamNode, Connect.UpstreamNode->getName().c_str(), *ShaderType, DefaultOutput);
 
 		ConnectOutputToInputInternal(OperatorNode, bFindMatchingInput);
 	};
@@ -329,7 +329,7 @@ void FMaterialXSurfaceShaderAbstract::ConnectNodeCategoryOutputToInput(const Mat
 					if(ConnectedInput->hasChannels())
 					{
 						using namespace UE::Interchange::Materials::Standard::Nodes;
-						ChannelsNode = CreateShaderNode((UpstreamNode->getName() + ConnectedInput->getName()).c_str() + FString{ TEXT("_Channels") }, Swizzle::Name.ToString());
+						ChannelsNode = CreateShaderNode(ConnectedInput, (UpstreamNode->getName() + ConnectedInput->getName()).c_str() + FString{ TEXT("_Channels") }, Swizzle::Name.ToString());
 						ChannelsNode->AddStringAttribute(Swizzle::Attributes::Channels.ToString(), ConnectedInput->getChannels().c_str());
 					}
 				}
@@ -425,7 +425,7 @@ void FMaterialXSurfaceShaderAbstract::ConnectExtractInputToOutput(const FConnect
 	}
 	else
 	{
-		UInterchangeShaderNode* MaskShaderNode = CreateMaskShaderNode(1 << (3 - Index), Connect.UpstreamNode->getName().c_str());
+		UInterchangeShaderNode* MaskShaderNode = CreateMaskShaderNode(1 << (3 - Index), Connect.UpstreamNode, Connect.UpstreamNode->getName().c_str());
 		UInterchangeShaderPortsAPI::ConnectDefaultOuputToInput(Connect.ParentShaderNode, Connect.InputChannelName, MaskShaderNode->GetUniqueID());
 	}
 }
@@ -441,14 +441,14 @@ void FMaterialXSurfaceShaderAbstract::ConnectDotInputToOutput(const FConnectNode
 
 void FMaterialXSurfaceShaderAbstract::ConnectTransformPositionInputToOutput(const FConnectNode& Connect)
 {
-	UInterchangeShaderNode* TransformNode = CreateShaderNode(Connect.UpstreamNode->getName().c_str(), TEXT("TransformPosition"));
+	UInterchangeShaderNode* TransformNode = CreateShaderNode(Connect.UpstreamNode, Connect.UpstreamNode->getName().c_str(), TEXT("TransformPosition"));
 	AddAttributeFromValueOrInterface(Connect.UpstreamNode->getInput("in"), TEXT("Input"), TransformNode);
 	UInterchangeShaderPortsAPI::ConnectDefaultOuputToInput(Connect.ParentShaderNode, Connect.InputChannelName, TransformNode->GetUniqueID());
 }
 
 void FMaterialXSurfaceShaderAbstract::ConnectTransformVectorInputToOutput(const FConnectNode& Connect)
 {
-	UInterchangeShaderNode* TransformNode = CreateShaderNode(Connect.UpstreamNode->getName().c_str(), TEXT("Transform"));
+	UInterchangeShaderNode* TransformNode = CreateShaderNode(Connect.UpstreamNode, Connect.UpstreamNode->getName().c_str(), TEXT("Transform"));
 	AddAttributeFromValueOrInterface(Connect.UpstreamNode->getInput("in"), TEXT("Input"), TransformNode);
 	UInterchangeShaderPortsAPI::ConnectDefaultOuputToInput(Connect.ParentShaderNode, Connect.InputChannelName, TransformNode->GetUniqueID());
 }
@@ -457,7 +457,7 @@ void FMaterialXSurfaceShaderAbstract::ConnectRotate2DInputToOutput(const FConnec
 {
 	using namespace UE::Interchange::Materials::Standard::Nodes;
 
-	UInterchangeShaderNode* Rotate2DNode = CreateShaderNode(Connect.UpstreamNode->getName().c_str(), Rotator::Name.ToString());
+	UInterchangeShaderNode* Rotate2DNode = CreateShaderNode(Connect.UpstreamNode, Connect.UpstreamNode->getName().c_str(), Rotator::Name.ToString());
 
 
 	if(mx::InputPtr Input = Connect.UpstreamNode->getInput("in"))
@@ -469,7 +469,7 @@ void FMaterialXSurfaceShaderAbstract::ConnectRotate2DInputToOutput(const FConnec
 	if(mx::InputPtr Input = Connect.UpstreamNode->getInput("amount"))
 	{
 		FString InputName = GetInputName(Input);
-		UInterchangeShaderNode* DegreesToRadiansNode = CreateShaderNode((Connect.UpstreamNode->getName() + "multiply").c_str(), TEXT("Multiply"));
+		UInterchangeShaderNode* DegreesToRadiansNode = CreateShaderNode(Connect.UpstreamNode, (Connect.UpstreamNode->getName() + "multiply").c_str(), TEXT("Multiply"));
 		constexpr float DegreesToRadians = UE_PI / 180.f;
 		DegreesToRadiansNode->AddFloatAttribute(UInterchangeShaderPortsAPI::MakeInputValueKey(TEXT("B")), DegreesToRadians);
 
@@ -488,7 +488,7 @@ void FMaterialXSurfaceShaderAbstract::ConnectRotate2DInputToOutput(const FConnec
 
 void FMaterialXSurfaceShaderAbstract::ConnectRotate3DInputToOutput(const FConnectNode& Connect)
 {
-	UInterchangeShaderNode* Rotate3DNode = CreateShaderNode(Connect.UpstreamNode->getName().c_str(), TEXT("RotateAboutAxis"));
+	UInterchangeShaderNode* Rotate3DNode = CreateShaderNode(Connect.UpstreamNode, Connect.UpstreamNode->getName().c_str(), TEXT("RotateAboutAxis"));
 	Rotate3DNode->AddLinearColorAttribute(UInterchangeShaderPortsAPI::MakeInputValueKey(TEXT("PivotPoint")), FLinearColor(0.5, 0.5, 0));
 
 	AddAttributeFromValueOrInterface(Connect.UpstreamNode->getInput("in"), TEXT("Position"), Rotate3DNode);
@@ -534,7 +534,7 @@ void FMaterialXSurfaceShaderAbstract::ConnectImageInputToOutput(const FConnectNo
 			OutputChannel = TEXT("R");
 		}
 
-		UInterchangeShaderNode* TextureShaderNode = CreateShaderNode(Connect.UpstreamNode->getName().c_str(), TextureSample::Name.ToString());
+		UInterchangeShaderNode* TextureShaderNode = CreateShaderNode(Connect.UpstreamNode, Connect.UpstreamNode->getName().c_str(), TextureSample::Name.ToString());
 		TextureShaderNode->AddStringAttribute(UInterchangeShaderPortsAPI::MakeInputValueKey(TextureSample::Inputs::Texture.ToString()), TextureNode->GetUniqueID());
 		UInterchangeShaderPortsAPI::ConnectOuputToInputByName(Connect.ParentShaderNode, Connect.InputChannelName, TextureShaderNode->GetUniqueID(), OutputChannel);
 	}
@@ -603,15 +603,15 @@ void FMaterialXSurfaceShaderAbstract::ConnectConvertInputToOutput(const FConnect
 		UInterchangeShaderNode* MaskShaderNode = nullptr;
 		if(NodeN == '3')
 		{			
-			MaskShaderNode = CreateMaskShaderNode(0b1110, Connect.UpstreamNode->getName().c_str());
+			MaskShaderNode = CreateMaskShaderNode(0b1110, Connect.UpstreamNode, Connect.UpstreamNode->getName().c_str());
 		}
 		else if(NodeN == '2')
 		{
-			MaskShaderNode = CreateMaskShaderNode(0b1100, Connect.UpstreamNode->getName().c_str());
+			MaskShaderNode = CreateMaskShaderNode(0b1100, Connect.UpstreamNode, Connect.UpstreamNode->getName().c_str());
 		}
 		else
 		{
-			MaskShaderNode = CreateMaskShaderNode(0b1000, Connect.UpstreamNode->getName().c_str());
+			MaskShaderNode = CreateMaskShaderNode(0b1000, Connect.UpstreamNode, Connect.UpstreamNode->getName().c_str());
 		}
 
 		AddAttributeFromValueOrInterface(Input, TEXT("Input"), MaskShaderNode);
@@ -688,7 +688,7 @@ void FMaterialXSurfaceShaderAbstract::ConnectConvertInputToOutput(const FConnect
 
 void FMaterialXSurfaceShaderAbstract::ConnectIfGreaterInputToOutput(const FConnectNode& Connect)
 {
-	UInterchangeShaderNode* NodeIf = CreateShaderNode(Connect.UpstreamNode->getName().c_str(), TEXT("If"));
+	UInterchangeShaderNode* NodeIf = CreateShaderNode(Connect.UpstreamNode, Connect.UpstreamNode->getName().c_str(), TEXT("If"));
 	UInterchangeShaderPortsAPI::ConnectDefaultOuputToInput(Connect.ParentShaderNode, Connect.InputChannelName, NodeIf->GetUniqueID());
 
 	if(mx::InputPtr Input = Connect.UpstreamNode->getInput("value1"))
@@ -722,7 +722,7 @@ void FMaterialXSurfaceShaderAbstract::ConnectIfGreaterInputToOutput(const FConne
 
 void FMaterialXSurfaceShaderAbstract::ConnectIfGreaterEqInputToOutput(const FConnectNode& Connect)
 {
-	UInterchangeShaderNode* NodeIf = CreateShaderNode(Connect.UpstreamNode->getName().c_str(), TEXT("If"));
+	UInterchangeShaderNode* NodeIf = CreateShaderNode(Connect.UpstreamNode, Connect.UpstreamNode->getName().c_str(), TEXT("If"));
 	UInterchangeShaderPortsAPI::ConnectDefaultOuputToInput(Connect.ParentShaderNode, Connect.InputChannelName, NodeIf->GetUniqueID());
 
 	if(mx::InputPtr Input = Connect.UpstreamNode->getInput("value1"))
@@ -756,7 +756,7 @@ void FMaterialXSurfaceShaderAbstract::ConnectIfGreaterEqInputToOutput(const FCon
 
 void FMaterialXSurfaceShaderAbstract::ConnectIfEqualInputToOutput(const FConnectNode& Connect)
 {
-	UInterchangeShaderNode* NodeIf = CreateShaderNode(Connect.UpstreamNode->getName().c_str(), TEXT("If"));
+	UInterchangeShaderNode* NodeIf = CreateShaderNode(Connect.UpstreamNode, Connect.UpstreamNode->getName().c_str(), TEXT("If"));
 	UInterchangeShaderPortsAPI::ConnectDefaultOuputToInput(Connect.ParentShaderNode, Connect.InputChannelName, NodeIf->GetUniqueID());
 
 	if(mx::InputPtr Input = Connect.UpstreamNode->getInput("value1"))
@@ -791,11 +791,11 @@ void FMaterialXSurfaceShaderAbstract::ConnectIfEqualInputToOutput(const FConnect
 void FMaterialXSurfaceShaderAbstract::ConnectOutsideInputToOutput(const FConnectNode& Connect)
 {
 	//in * (1 - mask)
-	UInterchangeShaderNode* NodeMultiply = CreateShaderNode(Connect.UpstreamNode->getName().c_str(), TEXT("Multiply"));
+	UInterchangeShaderNode* NodeMultiply = CreateShaderNode(Connect.UpstreamNode, Connect.UpstreamNode->getName().c_str(), TEXT("Multiply"));
 	AddAttributeFromValueOrInterface(Connect.UpstreamNode->getInput("in"), TEXT("A"), NodeMultiply);
 	UInterchangeShaderPortsAPI::ConnectDefaultOuputToInput(Connect.ParentShaderNode, Connect.InputChannelName, NodeMultiply->GetUniqueID());
 
-	UInterchangeShaderNode* NodeOneMinus = CreateShaderNode(Connect.UpstreamNode->getName().c_str() + FString(TEXT("_OneMinus")), TEXT("OneMinus"));
+	UInterchangeShaderNode* NodeOneMinus = CreateShaderNode(Connect.UpstreamNode, Connect.UpstreamNode->getName().c_str() + FString(TEXT("_OneMinus")), TEXT("OneMinus"));
 	AddAttributeFromValueOrInterface(Connect.UpstreamNode->getInput("mask"), TEXT("Input"), NodeOneMinus);
 	UInterchangeShaderPortsAPI::ConnectDefaultOuputToInput(NodeMultiply, TEXT("B"), NodeOneMinus->GetUniqueID());
 }
@@ -808,7 +808,7 @@ UInterchangeShaderNode* FMaterialXSurfaceShaderAbstract::ConnectGeometryInputToO
 	// world : The global coordinate space of the geometry, after local deformationsand global transforms have been applied.
 
 	// In case of model/object we need to add a TransformVector from world to local space
-	UInterchangeShaderNode* GeometryNode = CreateShaderNode(Connect.UpstreamNode->getName().c_str(), ShaderType);
+	UInterchangeShaderNode* GeometryNode = CreateShaderNode(Connect.UpstreamNode, Connect.UpstreamNode->getName().c_str(), ShaderType);
 
 	UInterchangeShaderNode* NodeToConnectTo = Connect.ParentShaderNode;
 	FString InputToConnectTo = Connect.InputChannelName;
@@ -822,7 +822,7 @@ UInterchangeShaderNode* FMaterialXSurfaceShaderAbstract::ConnectGeometryInputToO
 	if(bTangentSpaceInput && bIsVector)
 	{
 		using namespace UE::Interchange::Materials::Standard::Nodes;
-		UInterchangeShaderNode* TransformTSNode = CreateShaderNode(Connect.UpstreamNode->getName().c_str() + FString(TEXT("_TransformTS")), TransformShaderType);
+		UInterchangeShaderNode* TransformTSNode = CreateShaderNode(Connect.UpstreamNode, Connect.UpstreamNode->getName().c_str() + FString(TEXT("_TransformTS")), TransformShaderType);
 		EMaterialVectorCoordTransformSource SpaceSource = bIsObjectSpace ? EMaterialVectorCoordTransformSource::TRANSFORMSOURCE_Local : EMaterialVectorCoordTransformSource::TRANSFORMSOURCE_World;
 		TransformTSNode->AddInt32Attribute(TransformSourceType, SpaceSource);
 		TransformTSNode->AddInt32Attribute(TransformType, EMaterialVectorCoordTransform::TRANSFORM_Tangent);
@@ -833,7 +833,7 @@ UInterchangeShaderNode* FMaterialXSurfaceShaderAbstract::ConnectGeometryInputToO
 
 	if(bIsObjectSpace)
 	{
-		UInterchangeShaderNode* TransformNode = CreateShaderNode(Connect.UpstreamNode->getName().c_str() + FString(TEXT("_Transform")), TransformShaderType);
+		UInterchangeShaderNode* TransformNode = CreateShaderNode(Connect.UpstreamNode, Connect.UpstreamNode->getName().c_str() + FString(TEXT("_Transform")), TransformShaderType);
 		TransformNode->AddInt32Attribute(TransformSourceType, TransformSource);
 		TransformNode->AddInt32Attribute(TransformType, TransformSDestination);
 		UInterchangeShaderPortsAPI::ConnectDefaultOuputToInput(NodeToConnectTo, InputToConnectTo, TransformNode->GetUniqueID());
@@ -858,7 +858,7 @@ void FMaterialXSurfaceShaderAbstract::ConnectPositionInputToOutput(const FConnec
 
 	// In case of the position node, it seems that the unit is different, we assume for now a conversion from mm -> m, even if UE by default is cm
 	// See standard_surface_marble_solid file, especially on the fractal3d node
-	UInterchangeShaderNode* PositionNode = CreateShaderNode(Connect.UpstreamNode->getName().c_str() + FString(TEXT("_Position")), TEXT("WorldPosition"));
+	UInterchangeShaderNode* PositionNode = CreateShaderNode(Connect.UpstreamNode, Connect.UpstreamNode->getName().c_str() + FString(TEXT("_Position")), TEXT("WorldPosition"));
 	UnitNode->AddFloatAttribute(UInterchangeShaderPortsAPI::MakeInputValueKey(TEXT("B")), 0.001f);
 	UInterchangeShaderPortsAPI::ConnectDefaultOuputToInput(UnitNode, TEXT("A"), PositionNode->GetUniqueID());
 }
@@ -893,8 +893,8 @@ void FMaterialXSurfaceShaderAbstract::ConnectBitangentInputToOutput(const FConne
 																		 TransformVector::Attributes::TransformSourceType.ToString(), EMaterialVectorCoordTransformSource::TRANSFORMSOURCE_World,
 																		 TransformVector::Attributes::TransformType.ToString(), EMaterialVectorCoordTransform::TRANSFORM_Local);
 
-	UInterchangeShaderNode* NormalNode = CreateShaderNode(Connect.UpstreamNode->getName().c_str() + FString(TEXT("_Normal")), TEXT("VertexNormalWS"));
-	UInterchangeShaderNode* TangentNode = CreateShaderNode(Connect.UpstreamNode->getName().c_str() + FString(TEXT("_Tangent")), TEXT("VertexTangentWS"));
+	UInterchangeShaderNode* NormalNode = CreateShaderNode(Connect.UpstreamNode, Connect.UpstreamNode->getName().c_str() + FString(TEXT("_Normal")), TEXT("VertexNormalWS"));
+	UInterchangeShaderNode* TangentNode = CreateShaderNode(Connect.UpstreamNode, Connect.UpstreamNode->getName().c_str() + FString(TEXT("_Tangent")), TEXT("VertexTangentWS"));
 
 	UInterchangeShaderPortsAPI::ConnectDefaultOuputToInput(BitangentNode, TEXT("A"), NormalNode->GetUniqueID());
 	UInterchangeShaderPortsAPI::ConnectDefaultOuputToInput(BitangentNode, TEXT("B"), TangentNode->GetUniqueID());
@@ -903,7 +903,7 @@ void FMaterialXSurfaceShaderAbstract::ConnectBitangentInputToOutput(const FConne
 void FMaterialXSurfaceShaderAbstract::ConnectTimeInputToOutput(const FConnectNode& Connect)
 {
 	using namespace UE::Interchange::Materials::Standard::Nodes;
-	UInterchangeShaderNode* TimeNode = CreateShaderNode(Connect.UpstreamNode->getName().c_str(), TEXT("Time"));
+	UInterchangeShaderNode* TimeNode = CreateShaderNode(Connect.UpstreamNode, Connect.UpstreamNode->getName().c_str(), TEXT("Time"));
 	TimeNode->AddBooleanAttribute(Time::Attributes::OverridePeriod.ToString(), true);
 
 	float FPS;
@@ -928,7 +928,7 @@ void FMaterialXSurfaceShaderAbstract::ConnectNoise3DInputToOutput(const FConnect
 	using namespace UE::Interchange::Materials::Standard::Nodes;
 
 	// MaterialX defines the Noise3d as Perlin Noise which is multiplied by the Amplitude then Added to Pivot
-	UInterchangeShaderNode* NoiseNode = CreateShaderNode(Connect.UpstreamNode->getName().c_str(), Noise::Name.ToString());
+	UInterchangeShaderNode* NoiseNode = CreateShaderNode(Connect.UpstreamNode, Connect.UpstreamNode->getName().c_str(), Noise::Name.ToString());
 	NoiseNode->AddInt32Attribute(Noise::Attributes::Function.ToString(), ENoiseFunction::NOISEFUNCTION_GradientTex);
 	NoiseNode->AddBooleanAttribute(Noise::Attributes::Turbulence.ToString(), false);
 	NoiseNode->AddFloatAttribute(Noise::Attributes::OutputMin.ToString(), 0);
@@ -944,7 +944,7 @@ void FMaterialXSurfaceShaderAbstract::ConnectNoise3DInputToOutput(const FConnect
 		}
 
 		const FString ShaderNodeName = Connect.UpstreamNode->getName().c_str() + FString{ TEXT("_") } + ShaderType;
-		UInterchangeShaderNode* ShaderNode = CreateShaderNode(ShaderNodeName, ShaderType);
+		UInterchangeShaderNode* ShaderNode = CreateShaderNode(Connect.UpstreamNode, ShaderNodeName, ShaderType);
 
 		UInterchangeShaderPortsAPI::ConnectDefaultOuputToInput(ShaderNode, TEXT("A"), NodeToConnectTo->GetUniqueID());
 
@@ -976,7 +976,7 @@ void FMaterialXSurfaceShaderAbstract::ConnectCellNoise3DInputToOutput(const FCon
 	using namespace UE::Interchange::Materials::Standard::Nodes;
 
 	// Let's use a vector noise for this one, the only one that is close to MaterialX implementation
-	UInterchangeShaderNode* NoiseNode = CreateShaderNode(Connect.UpstreamNode->getName().c_str(), VectorNoise::Name.ToString());
+	UInterchangeShaderNode* NoiseNode = CreateShaderNode(Connect.UpstreamNode, Connect.UpstreamNode->getName().c_str(), VectorNoise::Name.ToString());
 	NoiseNode->AddInt32Attribute(VectorNoise::Attributes::Function.ToString(), EVectorNoiseFunction::VNF_CellnoiseALU);
 
 	UInterchangeShaderPortsAPI::ConnectDefaultOuputToInput(Connect.ParentShaderNode, Connect.InputChannelName, NoiseNode->GetUniqueID());
@@ -987,7 +987,7 @@ void FMaterialXSurfaceShaderAbstract::ConnectWorleyNoise3DInputToOutput(const FC
 	using namespace UE::Interchange::Materials::Standard::Nodes;
 
 	//Also called Voronoi, the implementation is a bit different in UE, especially we don't have access to the jitter
-	UInterchangeShaderNode* NoiseNode = CreateShaderNode(Connect.UpstreamNode->getName().c_str(), Noise::Name.ToString());
+	UInterchangeShaderNode* NoiseNode = CreateShaderNode(Connect.UpstreamNode, Connect.UpstreamNode->getName().c_str(), Noise::Name.ToString());
 	NoiseNode->AddInt32Attribute(Noise::Attributes::Function.ToString(), ENoiseFunction::NOISEFUNCTION_VoronoiALU);
 
 	UInterchangeShaderPortsAPI::ConnectDefaultOuputToInput(Connect.ParentShaderNode, Connect.InputChannelName, NoiseNode->GetUniqueID());
@@ -1017,11 +1017,11 @@ void FMaterialXSurfaceShaderAbstract::ConnectHeightToNormalInputToOutput(const F
 
 			if(UInterchangeTextureNode* TextureNode = CreateTextureNode<UInterchangeTexture2DNode>(Connect.UpstreamNode))
 			{
-				UInterchangeShaderNode* HeightMapNode = CreateShaderNode(Connect.UpstreamNode->getName().c_str(), NormalFromHeightMap::Name.ToString());
+				UInterchangeShaderNode* HeightMapNode = CreateShaderNode(Connect.UpstreamNode, Connect.UpstreamNode->getName().c_str(), NormalFromHeightMap::Name.ToString());
 				UInterchangeShaderPortsAPI::ConnectDefaultOuputToInput(Connect.ParentShaderNode, Connect.InputChannelName, HeightMapNode->GetUniqueID());
 
 				const FString TextureNodeName = Connect.UpstreamNode->getName().c_str() + FString{ "_texture" };
-				UInterchangeShaderNode* TextureShaderNode = CreateShaderNode(TextureNodeName, TextureObject::Name.ToString());
+				UInterchangeShaderNode* TextureShaderNode = CreateShaderNode(Connect.UpstreamNode, TextureNodeName, TextureObject::Name.ToString());
 				TextureShaderNode->AddStringAttribute(UInterchangeShaderPortsAPI::MakeInputValueKey(TextureObject::Inputs::Texture.ToString()), TextureNode->GetUniqueID());
 				UInterchangeShaderPortsAPI::ConnectDefaultOuputToInput(HeightMapNode, NormalFromHeightMap::Inputs::HeightMap.ToString(), TextureShaderNode->GetUniqueID());
 
@@ -1077,7 +1077,7 @@ void FMaterialXSurfaceShaderAbstract::ConnectBlurInputToOutput(const FConnectNod
 					OutputChannel = TEXT("R");
 				}
 
-				UInterchangeShaderNode* TextureShaderNode = CreateShaderNode(Connect.UpstreamNode->getName().c_str(), TextureSampleBlur::Name.ToString());
+				UInterchangeShaderNode* TextureShaderNode = CreateShaderNode(Connect.UpstreamNode, Connect.UpstreamNode->getName().c_str(), TextureSampleBlur::Name.ToString());
 				TextureShaderNode->AddStringAttribute(UInterchangeShaderPortsAPI::MakeInputValueKey(TextureSampleBlur::Inputs::Texture.ToString()), TextureNode->GetUniqueID());
 				UInterchangeShaderPortsAPI::ConnectOuputToInputByName(Connect.ParentShaderNode, Connect.InputChannelName, TextureShaderNode->GetUniqueID(), OutputChannel);
 
@@ -1129,23 +1129,23 @@ void FMaterialXSurfaceShaderAbstract::ConnectBlurInputToOutput(const FConnectNod
 void FMaterialXSurfaceShaderAbstract::ConnectTexCoordInputToOutput(const FConnectNode& Connect)
 {
 	using namespace UE::Interchange::Materials::Standard::Nodes;
-	UInterchangeShaderNode* TexCoord = CreateShaderNode(Connect.UpstreamNode->getName().c_str(), TextureCoordinate::Name.ToString());
+	UInterchangeShaderNode* TexCoord = CreateShaderNode(Connect.UpstreamNode, Connect.UpstreamNode->getName().c_str(), TextureCoordinate::Name.ToString());
 	if(mx::InputPtr Input = Connect.UpstreamNode->getInput("index"))
 	{
 		TexCoord->AddInt32Attribute(UInterchangeShaderPortsAPI::MakeInputValueKey(TextureCoordinate::Inputs::Index.ToString()), mx::fromValueString<int>(Input->getValueString()));
 	}
 
 	// Flip UV coordinates to be bottom left instead of top left
-	UInterchangeShaderNode* TexCoordR = CreateMaskShaderNode(0b1000, (Connect.UpstreamNode->getName() + "ComponentMaskR").c_str());
+	UInterchangeShaderNode* TexCoordR = CreateMaskShaderNode(0b1000, Connect.UpstreamNode, (Connect.UpstreamNode->getName() + "ComponentMaskR").c_str());
 	UInterchangeShaderPortsAPI::ConnectDefaultOuputToInput(TexCoordR, Mask::Inputs::Input.ToString(), TexCoord->GetUniqueID());
 
-	UInterchangeShaderNode* TexCoordG = CreateMaskShaderNode(0b0100, (Connect.UpstreamNode->getName() + "ComponentMaskG").c_str());
+	UInterchangeShaderNode* TexCoordG = CreateMaskShaderNode(0b0100, Connect.UpstreamNode, (Connect.UpstreamNode->getName() + "ComponentMaskG").c_str());
 	UInterchangeShaderPortsAPI::ConnectDefaultOuputToInput(TexCoordG, Mask::Inputs::Input.ToString(), TexCoord->GetUniqueID());
 
-	UInterchangeShaderNode* OneMinusG = CreateShaderNode((Connect.UpstreamNode->getName() + "OneMinusG").c_str(), OneMinus::Name.ToString());
+	UInterchangeShaderNode* OneMinusG = CreateShaderNode(Connect.UpstreamNode, (Connect.UpstreamNode->getName() + "OneMinusG").c_str(), OneMinus::Name.ToString());
 	UInterchangeShaderPortsAPI::ConnectDefaultOuputToInput(OneMinusG, OneMinus::Inputs::Input.ToString(), TexCoordG->GetUniqueID());
 
-	UInterchangeShaderNode* TexCoordAppend = CreateShaderNode((Connect.UpstreamNode->getName() + "Append").c_str(), TEXT("AppendVector"));
+	UInterchangeShaderNode* TexCoordAppend = CreateShaderNode(Connect.UpstreamNode, (Connect.UpstreamNode->getName() + "Append").c_str(), TEXT("AppendVector"));
 	UInterchangeShaderPortsAPI::ConnectDefaultOuputToInput(TexCoordAppend, TEXT("A"), TexCoordR->GetUniqueID());
 	UInterchangeShaderPortsAPI::ConnectDefaultOuputToInput(TexCoordAppend, TEXT("B"), OneMinusG->GetUniqueID());
 
@@ -1190,7 +1190,7 @@ void FMaterialXSurfaceShaderAbstract::ConnectSeparateInputToOutput(const FConnec
 
 	if(!bHasValueOrInterface)
 	{
-		UInterchangeShaderNode* OutNode = CreateMaskShaderNode(Mask, Connect.UpstreamNode->getName().c_str(), Connect.OutputName);
+		UInterchangeShaderNode* OutNode = CreateMaskShaderNode(Mask, Connect.UpstreamNode, Connect.UpstreamNode->getName().c_str(), Connect.OutputName);
 		UInterchangeShaderPortsAPI::ConnectDefaultOuputToInput(Connect.ParentShaderNode, Connect.InputChannelName, OutNode->GetUniqueID());
 	}
 	else
@@ -1203,7 +1203,7 @@ void FMaterialXSurfaceShaderAbstract::ConnectSwizzleInputToOutput(const FConnect
 {
 	using namespace UE::Interchange::Materials::Standard::Nodes;
 
-	UInterchangeShaderNode* SwizzleNode = CreateShaderNode(Connect.UpstreamNode->getName().c_str(), Swizzle::Name.ToString());
+	UInterchangeShaderNode* SwizzleNode = CreateShaderNode(Connect.UpstreamNode, Connect.UpstreamNode->getName().c_str(), Swizzle::Name.ToString());
 	if(mx::InputPtr Input = Connect.UpstreamNode->getInput("in"))
 	{
 		AddAttributeFromValueOrInterface(Input, GetInputName(Input), SwizzleNode);
@@ -1234,9 +1234,9 @@ void FMaterialXSurfaceShaderAbstract::ConnectNormalMapInputToOutput(const FConne
 
 	if(Input && bIsNotEqualOne)
 	{
-		UInterchangeShaderNode* FlattenNormalNode = CreateFunctionCallShaderNode(Connect.UpstreamNode->getName().c_str(), TEXT("/Engine/Functions/Engine_MaterialFunctions01/Texturing/FlattenNormal.FlattenNormal"));
+		UInterchangeShaderNode* FlattenNormalNode = CreateFunctionCallShaderNode(Connect.UpstreamNode, Connect.UpstreamNode->getName().c_str(), TEXT("/Engine/Functions/Engine_MaterialFunctions01/Texturing/FlattenNormal.FlattenNormal"));
 
-		UInterchangeShaderNode* OneMinusNode = CreateShaderNode((Connect.UpstreamNode->getName() + "_OneMinusFlatness").c_str(), OneMinus::Name.ToString());
+		UInterchangeShaderNode* OneMinusNode = CreateShaderNode(Connect.UpstreamNode, (Connect.UpstreamNode->getName() + "_OneMinusFlatness").c_str(), OneMinus::Name.ToString());
 		UInterchangeShaderPortsAPI::ConnectDefaultOuputToInput(FlattenNormalNode, FlattenNormal::Inputs::Flatness.ToString(), OneMinusNode->GetUniqueID());
 
 		AddAttributeFromValueOrInterface(Input, OneMinus::Inputs::Input.ToString(), OneMinusNode);
@@ -1249,14 +1249,14 @@ void FMaterialXSurfaceShaderAbstract::ConnectNormalMapInputToOutput(const FConne
 	}
 }
 
-UInterchangeShaderNode* FMaterialXSurfaceShaderAbstract::CreateMaskShaderNode(uint8 RGBA, const FString& NodeName, const FString& OutputName)
+UInterchangeShaderNode* FMaterialXSurfaceShaderAbstract::CreateMaskShaderNode(uint8 RGBA, mx::ElementPtr Element, const FString& NodeName, const FString& OutputName)
 {
 	bool bR = (0b1000 & RGBA) >> 3;
 	bool bG = (0b0100 & RGBA) >> 2;
 	bool bB = (0b0010 & RGBA) >> 1;
 	bool bA = (0b0001 & RGBA) >> 0;
 	using namespace UE::Interchange::Materials::Standard::Nodes;
-	UInterchangeShaderNode* MaskShaderNode = CreateShaderNode(NodeName, Mask::Name.ToString(), OutputName);
+	UInterchangeShaderNode* MaskShaderNode = CreateShaderNode(Element, NodeName, Mask::Name.ToString(), OutputName);
 	MaskShaderNode->AddBooleanAttribute(Mask::Attributes::R.ToString(), bR);
 	MaskShaderNode->AddBooleanAttribute(Mask::Attributes::G.ToString(), bG);
 	MaskShaderNode->AddBooleanAttribute(Mask::Attributes::B.ToString(), bB);
@@ -1265,11 +1265,11 @@ UInterchangeShaderNode* FMaterialXSurfaceShaderAbstract::CreateMaskShaderNode(ui
 	return MaskShaderNode;
 }
 
-UInterchangeShaderNode* FMaterialXSurfaceShaderAbstract::CreateShaderNode(const FString& NodeName, const FString& ShaderType, const FString& OutputName)
+UInterchangeShaderNode* FMaterialXSurfaceShaderAbstract::CreateShaderNode(mx::ElementPtr Element, const FString& NodeName, const FString& ShaderType, const FString& OutputName)
 {
 	UInterchangeShaderNode* Node;
 
-	const FString NodeUID = UInterchangeShaderNode::MakeNodeUid(NodeName + TEXT('_') + OutputName, FStringView{});
+	const FString NodeUID = UInterchangeShaderNode::MakeNodeUid(GetUniqueName(Element) + TEXT("_") + NodeName + TEXT('_') + OutputName, FStringView{});
 
 	//Test directly in the NodeContainer, because the ShaderNodes can be altered during the node graph either by the parent (dot/normalmap),
 	//or by putting an intermediary node between the child and the parent (tiledimage)
@@ -1286,7 +1286,7 @@ UInterchangeShaderNode* FMaterialXSurfaceShaderAbstract::CreateShaderNode(const
 	return Node;
 }
 
-UInterchangeFunctionCallShaderNode* FMaterialXSurfaceShaderAbstract::CreateFunctionCallShaderNode(const FString& NodeName, const FString& FunctionPath, const FString& OutputName)
+UInterchangeFunctionCallShaderNode* FMaterialXSurfaceShaderAbstract::CreateFunctionCallShaderNode(MaterialX::ElementPtr Element, const FString& NodeName, const FString& FunctionPath, const FString& OutputName)
 {
 	UInterchangeFunctionCallShaderNode* Node;
 
@@ -1431,6 +1431,39 @@ FString FMaterialXSurfaceShaderAbstract::GetAttributeParentName(MaterialX::NodeP
 	return ParentName.IsEmpty() ? Node->getName().c_str() : ParentName;
 }
 
+namespace UE::MaterialX::Private
+{
+	void GetUniqueName(mx::ElementPtr Element, TStringBuilder<256>& Buffer)
+	{
+		//Write the hierarchy name as A_B_C
+		if (Element)
+		{
+			if (mx::ElementPtr Parent = Element->getParent())
+			{
+				GetUniqueName(Parent, Buffer);
+			}
+
+			Buffer += Element->getName().c_str();
+			Buffer += TEXT("_");
+			Element->setAttribute(mx::Attributes::UniqueName, TCHAR_TO_ANSI(*FString{ Buffer.ToString() }.TrimChar(TEXT('_'))));
+		}
+	}
+}
+
+FString FMaterialXSurfaceShaderAbstract::GetUniqueName(MaterialX::ElementPtr Element) const
+{
+	if (Element->hasAttribute(mx::Attributes::UniqueName))
+	{
+		return Element->getAttribute(mx::Attributes::UniqueName).c_str();
+	}
+	else
+	{
+		TStringBuilder<256> Buffer;
+		UE::MaterialX::Private::GetUniqueName(Element, Buffer);
+		return FString{ Buffer.ToString() }.TrimChar(TEXT('_')); // remove any trailing '_'
+	}
+}
+
 void FMaterialXSurfaceShaderAbstract::RegisterConnectNodeOutputToInputDelegates()
 {
 	MatchingConnectNodeDelegates.Add(mx::Category::Constant,		FOnConnectNodeOutputToInput::CreateSP(this, &FMaterialXSurfaceShaderAbstract::ConnectConstantInputToOutput));
diff --git a/Engine/Plugins/Interchange/Runtime/Source/Import/Public/MaterialX/MaterialXUtils/MaterialXSurfaceShaderAbstract.h b/Engine/Plugins/Interchange/Runtime/Source/Import/Public/MaterialX/MaterialXUtils/MaterialXSurfaceShaderAbstract.h
--- a/Engine/Plugins/Interchange/Runtime/Source/Import/Public/MaterialX/MaterialXUtils/MaterialXSurfaceShaderAbstract.h
+++ b/Engine/Plugins/Interchange/Runtime/Source/Import/Public/MaterialX/MaterialXUtils/MaterialXSurfaceShaderAbstract.h
@@ -292,22 +292,24 @@ protected:
 	 * Create a ComponentMask shader node.
 	 *
 	 * @param RGBA - The mask component. For example: 0b1011 -> Only RBA are toggled
+	 * @param Element - The element that we take the whole hierarchy from to ensure the uniqueness of the name
 	 * @param NodeName - the name of the shader node.
 	 * @param OutputName - the name of the output of the MaterialX node. The default name is 'out' as stated by the standard library.
 	 * @return The ComponentMask node.
 	 */
-	UInterchangeShaderNode* CreateMaskShaderNode(uint8 RGBA, const FString& NodeName, const FString& OutputName = TEXT("out"));
+	UInterchangeShaderNode* CreateMaskShaderNode(uint8 RGBA, MaterialX::ElementPtr Element, const FString& NodeName, const FString& OutputName = TEXT("out"));
 
 	/**
 	 * Helper function to create an InterchangeShaderNode.
 	 *
+	 * @param Element - The element that we take the whole hierarchy from to ensure the uniqueness of the name
 	 * @param NodeName - The name of the shader node.
 	 * @param ShaderType - The type of shader node we want to create.
 	 * @param OutputName - The output name of the MaterialX node. The default name is 'out' as stated by the standard library.
 	 *
 	 * @return The shader node that was created.
 	 */
-	UInterchangeShaderNode* CreateShaderNode(const FString& NodeName, const FString& ShaderType, const FString& OutputName = TEXT("out"));
+	UInterchangeShaderNode* CreateShaderNode(MaterialX::ElementPtr Element, const FString& NodeName, const FString& ShaderType, const FString& OutputName = TEXT("out"));
 
 	/**
 	 * Helper function to create an InterchangeFunctionCallShaderNode.
@@ -318,7 +320,7 @@ protected:
 	 *
 	 * @return The shader node that was created.
 	 */
-	UInterchangeFunctionCallShaderNode* CreateFunctionCallShaderNode(const FString& NodeName, const FString& FunctionPath, const FString& OutputName = TEXT("out"));
+	UInterchangeFunctionCallShaderNode* CreateFunctionCallShaderNode(MaterialX::ElementPtr Element, const FString& NodeName, const FString& FunctionPath, const FString& OutputName = TEXT("out"));
 	UInterchangeFunctionCallShaderNode* CreateFunctionCallShaderNode(const FString& NodeName, uint8 EnumType, uint8 EnumValue, const FString& OutputName = TEXT("out"));
 
 	/**
@@ -450,6 +452,14 @@ protected:
 	 */
 	FString GetAttributeParentName(MaterialX::NodePtr Node) const;
 
+	/**
+	 * Ensure that we put in the node container the unique name, 2 nodes in a same file may have the same name as long as their parent's name is different, we just traverse the whole hierarchy to have a unique name
+	 * 
+	 * @param Element - An Element that may be an Input a node or a nodegraph
+	 * @return The unique name of the Element taking the whole hierarchy into account
+	 */
+	FString GetUniqueName(MaterialX::ElementPtr Element) const;
+
 	virtual void RegisterConnectNodeOutputToInputDelegates();
 
 	/**
