[Interchange-glTF] Fixes for MultiRootJoint scenario.

#rnx

diff --git a/Engine/Plugins/Interchange/Runtime/Source/Parsers/GLTFCore/Private/GLTF/GLTFReader.cpp b/Engine/Plugins/Interchange/Runtime/Source/Parsers/GLTFCore/Private/GLTF/GLTFReader.cpp
--- a/Engine/Plugins/Interchange/Runtime/Source/Parsers/GLTFCore/Private/GLTF/GLTFReader.cpp
+++ b/Engine/Plugins/Interchange/Runtime/Source/Parsers/GLTFCore/Private/GLTF/GLTFReader.cpp
@@ -1177,16 +1177,16 @@ namespace GLTF
 		if (!SetupObjects(CameraCount, TEXT("cameras"), [this](const FJsonObject& Object) { SetupCamera(Object); })) { return; }
 		if (!SetupObjects(SkinCount, TEXT("skins"), [this](const FJsonObject& Object) { SetupSkin(Object); })) { return; }
 
-		{//BuildRootJoints can affect the node hierarchy and data, as the Animation setup currently stores references to the Nodes, we have to do these setups before the SetupAnimations.
+		{//BuildRootJoints can affect the node hierarchy and data, as the Animation setup currently stores references to the Nodes, we have to do these setups before the SetupAnimation calls.
 			SetupUsedSkins();
 
 			SetupNodesType();
 
+			BuildRootJoints();
+
 			GenerateInverseBindPosesPerSkinIndices();
 			GenerateLocalBindPosesPerSkinIndices();
 			SetLocalBindPosesForJoints();
-
-			BuildRootJoints();
 		}
 
 		if (!SetupObjects(AnimationsCount, TEXT("animations"), [this](const FJsonObject& Object) { SetupAnimation(Object); })) { return; }
@@ -1612,25 +1612,74 @@ namespace GLTF
 					{
 						continue;
 					}
-					GLTF::FNode& OriginalNode = Asset->Nodes[Group.Key];
 
 					GLTF::FNode& Node = Asset->Nodes.Emplace_GetRef();
 					Node.Index = Asset->Nodes.Num() - 1;
-
-					Node.Name = OriginalNode.Name + "_ProxyTrueRootJoint";
-					Node.Transform = FTransform::Identity;
-					Node.Children = OriginalNode.Children;
-					OriginalNode.Children.Reset();
-					OriginalNode.Children.Add(Node.Index);
 					Node.Type = GLTF::FNode::EType::Joint;
-					Node.ParentIndex = OriginalNode.Index;
-					Node.bHasLocalBindPose = true;
-					Node.LocalBindPose = FTransform::Identity;
 
-					for (size_t RootJointIndex = 0; RootJointIndex < Group.Value.Indices.Num(); RootJointIndex++)
+					if (Asset->Nodes.IsValidIndex(Group.Key))
+					{
+						GLTF::FNode& OriginalNode = Asset->Nodes[Group.Key];
+
+						Node.Name = OriginalNode.Name + "_ProxyTrueRootJoint";
+						Node.Children = OriginalNode.Children;
+						OriginalNode.Children.Reset();
+						OriginalNode.Children.Add(Node.Index);
+						Node.ParentIndex = OriginalNode.Index;
+
+						Node.Transform = OriginalNode.Transform;
+						Node.bHasLocalBindPose = true;
+						Node.LocalBindPose = FTransform::Identity;
+
+						//also update the Scenes.Nodes idx in case the oroginal node is part of the Scenes idx list:
+						for (GLTF::FScene& Scene : Asset->Scenes)
+						{
+							int32 OriginalNodeIndexInSceneNodes;
+							if (Scene.Nodes.Find(OriginalNode.Index, OriginalNodeIndexInSceneNodes))
+							{
+								Scene.Nodes[OriginalNodeIndexInSceneNodes] = Node.Index;
+							}
+						}
+
+						for (size_t RootJointIndex = 0; RootJointIndex < Group.Value.Indices.Num(); RootJointIndex++)
+						{
+							GLTF::FNode& JointNode = Asset->Nodes[Group.Value.Indices[RootJointIndex]];
+							JointNode.ParentIndex = Node.Index;
+						}
+					}
+					else
 					{
-						GLTF::FNode& JointNode = Asset->Nodes[Group.Value.Indices[RootJointIndex]];
-						JointNode.ParentIndex = Node.Index;
+						Node.Name = "_ProxyTrueRootJoint";
+						Node.Transform = FTransform::Identity;
+						Node.bHasLocalBindPose = true;
+						Node.LocalBindPose = FTransform::Identity;
+
+						for (size_t RootJointIndex = 0; RootJointIndex < Group.Value.Indices.Num(); RootJointIndex++)
+						{
+							GLTF::FNode& JointNode = Asset->Nodes[Group.Value.Indices[RootJointIndex]];
+							JointNode.ParentIndex = Node.Index;
+							Node.Children.Add(JointNode.Index);
+						}
+
+						for (GLTF::FScene& Scene : Asset->Scenes)
+						{
+							bool bHadChildAsSceneNode = false;
+							Scene.Nodes.RemoveAll([&bHadChildAsSceneNode, &Node](const int32& NodeIdx)
+								{
+									int ElementIndex;
+									if (Node.Children.Find(NodeIdx, ElementIndex))
+									{
+										bHadChildAsSceneNode = true;
+										return true;
+									}
+									return false;
+								});
+
+							if (bHadChildAsSceneNode)
+							{
+								Scene.Nodes.Add(Node.Index);
+							}
+						}
 					}
 				}
 			}
