MaterialX:
- Fix compile errors happening in InterchangeUSDTranslator, including Document instead of Node as it's more relevant.

diff --git a/Engine/Plugins/Interchange/Runtime/Source/Import/Public/MaterialX/MaterialXUtils/MaterialXManager.h b/Engine/Plugins/Interchange/Runtime/Source/Import/Public/MaterialX/MaterialXUtils/MaterialXManager.h
--- a/Engine/Plugins/Interchange/Runtime/Source/Import/Public/MaterialX/MaterialXUtils/MaterialXManager.h
+++ b/Engine/Plugins/Interchange/Runtime/Source/Import/Public/MaterialX/MaterialXUtils/MaterialXManager.h
@@ -6,7 +6,7 @@
 
 #if WITH_EDITOR
 
-#include "MaterialXCore/Node.h"
+#include "MaterialXCore/Document.h"
 #include "Misc/TVariant.h"
 #include "MaterialX/InterchangeMaterialXDefinitions.h"
 
