[Interchange] fix crash when importing skeletalmesh

#rnx

diff --git a/Engine/Plugins/Interchange/Runtime/Source/Import/Private/Mesh/InterchangeSkeletalMeshFactory.cpp b/Engine/Plugins/Interchange/Runtime/Source/Import/Private/Mesh/InterchangeSkeletalMeshFactory.cpp
--- a/Engine/Plugins/Interchange/Runtime/Source/Import/Private/Mesh/InterchangeSkeletalMeshFactory.cpp
+++ b/Engine/Plugins/Interchange/Runtime/Source/Import/Private/Mesh/InterchangeSkeletalMeshFactory.cpp
@@ -66,46 +66,35 @@ namespace UE
 {
 	namespace Interchange
 	{
-		namespace Private
+		FInterchangeMeshPayLoadKey FMeshNodeContext::GetTranslatorAndTransformPayloadKey() const
+		{
+			FInterchangeMeshPayLoadKey GlobalPayloadKey = TranslatorPayloadKey;
+			GlobalPayloadKey.UniqueId = GetUniqueId();
+			return GlobalPayloadKey;
+		}
+		FInterchangeMeshPayLoadKey FMeshNodeContext::GetMorphTargetAndTransformPayloadKey(const FInterchangeMeshPayLoadKey& MorphTargetKey) const
 		{
-			//Get the mesh node context for each MeshUids
-			struct FMeshNodeContext
+			FInterchangeMeshPayLoadKey GlobalPayloadKey = MorphTargetKey;
+			if (SceneGlobalTransform.IsSet())
 			{
-				const UInterchangeMeshNode* MeshNode = nullptr;
-				const UInterchangeSceneNode* SceneNode = nullptr;
-				TOptional<FTransform> SceneGlobalTransform;
-				FInterchangeMeshPayLoadKey TranslatorPayloadKey;
-
-				//Return a new key with the translator key merge with the transform
-				FInterchangeMeshPayLoadKey GetTranslatorAndTransformPayloadKey() const
-				{
-					FInterchangeMeshPayLoadKey GlobalPayloadKey = TranslatorPayloadKey;
-					GlobalPayloadKey.UniqueId = GetUniqueId();
-					return GlobalPayloadKey;
-				}
+				GlobalPayloadKey.UniqueId += FInterchangeMeshPayLoadKey::GetTransformString(SceneGlobalTransform.GetValue());
+			}
+			return GlobalPayloadKey;
+		}
 
-				FInterchangeMeshPayLoadKey GetMorphTargetAndTransformPayloadKey(const FInterchangeMeshPayLoadKey& MorphTargetKey) const
-				{
-					FInterchangeMeshPayLoadKey GlobalPayloadKey = MorphTargetKey;
-					if (SceneGlobalTransform.IsSet())
-					{
-						GlobalPayloadKey.UniqueId += FInterchangeMeshPayLoadKey::GetTransformString(SceneGlobalTransform.GetValue());
-					}
-					return GlobalPayloadKey;
-				}
-				
-				//Return the translator key merge with the transform
-				FString GetUniqueId() const
-				{
-					FString UniqueId = TranslatorPayloadKey.UniqueId;
-					if (SceneGlobalTransform.IsSet())
-					{
-						UniqueId += FInterchangeMeshPayLoadKey::GetTransformString(SceneGlobalTransform.GetValue());
-					}
-					return UniqueId;
-				}
-			};
+		//Return the translator key merge with the transform
+		FString FMeshNodeContext::GetUniqueId() const
+		{
+			FString UniqueId = TranslatorPayloadKey.UniqueId;
+			if (SceneGlobalTransform.IsSet())
+			{
+				UniqueId += FInterchangeMeshPayLoadKey::GetTransformString(SceneGlobalTransform.GetValue());
+			}
+			return UniqueId;
+		}
 
+		namespace Private
+		{
 			void FillMorphTargetMeshDescriptionsPerMorphTargetName(const FMeshNodeContext& MeshNodeContext
 																 , TMap<FString, TOptional<UE::Interchange::FMeshPayloadData>>& MorphTargetMeshDescriptionsPerMorphTargetName
 																 , UInterchangeSkeletalMeshFactory::FLodPayloads& LodPayloads
@@ -393,7 +382,7 @@ namespace UE
 					const FString& JointName = JointNames[JointIndex];
 
 					const UInterchangeSceneNode* JointNode = RecursiveFindJointByName(NodeContainer, RootJointNodeId, JointName);
-					if (!ensure(JointNode))
+					if (!JointNode)
 					{
 						continue;
 					}
@@ -1017,6 +1006,129 @@ namespace UE
 				ContentInfo.bApplySkinningOnly = ContentInfo.bApplyPartialContent && ContentInfo.bApplySkinning;
 				return ContentInfo;
 			}
+
+			bool BuildMeshReferences(const UInterchangeFactoryBase::FImportAssetObjectParams& Arguments
+				, const FString& RootJointNodeId
+				, UInterchangeSkeletalMeshFactory::FImportAssetObjectLODData& ImportAssetObjectLODDataRef
+				, const FTransform& GlobalOffsetTransform
+				, const UInterchangeSkeletalMeshLodDataNode* LodDataNode
+				, const bool bBakeMeshes)
+			{
+				if (!LodDataNode)
+				{
+					return false;
+				}
+				const UInterchangeSceneNode* RootJointNode = Cast<UInterchangeSceneNode>(Arguments.NodeContainer->GetNode(RootJointNodeId));
+				if (!RootJointNode)
+				{
+					return false;
+				}
+				if (!ImportAssetObjectLODDataRef.bUseTimeZeroAsBindPose)
+				{
+					bool bHasBoneWithoutBindPose = false;
+					UE::Interchange::Private::FSkeletonHelper::RecursiveBoneHasBindPose(Arguments.NodeContainer, RootJointNodeId, bHasBoneWithoutBindPose);
+					if (bHasBoneWithoutBindPose)
+					{
+						ImportAssetObjectLODDataRef.bUseTimeZeroAsBindPose = true;
+					}
+				}
+
+				ImportAssetObjectLODDataRef.bDiffPose = false;
+				TArray <Private::FJointInfo> JointInfos;
+				TArray<FString> BoneNotBindNames;
+				Private::FSkeletonHelper::RecursiveAddBones(Arguments.NodeContainer
+					, RootJointNodeId
+					, JointInfos
+					, INDEX_NONE
+					, ImportAssetObjectLODDataRef.RefBonesBinary
+					, ImportAssetObjectLODDataRef.bUseTimeZeroAsBindPose
+					, ImportAssetObjectLODDataRef.bDiffPose
+					, BoneNotBindNames);
+
+				FTransform RootJointNodeGlobalTransform;
+				ensure(RootJointNode->GetCustomGlobalTransform(Arguments.NodeContainer, GlobalOffsetTransform, RootJointNodeGlobalTransform));
+				FTransform RootJointNodeLocalTransform;
+				ensure(RootJointNode->GetCustomLocalTransform(RootJointNodeLocalTransform));
+				FTransform BakeToRootJointTransfromModifier = RootJointNodeGlobalTransform.Inverse() * RootJointNodeLocalTransform; //It is used for !bBakeMeshes, and so the GlobalTransform will be inversed out when multiplied into the CustomBindPoseGlobalTransform
+				FTransform BakeFromRootJointTransfromModifier = RootJointNodeLocalTransform.Inverse() * RootJointNodeGlobalTransform * GlobalOffsetTransform.Inverse(); //GlobalOffseTransform will be added by the BindPoseGlobalTransform when bBakeMeshes && !bRootAncestorOfMeshDependency is used
+
+				//Scope to query the mesh node
+				{
+					TArray<FString> MeshUids;
+					LodDataNode->GetMeshUids(MeshUids);
+					ImportAssetObjectLODDataRef.MeshNodeContexts.Reserve(MeshUids.Num());
+					for (const FString& MeshUid : MeshUids)
+					{
+						FMeshNodeContext MeshReference;
+						MeshReference.MeshNode = Cast<UInterchangeMeshNode>(Arguments.NodeContainer->GetNode(MeshUid));
+						if (!MeshReference.MeshNode)
+						{
+							//The reference is a scene node and we need to bake the geometry
+							MeshReference.SceneNode = Cast<UInterchangeSceneNode>(Arguments.NodeContainer->GetNode(MeshUid));
+							if (!ensure(MeshReference.SceneNode != nullptr))
+							{
+								continue;
+							}
+							FString MeshDependencyUid;
+							MeshReference.SceneNode->GetCustomAssetInstanceUid(MeshDependencyUid);
+							MeshReference.MeshNode = Cast<UInterchangeMeshNode>(Arguments.NodeContainer->GetNode(MeshDependencyUid));
+							bool bRootAncestorOfSceneNode = Arguments.NodeContainer->GetIsAncestor(MeshReference.SceneNode->GetUniqueID(), RootJointNode->GetParentUid());
+							//Cache the scene node global matrix, we will use this matrix to bake the vertices, add the node geometric mesh offset to this matrix to bake it properly
+							FTransform SceneNodeTransform;
+							if (!ImportAssetObjectLODDataRef.bUseTimeZeroAsBindPose || !MeshReference.SceneNode->GetCustomTimeZeroGlobalTransform(Arguments.NodeContainer, GlobalOffsetTransform, SceneNodeTransform))
+							{
+								ensure(MeshReference.SceneNode->GetCustomBindPoseGlobalTransform(Arguments.NodeContainer, GlobalOffsetTransform, SceneNodeTransform));
+								if (bRootAncestorOfSceneNode)
+								{
+									if (!bBakeMeshes)
+									{
+										SceneNodeTransform *= BakeToRootJointTransfromModifier;
+									}
+								}
+								else
+								{
+									if (bBakeMeshes)
+									{
+										SceneNodeTransform = BakeFromRootJointTransfromModifier * SceneNodeTransform;
+									}
+									else
+									{
+										SceneNodeTransform *= GlobalOffsetTransform.Inverse();
+									}
+								}
+							}
+
+							FTransform SceneNodeGeometricTransform;
+							if (MeshReference.SceneNode->GetCustomGeometricTransform(SceneNodeGeometricTransform))
+							{
+								SceneNodeTransform = SceneNodeGeometricTransform * SceneNodeTransform;
+							}
+							MeshReference.SceneGlobalTransform = SceneNodeTransform;
+						}
+						else
+						{
+							MeshReference.SceneGlobalTransform = GlobalOffsetTransform;
+						}
+
+						if (!ensure(MeshReference.MeshNode != nullptr))
+						{
+							continue;
+						}
+
+						TOptional<FInterchangeMeshPayLoadKey> OptionalPayLoadKey = MeshReference.MeshNode->GetPayLoadKey();
+						if (OptionalPayLoadKey.IsSet())
+						{
+							MeshReference.TranslatorPayloadKey = OptionalPayLoadKey.GetValue();
+						}
+						else
+						{
+							continue;
+						}
+						ImportAssetObjectLODDataRef.MeshNodeContexts.Add(MeshReference);
+					}
+				}
+				return true;
+			}
 		} //Namespace Private
 	} //namespace Interchange
 } //namespace UE
@@ -1072,6 +1184,7 @@ void UInterchangeSkeletalMeshFactory::CreatePayloadTasks(const FImportAssetObjec
 	SkeletalMeshFactoryNode->GetLodDataUniqueIds(LodDataUniqueIds);
 	ensure(LodDataUniqueIds.Num() == LodCount);
 	PayloadsPerLodIndex.Reserve(LodCount);
+	int32 CurrentLodIndex = 0;
 	for (int32 LodIndex = 0; LodIndex < LodCount; ++LodIndex)
 	{
 		FString LodUniqueId = LodDataUniqueIds[LodIndex];
@@ -1106,104 +1219,20 @@ void UInterchangeSkeletalMeshFactory::CreatePayloadTasks(const FImportAssetObjec
 
 		FLodPayloads& LodPayloads = PayloadsPerLodIndex.FindOrAdd(LodIndex);
 
-		TArray<SkeletalMeshImportData::FBone> RefBonesBinary;
-		bool bUseTimeZeroAsBindPose = false;
-		bool bDiffPose = false;
+		FImportAssetObjectLODData& ImportAssetObjectLODData = ImportAssetObjectData.LodDatas.AddDefaulted_GetRef();
+		ImportAssetObjectLODData.LodIndex = CurrentLodIndex;
 
-		TArray <Private::FJointInfo> JointInfos;
-		TArray<FString> BoneNotBindNames;
-		Private::FSkeletonHelper::RecursiveAddBones(Arguments.NodeContainer
+		bool bUseTimeZeroAsBindPose = false;
+		SkeletonNode->GetCustomUseTimeZeroForBindPose(ImportAssetObjectLODData.bUseTimeZeroAsBindPose);
+		
+		Private::BuildMeshReferences(
+			Arguments
 			, RootJointNodeId
-			, JointInfos
-			, INDEX_NONE
-			, RefBonesBinary
-			, bUseTimeZeroAsBindPose
-			, bDiffPose
-			, BoneNotBindNames);
-
-		FTransform RootJointNodeGlobalTransform;
-		ensure(RootJointNode->GetCustomGlobalTransform(Arguments.NodeContainer, GlobalOffsetTransform, RootJointNodeGlobalTransform));
-		FTransform RootJointNodeLocalTransform;
-		ensure(RootJointNode->GetCustomLocalTransform(RootJointNodeLocalTransform));
-		FTransform BakeToRootJointTransfromModifier = RootJointNodeGlobalTransform.Inverse() * RootJointNodeLocalTransform; //It is used for !bBakeMeshes, and so the GlobalTransform will be inversed out when multiplied into the CustomBindPoseGlobalTransform
-		FTransform BakeFromRootJointTransfromModifier = RootJointNodeLocalTransform.Inverse() * RootJointNodeGlobalTransform * GlobalOffsetTransform.Inverse(); //GlobalOffseTransform will be added by the BindPoseGlobalTransform when bBakeMeshes && !bRootAncestorOfMeshDependency is used
-
-		TArray<Private::FMeshNodeContext> MeshReferences;
-		//Scope to query the mesh node
-		{
-			TArray<FString> MeshUids;
-			LodDataNode->GetMeshUids(MeshUids);
-			MeshReferences.Reserve(MeshUids.Num());
-			for (const FString& MeshUid : MeshUids)
-			{
-				Private::FMeshNodeContext MeshReference;
-				MeshReference.MeshNode = Cast<UInterchangeMeshNode>(Arguments.NodeContainer->GetNode(MeshUid));
-				if (!MeshReference.MeshNode)
-				{
-					//The reference is a scene node and we need to bake the geometry
-					MeshReference.SceneNode = Cast<UInterchangeSceneNode>(Arguments.NodeContainer->GetNode(MeshUid));
-					if (!ensure(MeshReference.SceneNode != nullptr))
-					{
-						continue;
-					}
-					FString MeshDependencyUid;
-					MeshReference.SceneNode->GetCustomAssetInstanceUid(MeshDependencyUid);
-					MeshReference.MeshNode = Cast<UInterchangeMeshNode>(Arguments.NodeContainer->GetNode(MeshDependencyUid));
-					bool bRootAncestorOfSceneNode = Arguments.NodeContainer->GetIsAncestor(MeshReference.SceneNode->GetUniqueID(), RootJointNode->GetParentUid());
-					//Cache the scene node global matrix, we will use this matrix to bake the vertices, add the node geometric mesh offset to this matrix to bake it properly
-					FTransform SceneNodeTransform;
-					if (!bUseTimeZeroAsBindPose || !MeshReference.SceneNode->GetCustomTimeZeroGlobalTransform(Arguments.NodeContainer, GlobalOffsetTransform, SceneNodeTransform))
-					{
-						ensure(MeshReference.SceneNode->GetCustomBindPoseGlobalTransform(Arguments.NodeContainer, GlobalOffsetTransform, SceneNodeTransform));
-						if (bRootAncestorOfSceneNode)
-						{
-							if (!bBakeMeshes)
-							{
-								SceneNodeTransform *= BakeToRootJointTransfromModifier;
-							}
-						}
-						else
-						{
-							if (bBakeMeshes)
-							{
-								SceneNodeTransform = BakeFromRootJointTransfromModifier * SceneNodeTransform;
-							}
-							else
-							{
-								SceneNodeTransform *= GlobalOffsetTransform.Inverse();
-							}
-						}
-					}
-
-					FTransform SceneNodeGeometricTransform;
-					if (MeshReference.SceneNode->GetCustomGeometricTransform(SceneNodeGeometricTransform))
-					{
-						SceneNodeTransform = SceneNodeGeometricTransform * SceneNodeTransform;
-					}
-					MeshReference.SceneGlobalTransform = SceneNodeTransform;
-				}
-				else
-				{
-					MeshReference.SceneGlobalTransform = GlobalOffsetTransform;
-				}
-
-				if (!ensure(MeshReference.MeshNode != nullptr))
-				{
-					continue;
-				}
-
-				TOptional<FInterchangeMeshPayLoadKey> OptionalPayLoadKey = MeshReference.MeshNode->GetPayLoadKey();
-				if (OptionalPayLoadKey.IsSet())
-				{
-					MeshReference.TranslatorPayloadKey = OptionalPayLoadKey.GetValue();
-				}
-				else
-				{
-					continue;
-				}
-				MeshReferences.Add(MeshReference);
-			}
-		}
+			, ImportAssetObjectLODData
+			, GlobalOffsetTransform
+			, LodDataNode
+			, bBakeMeshes
+		);
 
 		struct FInternalInstanceData
 		{
@@ -1215,7 +1244,7 @@ void UInterchangeSkeletalMeshFactory::CreatePayloadTasks(const FImportAssetObjec
 			}
 		};
 		TMap<FString, FInternalInstanceData> MeshInstancesDatas;
-		for (const Private::FMeshNodeContext& MeshNodeContext : MeshReferences)
+		for (const FMeshNodeContext& MeshNodeContext : ImportAssetObjectLODData.MeshNodeContexts)
 		{
 			FInternalInstanceData& InstanceData = MeshInstancesDatas.FindOrAdd(MeshNodeContext.TranslatorPayloadKey.UniqueId);
 			InstanceData.Count++;
@@ -1223,9 +1252,9 @@ void UInterchangeSkeletalMeshFactory::CreatePayloadTasks(const FImportAssetObjec
 		}
 		
 		//Reserve the correct amount since we point in the array for the lambda so array must not be resized at any moments after we create the tasks
-		LodPayloads.MeshPayloadPerKey.Reserve(MeshReferences.Num());
+		LodPayloads.MeshPayloadPerKey.Reserve(ImportAssetObjectLODData.MeshNodeContexts.Num());
 		int32 MorphTargetCount = 0;
-		for (const Private::FMeshNodeContext& MeshNodeContext : MeshReferences)
+		for (const FMeshNodeContext& MeshNodeContext : ImportAssetObjectLODData.MeshNodeContexts)
 		{
 			//Count the morph target dependencies so we can reserve the right amount
 			if (bImportMorphTarget)
@@ -1235,7 +1264,7 @@ void UInterchangeSkeletalMeshFactory::CreatePayloadTasks(const FImportAssetObjec
 		}
 		LodPayloads.MorphPayloadPerKey.Reserve(MorphTargetCount);
 
-		for (const Private::FMeshNodeContext& MeshNodeContext : MeshReferences)
+		for (const FMeshNodeContext& MeshNodeContext : ImportAssetObjectLODData.MeshNodeContexts)
 		{
 			const FInternalInstanceData& InstanceData = MeshInstancesDatas.FindChecked(MeshNodeContext.TranslatorPayloadKey.UniqueId);
 			FTransform ApplyTransformWhenFetchPayload = InstanceData.ShouldFetchWithTransform() ? MeshNodeContext.SceneGlobalTransform.Get(FTransform::Identity) : FTransform::Identity;
@@ -1278,6 +1307,7 @@ void UInterchangeSkeletalMeshFactory::CreatePayloadTasks(const FImportAssetObjec
 				}
 			}
 		}
+		CurrentLodIndex++;
 	}
 #endif
 }
@@ -1482,29 +1512,15 @@ UInterchangeFactoryBase::FImportAssetResult UInterchangeSkeletalMeshFactory::Beg
 			Message->Text = WarningMessage_InvalidRootJoint;
 			continue;
 		}
-
-		FImportAssetObjectLODData& ImportAssetObjectLODData = ImportAssetObjectData.LodDatas.AddDefaulted_GetRef();
-		ImportAssetObjectLODData.LodIndex = CurrentLodIndex;
-
-		bool bUseTimeZeroAsBindPose = false;
-		SkeletonNode->GetCustomUseTimeZeroForBindPose(bUseTimeZeroAsBindPose);
-		if (!bUseTimeZeroAsBindPose)
+		if (!ImportAssetObjectData.LodDatas.IsValidIndex(CurrentLodIndex))
 		{
-			bool bHasBoneWithoutBindPose = false;
-			UE::Interchange::Private::FSkeletonHelper::RecursiveBoneHasBindPose(Arguments.NodeContainer, RootJointNodeId, bHasBoneWithoutBindPose);
-			if (bHasBoneWithoutBindPose)
-			{
-				UInterchangeResultDisplay_Generic* Message = AddMessage<UInterchangeResultDisplay_Generic>();
-				Message->Text = FText::Format(NSLOCTEXT("InterchangeSkeletalMeshFactory", "BeginImportAsset_GameThread_NotAllJointsHaveBindPose", "Not all joints have BindPoses in the skeleton {0}, will use T0 as BindPose instead.")
-					, FText::FromString(SkeletonReference->GetName()));
-
-				UE_LOG(LogInterchangeImport, Display, TEXT("Not all joints have BindPoses in the skeleton %s, will use T0 as BindPose instead."), *SkeletonReference->GetName());
-
-				bUseTimeZeroAsBindPose = true;
-			}
+			UInterchangeResultWarning_Generic* Message = AddMessage<UInterchangeResultWarning_Generic>();
+			Message->Text = FText::Format(NSLOCTEXT("InterchangeSkeletalMeshFactory", "BeginImportAsset_GameThread_BadLodIndexPrecomputed", "Invalid precompute LOD {0} data when importing SkeletalMesh asset {1}.")
+				, FText::AsNumber(LodIndex)
+				, FText::FromString(Arguments.AssetName));
+			continue;
 		}
-
-		ImportAssetObjectLODData.bUseTimeZeroAsBindPose = bUseTimeZeroAsBindPose;
+		FImportAssetObjectLODData& ImportAssetObjectLODData = ImportAssetObjectData.LodDatas[CurrentLodIndex];
 		
 		//Do not alter the skeletal mesh reference skeleton when importing geometry only
 		FReferenceSkeleton RefSkeleton;
@@ -1707,110 +1723,13 @@ UInterchangeFactoryBase::FImportAssetResult UInterchangeSkeletalMeshFactory::Imp
 		FImportAssetObjectLODData& ImportAssetObjectLODData = ImportAssetObjectData.LodDatas[CurrentLodIndex];
 		ensure(ImportAssetObjectLODData.LodIndex == CurrentLodIndex);
 
-		FTransform RootJointNodeGlobalTransform;
-		ensure(RootJointNode->GetCustomGlobalTransform(Arguments.NodeContainer, GlobalOffsetTransform, RootJointNodeGlobalTransform));
-		FTransform RootJointNodeLocalTransform;
-		ensure(RootJointNode->GetCustomLocalTransform(RootJointNodeLocalTransform));
-		FTransform BakeToRootJointTransfromModifier = RootJointNodeGlobalTransform.Inverse() * RootJointNodeLocalTransform; //It is used for !bBakeMeshes, and so the GlobalTransform will be inversed out when multiplied into the CustomBindPoseGlobalTransform
-		FTransform BakeFromRootJointTransfromModifier = RootJointNodeLocalTransform.Inverse() * RootJointNodeGlobalTransform * GlobalOffsetTransform.Inverse(); //GlobalOffseTransform will be added by the BindPoseGlobalTransform when bBakeMeshes && !bRootAncestorOfMeshDependency is used
-
-		TArray<UE::Interchange::Private::FMeshNodeContext> MeshReferences;
-		//Scope to query the mesh node
-		{
-			TArray<FString> MeshUids;
-			LodDataNode->GetMeshUids(MeshUids);
-			MeshReferences.Reserve(MeshUids.Num());
-
-			FText WarningMessage_InvalidLODMeshReference = FText::Format(NSLOCTEXT("InterchangeSkeletalMeshFactory", "ImportAsset_Async_InvalidLODMeshReference", "Invalid LOD {0} mesh reference when importing SkeletalMesh asset {1}")
-				, FText::AsNumber(LodIndex)
-				, FText::FromString(Arguments.AssetName));
-
-			for (const FString& MeshUid : MeshUids)
-			{
-				UE::Interchange::Private::FMeshNodeContext MeshReference;
-				MeshReference.MeshNode = Cast<UInterchangeMeshNode>(Arguments.NodeContainer->GetNode(MeshUid));
-				if (!MeshReference.MeshNode)
-				{
-					//The reference is a scene node and we need to bake the geometry
-					MeshReference.SceneNode = Cast<UInterchangeSceneNode>(Arguments.NodeContainer->GetNode(MeshUid));
-					if (!ensure(MeshReference.SceneNode != nullptr))
-					{
-						UInterchangeResultWarning_Generic* Message = AddMessage<UInterchangeResultWarning_Generic>();
-						Message->Text = WarningMessage_InvalidLODMeshReference;
-						continue;
-					}
-					FString MeshDependencyUid;
-					MeshReference.SceneNode->GetCustomAssetInstanceUid(MeshDependencyUid);
-					MeshReference.MeshNode = Cast<UInterchangeMeshNode>(Arguments.NodeContainer->GetNode(MeshDependencyUid));
-					bool bRootAncestorOfSceneNode = Arguments.NodeContainer->GetIsAncestor(MeshReference.SceneNode->GetUniqueID(), RootJointNode->GetParentUid());
-					//Cache the scene node global matrix, we will use this matrix to bake the vertices, add the node geometric mesh offset to this matrix to bake it properly
-					FTransform SceneNodeTransform;
-					if (!ImportAssetObjectLODData.bUseTimeZeroAsBindPose || !MeshReference.SceneNode->GetCustomTimeZeroGlobalTransform(Arguments.NodeContainer, GlobalOffsetTransform, SceneNodeTransform))
-					{
-						ensure(MeshReference.SceneNode->GetCustomBindPoseGlobalTransform(Arguments.NodeContainer, GlobalOffsetTransform, SceneNodeTransform));
-						if (bRootAncestorOfSceneNode)
-						{
-							if (!bBakeMeshes)
-							{
-								SceneNodeTransform *= BakeToRootJointTransfromModifier;
-							}
-						}
-						else
-						{
-							if (bBakeMeshes)
-							{
-								SceneNodeTransform = BakeFromRootJointTransfromModifier * SceneNodeTransform;
-							}
-							else
-							{
-								SceneNodeTransform *= GlobalOffsetTransform.Inverse();
-							}
-						}
-					}
-
-					FTransform SceneNodeGeometricTransform;
-					if(MeshReference.SceneNode->GetCustomGeometricTransform(SceneNodeGeometricTransform))
-					{
-						SceneNodeTransform = SceneNodeGeometricTransform * SceneNodeTransform;
-					}
-					MeshReference.SceneGlobalTransform = SceneNodeTransform;
-				}
-				else
-				{
-					MeshReference.SceneGlobalTransform = GlobalOffsetTransform;
-				}
-
-				if (!ensure(MeshReference.MeshNode != nullptr))
-				{
-					UInterchangeResultWarning_Generic* Message = AddMessage<UInterchangeResultWarning_Generic>();
-					Message->Text = WarningMessage_InvalidLODMeshReference;
-					continue;
-				}
-
-				TOptional<FInterchangeMeshPayLoadKey> OptionalPayLoadKey = MeshReference.MeshNode->GetPayLoadKey();
-				if (OptionalPayLoadKey.IsSet())
-				{
-					MeshReference.TranslatorPayloadKey = OptionalPayLoadKey.GetValue();
-				}
-				else
-				{
-					UInterchangeResultWarning_Generic* Message = AddMessage<UInterchangeResultWarning_Generic>();
-					Message->Text = FText::Format(NSLOCTEXT("InterchangeSkeletalMeshFactory", "ImportAsset_Async_EmptyLODMeshReference", "Empty LOD {0} mesh reference payload when importing SkeletalMesh asset {1}.")
-						, FText::AsNumber(LodIndex)
-						, FText::FromString(Arguments.AssetName));
-					continue;
-				}
-				MeshReferences.Add(MeshReference);
-			}
-		}
-
 		//Add the lod mesh data to the skeletalmesh
 		FSkeletalMeshImportData SkeletalMeshImportData;
 		const bool bSkinControlPointToTimeZero = ImportAssetObjectLODData.bUseTimeZeroAsBindPose && ImportAssetObjectLODData.bDiffPose;
 		//Get all meshes and morph targets payload and fill the SkeletalMeshImportData structure
 		UE::Interchange::Private::RetrieveAllSkeletalMeshPayloadsAndFillImportData(SkeletalMeshFactoryNode
 																					, SkeletalMeshImportData
-																					, MeshReferences
+																					, ImportAssetObjectLODData.MeshNodeContexts
 																					, LodPayloads
 																					, ImportAssetObjectLODData.RefBonesBinary
 																					, Arguments
diff --git a/Engine/Plugins/Interchange/Runtime/Source/Import/Public/Mesh/InterchangeSkeletalMeshFactory.h b/Engine/Plugins/Interchange/Runtime/Source/Import/Public/Mesh/InterchangeSkeletalMeshFactory.h
--- a/Engine/Plugins/Interchange/Runtime/Source/Import/Public/Mesh/InterchangeSkeletalMeshFactory.h
+++ b/Engine/Plugins/Interchange/Runtime/Source/Import/Public/Mesh/InterchangeSkeletalMeshFactory.h
@@ -14,10 +14,31 @@
 
 #include "InterchangeSkeletalMeshFactory.generated.h"
 
+class UInterchangeSceneNode;
 class UInterchangeSkeletalMeshFactoryNode;
 class USkeletalMesh;
 class USkeleton;
 
+namespace UE::Interchange
+{
+	//Get the mesh node context for each MeshUids
+	struct FMeshNodeContext
+	{
+		const UInterchangeMeshNode* MeshNode = nullptr;
+		const UInterchangeSceneNode* SceneNode = nullptr;
+		TOptional<FTransform> SceneGlobalTransform;
+		FInterchangeMeshPayLoadKey TranslatorPayloadKey;
+
+		//Return a new key with the translator key merge with the transform
+		FInterchangeMeshPayLoadKey GetTranslatorAndTransformPayloadKey() const;
+
+		FInterchangeMeshPayLoadKey GetMorphTargetAndTransformPayloadKey(const FInterchangeMeshPayLoadKey& MorphTargetKey) const;
+
+		//Return the translator key merge with the transform
+		FString GetUniqueId() const;
+	};
+} //UE::Interchange
+
 UCLASS(BlueprintType)
 class INTERCHANGEIMPORT_API UInterchangeSkeletalMeshFactory : public UInterchangeFactoryBase
 {
@@ -31,6 +52,7 @@ public:
 		TArray<SkeletalMeshImportData::FMaterial> ImportedMaterials;
 		TArray<SkeletalMeshImportData::FBone> RefBonesBinary;
 #endif
+		TArray<UE::Interchange::FMeshNodeContext> MeshNodeContexts;
 		bool bUseTimeZeroAsBindPose = false;
 		bool bDiffPose = false;
 	};
