USD:
- Fix a compile caused by a translation unit making a wrong ADL caused by a namespace (link to horde: https://horde.devtools.epicgames.com/job/66febfe3ca42997879777e63?step=e719)

diff --git a/Engine/Plugins/Interchange/Runtime/Source/Import/Private/MaterialX/MaterialXUtils/MaterialXSurfaceShaderAbstract.cpp b/Engine/Plugins/Interchange/Runtime/Source/Import/Private/MaterialX/MaterialXUtils/MaterialXSurfaceShaderAbstract.cpp
--- a/Engine/Plugins/Interchange/Runtime/Source/Import/Private/MaterialX/MaterialXUtils/MaterialXSurfaceShaderAbstract.cpp
+++ b/Engine/Plugins/Interchange/Runtime/Source/Import/Private/MaterialX/MaterialXUtils/MaterialXSurfaceShaderAbstract.cpp
@@ -1431,7 +1431,7 @@ FString FMaterialXSurfaceShaderAbstract::GetAttributeParentName(MaterialX::NodeP
 	return ParentName.IsEmpty() ? Node->getName().c_str() : ParentName;
 }
 
-namespace UE::MaterialX::Private
+namespace
 {
 	void GetUniqueName(mx::ElementPtr Element, TStringBuilder<256>& Buffer)
 	{
@@ -1459,7 +1459,7 @@ FString FMaterialXSurfaceShaderAbstract::GetUniqueName(MaterialX::ElementPtr Ele
 	else
 	{
 		TStringBuilder<256> Buffer;
-		UE::MaterialX::Private::GetUniqueName(Element, Buffer);
+		::GetUniqueName(Element, Buffer);
 		return FString{ Buffer.ToString() }.TrimChar(TEXT('_')); // remove any trailing '_'
 	}
 }
