USD: Remove ensure when handling texture payloads via the new GetMaterialXTexturePayloadData function as it may fail with user data;

Run auto-formatter;
Put free function inside internal namespace.

diff --git a/Engine/Plugins/Interchange/Runtime/Source/Import/Private/Usd/InterchangeUsdTranslator.cpp b/Engine/Plugins/Interchange/Runtime/Source/Import/Private/Usd/InterchangeUsdTranslator.cpp
--- a/Engine/Plugins/Interchange/Runtime/Source/Import/Private/Usd/InterchangeUsdTranslator.cpp
+++ b/Engine/Plugins/Interchange/Runtime/Source/Import/Private/Usd/InterchangeUsdTranslator.cpp
@@ -49,8 +49,8 @@
 
 #include "UObject/GCObjectScopeGuard.h"
 
-#include "MaterialX/MaterialXUtils/MaterialXManager.h"
 #include "MaterialX/MaterialXUtils/MaterialXBase.h"
+#include "MaterialX/MaterialXUtils/MaterialXManager.h"
 #include "MaterialXFormat/XmlIo.h"
 
 #if USE_USD_SDK
@@ -154,12 +154,19 @@ namespace UE::InterchangeUsdTranslator::Private
 		bool AddMaterialXShaderGraph(
 			const UE::FUsdPrim& Prim,
 			UInterchangeUsdTranslatorSettings* TranslatorSettings,
-			UInterchangeBaseNodeContainer& NodeContainer);
+			UInterchangeBaseNodeContainer& NodeContainer
+		);
 
 		/**
-		 * Set the material slot dependencies on the meshes, if bIterateShaderGraphs is true, it will iterate over the shader graphs to rename them according to the MaterialPrimName
+		 * Set the material slot dependencies on the meshes, if bIterateShaderGraphs is true, it will iterate over the shader graphs to rename them
+		 * according to the MaterialPrimName
 		 */
-		void SetMaterialSlotDependencies(const FString& MaterialPrimName, const FString& MaterialUid, UInterchangeBaseNodeContainer& NodeContainer, bool bIterateShaderGraphs = false);
+		void SetMaterialSlotDependencies(
+			const FString& MaterialPrimName,
+			const FString& MaterialUid,
+			UInterchangeBaseNodeContainer& NodeContainer,
+			bool bIterateShaderGraphs = false
+		);
 
 	public:
 		// We have to keep a stage reference so that we can parse the payloads after Translate() complete.
@@ -723,10 +730,14 @@ namespace UE::InterchangeUsdTranslator::Private
 		const FString* ParameterName = nullptr;
 	};
 
-
-	void UInterchangeUSDTranslatorImpl::SetMaterialSlotDependencies(const FString& MaterialPrimName, const FString& MaterialUid, UInterchangeBaseNodeContainer& NodeContainer, bool bIterateShaderGraphs)
+	void UInterchangeUSDTranslatorImpl::SetMaterialSlotDependencies(
+		const FString& MaterialPrimName,
+		const FString& MaterialUid,
+		UInterchangeBaseNodeContainer& NodeContainer,
+		bool bIterateShaderGraphs
+	)
 	{
-		if(bIterateShaderGraphs)
+		if (bIterateShaderGraphs)
 		{
 			// The material from the MaterialXTranslator doesn't have the same UID, both the PrimPath have the same name but not the same path
 			// We need to retrieve that name (which is the Material name) in the translator, then we can map it to the right mesh
@@ -734,7 +745,7 @@ namespace UE::InterchangeUsdTranslator::Private
 				[this, &MaterialPrimName, &MaterialUid](const FString&, UInterchangeShaderGraphNode* ShaderGraphNode)
 				{
 					FString UID = ShaderGraphNode->GetUniqueID();
-					if(FPaths::GetBaseFilename(UID) == MaterialPrimName)
+					if (FPaths::GetBaseFilename(UID) == MaterialPrimName)
 					{
 						MaterialUidToActualNodeUid.Add(MaterialUid, UID);
 						return true;
@@ -748,13 +759,13 @@ namespace UE::InterchangeUsdTranslator::Private
 		}
 
 		// Now we need to check if we have to set the slot of the mesh nodes here
-		if(TArray<FMaterialSlotMesh>* SlotMeshes = PrimPathToSlotMeshNodes.Find(MaterialUid))
+		if (TArray<FMaterialSlotMesh>* SlotMeshes = PrimPathToSlotMeshNodes.Find(MaterialUid))
 		{
-			if(FString* NewMaterialUID = MaterialUidToActualNodeUid.Find(MaterialUid))
+			if (FString* NewMaterialUID = MaterialUidToActualNodeUid.Find(MaterialUid))
 			{
-				for(const FMaterialSlotMesh& MaterialSlotMesh : *SlotMeshes)
+				for (const FMaterialSlotMesh& MaterialSlotMesh : *SlotMeshes)
 				{
-					if(!MaterialSlotMesh.MeshNode->GetSlotMaterialDependencyUid(MaterialSlotMesh.MaterialSlotName, *NewMaterialUID))
+					if (!MaterialSlotMesh.MeshNode->GetSlotMaterialDependencyUid(MaterialSlotMesh.MaterialSlotName, *NewMaterialUID))
 					{
 						MaterialSlotMesh.MeshNode->SetSlotMaterialDependencyUid(MaterialSlotMesh.MaterialSlotName, *NewMaterialUID);
 					}
@@ -777,20 +788,20 @@ namespace UE::InterchangeUsdTranslator::Private
 
 		// Check for any references of MaterialX
 #if WITH_EDITOR
-		if(RenderContext == UnrealIdentifiers::MaterialXRenderContext)
+		if (RenderContext == UnrealIdentifiers::MaterialXRenderContext)
 		{
 			TArray<FString> FilePaths = UsdUtils::GetMaterialXFilePaths(Prim);
-			for(const FString& File : FilePaths)
+			for (const FString& File : FilePaths)
 			{
 				// the file has already been handled no need to do a Translate again
-				if(!Translators.Find(File))
+				if (!Translators.Find(File))
 				{
 					UInterchangeManager& InterchangeManager = UInterchangeManager::GetInterchangeManager();
 					UInterchangeSourceData* SourceData = UInterchangeManager::CreateSourceData(File);
 
 					UInterchangeTranslatorBase* Translator = InterchangeManager.GetTranslatorForSourceData(SourceData);
 					// check on the Translator, it might return nullptr in case of reimport
-					if(Translator)
+					if (Translator)
 					{
 						Translator->Translate(NodeContainer);
 						Translators.Add(File, Translator);
@@ -798,28 +809,31 @@ namespace UE::InterchangeUsdTranslator::Private
 				}
 			}
 
-			if(!FilePaths.IsEmpty())
+			if (!FilePaths.IsEmpty())
 			{
 				return true;
 			}
 
 			// Only enable the shader graph on Windows for the time being as it crashes on Linux cause of a probable double free
 #if PLATFORM_WINDOWS
-			bool bHasShaderGraph = false;			
-			if(MaterialX::DocumentPtr Document = UsdToUnreal::CreateSurfaceMaterial(Prim, TranslatorSettings ? *RenderContext.ToString() : nullptr))
+			bool bHasShaderGraph = false;
+			if (MaterialX::DocumentPtr Document = UsdToUnreal::CreateSurfaceMaterial(Prim, TranslatorSettings ? *RenderContext.ToString() : nullptr))
 			{
-				if(MaterialX::NodePtr SurfaceMaterial = Document->getNode(TCHAR_TO_ANSI(*Prim.GetName().ToString())))
+				if (MaterialX::NodePtr SurfaceMaterial = Document->getNode(TCHAR_TO_ANSI(*Prim.GetName().ToString())))
 				{
 					try
 					{
-						TSharedPtr<FMaterialXBase> ShaderTranslator = FMaterialXManager::GetInstance().GetShaderTranslator(SurfaceMaterial->getCategory().c_str(), NodeContainer);
+						TSharedPtr<FMaterialXBase> ShaderTranslator = FMaterialXManager::GetInstance().GetShaderTranslator(
+							SurfaceMaterial->getCategory().c_str(),
+							NodeContainer
+						);
 						bHasShaderGraph = ShaderTranslator.IsValid();
-						if(bHasShaderGraph)
+						if (bHasShaderGraph)
 						{
 							ShaderTranslator->Translate(SurfaceMaterial);
 						}
 					}
-					catch(std::exception& Exception)
+					catch (std::exception& Exception)
 					{
 						UE_LOG(LogUsd, Error, TEXT("%s"), ANSI_TO_TCHAR(Exception.what()));
 					}
@@ -827,9 +841,9 @@ namespace UE::InterchangeUsdTranslator::Private
 
 				return bHasShaderGraph;
 			}
-#endif // PLATFORM_WINDOWS
+#endif	  // PLATFORM_WINDOWS
 		}
-#endif // WITH_EDITOR
+#endif	  // WITH_EDITOR
 		return false;
 	}
 
@@ -850,7 +864,7 @@ namespace UE::InterchangeUsdTranslator::Private
 		}
 
 		// we only create material instances if we didn't find any MaterialX instances (including a shader graph)
-		if(AddMaterialXShaderGraph(Prim, TranslatorSettings, NodeContainer))
+		if (AddMaterialXShaderGraph(Prim, TranslatorSettings, NodeContainer))
 		{
 			SetMaterialSlotDependencies(MaterialPrimName, MaterialUid, NodeContainer, true);
 			return;
@@ -2646,48 +2660,55 @@ TOptional<UE::Interchange::FMeshPayloadData> UInterchangeUSDTranslator::GetMeshP
 	return Result;
 }
 
-TOptional<UE::Interchange::FImportImage> GetMaterialXTexturePayloadData(const FString& PayloadKey, TOptional<FString>& AlternateTexturePath, const UInterchangeTranslatorBase* Translator)
+namespace UE::InterchangeUsdTranslator::Private
 {
-	FString Filename = PayloadKey;
-	TextureCompressionSettings CompressionSettings = TextureCompressionSettings::TC_Default;
+	TOptional<UE::Interchange::FImportImage> GetMaterialXTexturePayloadData(
+		const FString& PayloadKey,
+		TOptional<FString>& AlternateTexturePath,
+		const UInterchangeTranslatorBase* Translator
+	)
+	{
+		FString Filename = PayloadKey;
+		TextureCompressionSettings CompressionSettings = TextureCompressionSettings::TC_Default;
 
 #if WITH_EDITOR
-	if(int32 IndexTextureCompression; PayloadKey.FindChar(FMaterialXManager::TexturePayloadSeparator, IndexTextureCompression))
-	{
-		Filename = PayloadKey.Mid(0, IndexTextureCompression);
-		CompressionSettings = TextureCompressionSettings(FCString::Atoi(*PayloadKey.Mid(IndexTextureCompression + 1)));
-	}
+		if (int32 IndexTextureCompression; PayloadKey.FindChar(FMaterialXManager::TexturePayloadSeparator, IndexTextureCompression))
+		{
+			Filename = PayloadKey.Mid(0, IndexTextureCompression);
+			CompressionSettings = TextureCompressionSettings(FCString::Atoi(*PayloadKey.Mid(IndexTextureCompression + 1)));
+		}
 #endif
-	UInterchangeSourceData* PayloadSourceData = UInterchangeManager::GetInterchangeManager().CreateSourceData(Filename);
-	FGCObjectScopeGuard ScopedSourceData(PayloadSourceData);
+		UInterchangeSourceData* PayloadSourceData = UInterchangeManager::GetInterchangeManager().CreateSourceData(Filename);
+		FGCObjectScopeGuard ScopedSourceData(PayloadSourceData);
 
-	if(!PayloadSourceData)
-	{
-		return TOptional<UE::Interchange::FImportImage>();
-	}
+		if (!PayloadSourceData)
+		{
+			return TOptional<UE::Interchange::FImportImage>();
+		}
 
-	UInterchangeTranslatorBase* SourceTranslator = UInterchangeManager::GetInterchangeManager().GetTranslatorForSourceData(PayloadSourceData);
-	FGCObjectScopeGuard ScopedSourceTranslator(SourceTranslator);
-	const IInterchangeTexturePayloadInterface* TextureTranslator = Cast<IInterchangeTexturePayloadInterface>(SourceTranslator);
+		UInterchangeTranslatorBase* SourceTranslator = UInterchangeManager::GetInterchangeManager().GetTranslatorForSourceData(PayloadSourceData);
+		FGCObjectScopeGuard ScopedSourceTranslator(SourceTranslator);
+		const IInterchangeTexturePayloadInterface* TextureTranslator = Cast<IInterchangeTexturePayloadInterface>(SourceTranslator);
 
-	if(!ensure(TextureTranslator))
-	{
-		return TOptional<UE::Interchange::FImportImage>();
-	}
+		if (!TextureTranslator)
+		{
+			return TOptional<UE::Interchange::FImportImage>();
+		}
 
-	SourceTranslator->SetResultsContainer(Translator->Results);
+		SourceTranslator->SetResultsContainer(Translator->Results);
 
-	AlternateTexturePath = Filename;
+		AlternateTexturePath = Filename;
 
-	TOptional<UE::Interchange::FImportImage> TexturePayloadData = TextureTranslator->GetTexturePayloadData(PayloadKey, AlternateTexturePath);
+		TOptional<UE::Interchange::FImportImage> TexturePayloadData = TextureTranslator->GetTexturePayloadData(PayloadKey, AlternateTexturePath);
 
-	if(TexturePayloadData.IsSet())
-	{
-		TexturePayloadData.GetValue().CompressionSettings = CompressionSettings;
-	}
+		if (TexturePayloadData.IsSet())
+		{
+			TexturePayloadData.GetValue().CompressionSettings = CompressionSettings;
+		}
 
-	return TexturePayloadData;
-}
+		return TexturePayloadData;
+	}
+}	 // namespace UE::InterchangeUsdTranslator::Private
 
 TOptional<UE::Interchange::FImportImage> UInterchangeUSDTranslator::GetTexturePayloadData(
 	const FString& PayloadKey,
@@ -2742,8 +2763,9 @@ TOptional<UE::Interchange::FImportImage> UInterchangeUSDTranslator::GetTexturePa
 		}
 	}
 
-	//if we couldn't find a texture in either the USD Translator nor the Translators, then it's most likely coming from reading an mtlx in memory, copying the behavior from InterchangeMaterialXTranslator
-	if(!bTranslatorPayload)
+	// if we couldn't find a texture in either the USD Translator nor the Translators, then it's most likely coming from reading an mtlx in memory,
+	// copying the behavior from InterchangeMaterialXTranslator
+	if (!bTranslatorPayload)
 	{
 		TexturePayloadData = GetMaterialXTexturePayloadData(PayloadKey, AlternateTexturePath, this);
 	}
