USD: Fix textures not being parsed via Interchange translation in some scenarios.

diff --git a/Engine/Plugins/Interchange/Runtime/Source/Import/Private/Usd/InterchangeUsdTranslator.cpp b/Engine/Plugins/Interchange/Runtime/Source/Import/Private/Usd/InterchangeUsdTranslator.cpp
--- a/Engine/Plugins/Interchange/Runtime/Source/Import/Private/Usd/InterchangeUsdTranslator.cpp
+++ b/Engine/Plugins/Interchange/Runtime/Source/Import/Private/Usd/InterchangeUsdTranslator.cpp
@@ -2749,23 +2749,24 @@ TOptional<UE::Interchange::FImportImage> UInterchangeUSDTranslator::GetTexturePa
 	// We did not find a suitable Payload in USD Translator, let's find one in one of the Translators (MaterialX for the moment)
 	// The best way would be to have a direct association between the payload and the right Translator, but we don't have a suitable way of knowing
 	// which Payload belongs to which Translator So let's just loop over them all
-	bool bTranslatorPayload = false;
-	for (const TPair<FString, UInterchangeTranslatorBase*>& Pair : Impl->Translators)
+	if (!TexturePayloadData.IsSet())
 	{
-		if (IInterchangeTexturePayloadInterface* TexturePayloadInterface = Cast<IInterchangeTexturePayloadInterface>(Pair.Value))
+		for (const TPair<FString, UInterchangeTranslatorBase*>& Pair : Impl->Translators)
 		{
-			TexturePayloadData = TexturePayloadInterface->GetTexturePayloadData(PayloadKey, AlternateTexturePath);
-			if (TexturePayloadData)
+			if (IInterchangeTexturePayloadInterface* TexturePayloadInterface = Cast<IInterchangeTexturePayloadInterface>(Pair.Value))
 			{
-				bTranslatorPayload = true;
-				break;
+				TexturePayloadData = TexturePayloadInterface->GetTexturePayloadData(PayloadKey, AlternateTexturePath);
+				if (TexturePayloadData.IsSet())
+				{
+					break;
+				}
 			}
 		}
 	}
 
 	// if we couldn't find a texture in either the USD Translator nor the Translators, then it's most likely coming from reading an mtlx in memory,
 	// copying the behavior from InterchangeMaterialXTranslator
-	if (!bTranslatorPayload)
+	if (!TexturePayloadData.IsSet())
 	{
 		TexturePayloadData = GetMaterialXTexturePayloadData(PayloadKey, AlternateTexturePath, this);
 	}
