[Interchange-FBX] StaticAnalysis warning fix.

#rnx

diff --git a/Engine/Plugins/Interchange/Runtime/Source/Parsers/Fbx/Private/InterchangeFbxParser.cpp b/Engine/Plugins/Interchange/Runtime/Source/Parsers/Fbx/Private/InterchangeFbxParser.cpp
--- a/Engine/Plugins/Interchange/Runtime/Source/Parsers/Fbx/Private/InterchangeFbxParser.cpp
+++ b/Engine/Plugins/Interchange/Runtime/Source/Parsers/Fbx/Private/InterchangeFbxParser.cpp
@@ -62,28 +62,28 @@ namespace UE::Interchange
 		//Since we are not in main thread we cannot use TStrongPtr, so we will add the object to the root and remove it when we are done
 		UInterchangeBaseNodeContainer* Container = NewObject<UInterchangeBaseNodeContainer>(GetTransientPackage(), NAME_None);
 
-		if (!FbxParserPrivate->LoadFbxFile(Filename, *Container))
+		if (!ensure(Container != nullptr))
 		{
 			if (UInterchangeResultError_Generic* Error = AddMessage<UInterchangeResultError_Generic>())
 			{
 				Error->SourceAssetName = SourceFilename;
-				Error->Text = LOCTEXT("CantLoadFbxFile", "Cannot load the FBX file.");
+				Error->Text = LOCTEXT("CantAllocate", "Cannot allocate base node container to add FBX scene data.");
 			}
 			return;
 		}
 
-		ResultFilepath = ResultFolder + TEXT("/SceneDescription.itc");
-		
-		if (!ensure(Container != nullptr))
+		if (!FbxParserPrivate->LoadFbxFile(Filename, *Container))
 		{
 			if (UInterchangeResultError_Generic* Error = AddMessage<UInterchangeResultError_Generic>())
 			{
 				Error->SourceAssetName = SourceFilename;
-				Error->Text = LOCTEXT("CantAllocate", "Cannot allocate base node container to add FBX scene data.");
+				Error->Text = LOCTEXT("CantLoadFbxFile", "Cannot load the FBX file.");
 			}
 			return;
 		}
 
+		ResultFilepath = ResultFolder + TEXT("/SceneDescription.itc");
+
 		Container->AddToRoot();
 		FbxParserPrivate->FillContainerWithFbxScene(*Container);
 		Container->SaveToFile(ResultFilepath);
