Change Display Name for a Interchange Import Dialog

diff --git a/Engine/Plugins/Interchange/Runtime/Source/Pipelines/Public/InterchangeGenericAssetsPipelineSharedSettings.h b/Engine/Plugins/Interchange/Runtime/Source/Pipelines/Public/InterchangeGenericAssetsPipelineSharedSettings.h
--- a/Engine/Plugins/Interchange/Runtime/Source/Pipelines/Public/InterchangeGenericAssetsPipelineSharedSettings.h
+++ b/Engine/Plugins/Interchange/Runtime/Source/Pipelines/Public/InterchangeGenericAssetsPipelineSharedSettings.h
@@ -174,7 +174,7 @@ public:
 	bool bImportMeshesInBoneHierarchy = true;
 
 	/** If enabled, frame 0 is used as the reference pose for skeletal meshes. */
-	UPROPERTY(EditAnywhere, BlueprintReadWrite, Category = "Common Skeletal Meshes and Animations")
+	UPROPERTY(EditAnywhere, BlueprintReadWrite, Category = "Common Skeletal Meshes and Animations", meta=(DisplayName="Use T0 As Ref Pose"))
 	bool bUseT0AsRefPose = false;
 
 	/** If enabled, all static meshes that have morph targets will be imported as skeletal meshes instead. */
