MaterialX:
- Add support of <conditionalswitch> node

InterchangeGenericMaterialPipeline:
- Add support of MaterialExpressionSwitch

InterchangeMaterialFactory:
- Add a special case for MaterialExpressionSwitch since the inputs are stored in an Array

diff --git a/Engine/Plugins/Interchange/Runtime/Source/Common/Public/MaterialX/InterchangeMaterialXDefinitions.h b/Engine/Plugins/Interchange/Runtime/Source/Common/Public/MaterialX/InterchangeMaterialXDefinitions.h
--- a/Engine/Plugins/Interchange/Runtime/Source/Common/Public/MaterialX/InterchangeMaterialXDefinitions.h
+++ b/Engine/Plugins/Interchange/Runtime/Source/Common/Public/MaterialX/InterchangeMaterialXDefinitions.h
@@ -532,6 +532,7 @@ MATERIALX_NAMESPACE_BEGIN
 		static constexpr const char* IfGreater = "ifgreater";
 		static constexpr const char* IfGreaterEq = "ifgreatereq";
 		static constexpr const char* IfEqual = "ifequal";
+		static constexpr const char* Switch = "switch";
 		// Channel nodes
 		static constexpr const char* Convert = "convert";
 		static constexpr const char* Combine2 = "combine2";
diff --git a/Engine/Plugins/Interchange/Runtime/Source/Import/Private/Material/InterchangeMaterialFactory.cpp b/Engine/Plugins/Interchange/Runtime/Source/Import/Private/Material/InterchangeMaterialFactory.cpp
--- a/Engine/Plugins/Interchange/Runtime/Source/Import/Private/Material/InterchangeMaterialFactory.cpp
+++ b/Engine/Plugins/Interchange/Runtime/Source/Import/Private/Material/InterchangeMaterialFactory.cpp
@@ -23,6 +23,7 @@
 #include "Materials/MaterialExpressionTextureSample.h"
 #include "Materials/MaterialExpressionThinTranslucentMaterialOutput.h"
 #include "Materials/MaterialExpressionVectorParameter.h"
+#include "Materials/MaterialExpressionSwitch.h"
 #include "Materials/MaterialInterface.h"
 #include "Materials/MaterialInstanceConstant.h"
 #include "Materials/MaterialInstanceDynamic.h"
@@ -220,6 +221,25 @@ namespace UE::Interchange::MaterialFactory::Internal
 				return nullptr;
 			}
 
+			// Special Case for MaterialExpressionSwitch, since the inputs are in an array
+			if (MaterialExpression->IsA<UMaterialExpressionSwitch>())
+			{
+				using namespace UE::Interchange::Materials::Standard::Nodes;
+				if (int32 InputCount; ExpressionNode.GetInt32Attribute(Switch::Attributes::InputCount.ToString(), InputCount))
+				{
+					UMaterialExpressionSwitch* MaterialExpressionSwitch = Cast<UMaterialExpressionSwitch>(MaterialExpression);
+					MaterialExpressionSwitch->Inputs.Empty();
+					MaterialExpressionSwitch->Inputs.Reserve(InputCount);
+					for (int32 Index = 0; Index < InputCount; ++Index)
+					{
+						if (FString InputName; ExpressionNode.GetStringAttribute(Switch::Inputs::InputName.ToString() + FString::FromInt(Index), InputName))
+						{
+							MaterialExpressionSwitch->Inputs.Emplace(FSwitchCustomInput{ FName{InputName} });
+						}
+					}
+				}
+			}
+
 			Expressions.Add(ExpressionNode.GetUniqueID()) = MaterialExpression;
 
 			TArray<FString> Inputs;
diff --git a/Engine/Plugins/Interchange/Runtime/Source/Import/Private/MaterialX/MaterialXUtils/MaterialXManager.cpp b/Engine/Plugins/Interchange/Runtime/Source/Import/Private/MaterialX/MaterialXUtils/MaterialXManager.cpp
--- a/Engine/Plugins/Interchange/Runtime/Source/Import/Private/MaterialX/MaterialXUtils/MaterialXManager.cpp
+++ b/Engine/Plugins/Interchange/Runtime/Source/Import/Private/MaterialX/MaterialXUtils/MaterialXManager.cpp
@@ -69,11 +69,11 @@ FMaterialXManager::FMaterialXManager()
 		{{MaterialX::Category::Atan2,          TEXT("in2")},         TEXT("X")},
 		{{MaterialX::Category::HeightToNormal, TEXT("scale")},       UE::Interchange::Materials::Standard::Nodes::NormalFromHeightMap::Inputs::Intensity.ToString()},
 		{{MaterialX::Category::IfGreater,      TEXT("in1")},         TEXT("AGreaterThanB")},
-		{{MaterialX::Category::IfGreater,      TEXT("in2")},         TEXT("ALessThanB")}, //another input is added for the case equal see ConnectIfGreater
-		{{MaterialX::Category::IfGreaterEq,    TEXT("in1")},         TEXT("AGreaterThanB")}, //another is added for the case equal see ConnectIfGreaterEq
+		{{MaterialX::Category::IfGreater,      TEXT("in2")},         TEXT("ALessThanB")},     //another input is added for the case 'equal', see ConnectIfGreater
+		{{MaterialX::Category::IfGreaterEq,    TEXT("in1")},         TEXT("AGreaterThanB")},  //another input is added for the case 'equal', see ConnectIfGreaterEq
 		{{MaterialX::Category::IfGreaterEq,    TEXT("in2")},         TEXT("ALessThanB")},
 		{{MaterialX::Category::IfEqual,        TEXT("in1")},         TEXT("AEqualsB")},
-		{{MaterialX::Category::IfEqual,        TEXT("in2")},         TEXT("ALessThanB")},  // another input is added for the case greater see ConnectIfEqual
+		{{MaterialX::Category::IfEqual,        TEXT("in2")},         TEXT("ALessThanB")},     // another input is added for the case 'greater', see ConnectIfEqual
 		{{MaterialX::Category::Inside,         TEXT("in")},          TEXT("A")},			  // Inside is treated as a Multiply node
 		{{MaterialX::Category::Inside,         TEXT("mask")},        TEXT("B")},			  // Inside is treated as a Multiply node
 		{{MaterialX::Category::Invert,         TEXT("amount")},      TEXT("A")},
@@ -81,13 +81,13 @@ FMaterialXManager::FMaterialXManager()
 		{{MaterialX::Category::Mix,            TEXT("fg")},          TEXT("B")},
 		{{MaterialX::Category::Mix,            TEXT("bg")},          TEXT("A")},
 		{{MaterialX::Category::Mix,            TEXT("mix")},		 TEXT("Factor")},
-		{{MaterialX::Category::Noise3D,        TEXT("amplitude")},   TEXT("B")},            // The amplitude of the noise is connected to a multiply node
-		{{MaterialX::Category::Noise3D,        TEXT("pivot")},       TEXT("B")},            // The pivot of the noise is connected to a add node
+		{{MaterialX::Category::Noise3D,        TEXT("amplitude")},   TEXT("B")},              // The amplitude of the noise is connected to a multiply node
+		{{MaterialX::Category::Noise3D,        TEXT("pivot")},       TEXT("B")},              // The pivot of the noise is connected to a add node
 		{{MaterialX::Category::Normalize,      TEXT("in")},          TEXT("VectorInput")},
 		{{MaterialX::Category::NormalMap,      TEXT("in")},          TEXT("Normal")},
 		{{MaterialX::Category::NormalMap,      TEXT("scale")},       TEXT("Flatness")},
-		{{MaterialX::Category::Outside,        TEXT("in")},          TEXT("A")},				// Outside is treated as Multiply node
-		{{MaterialX::Category::Outside,        TEXT("mask")},        TEXT("B")},				// Outside is treated as Multiply node
+		{{MaterialX::Category::Outside,        TEXT("in")},          TEXT("A")},			  // Outside is treated as Multiply node
+		{{MaterialX::Category::Outside,        TEXT("mask")},        TEXT("B")},			  // Outside is treated as Multiply node
 		{{MaterialX::Category::Power,          TEXT("in1")},         TEXT("Base")},
 		{{MaterialX::Category::Power,	       TEXT("in2")},         TEXT("Exponent")},
 		{{MaterialX::Category::Rotate2D,       TEXT("in")},          TEXT("Coordinate")},
@@ -97,6 +97,11 @@ FMaterialXManager::FMaterialXManager()
 		{{MaterialX::Category::Rotate3D,       TEXT("in")},          TEXT("Position")},
 		{{MaterialX::Category::Saturate,       TEXT("amount")},      TEXT("Fraction")},
 		{{MaterialX::Category::Smoothstep,     TEXT("in")},          TEXT("Value")},
+		{{MaterialX::Category::Switch,         TEXT("in1")},         TEXT("in1")},
+		{{MaterialX::Category::Switch,         TEXT("in2")},         TEXT("in2")},
+		{{MaterialX::Category::Switch,         TEXT("in3")},         TEXT("in3")},
+		{{MaterialX::Category::Switch,         TEXT("in4")},         TEXT("in4")},
+		{{MaterialX::Category::Switch,         TEXT("which")},       TEXT("SwitchValue")},
 	}
 	, MatchingMaterialExpressions {
 		// Math nodes
@@ -184,6 +189,7 @@ FMaterialXManager::FMaterialXManager()
 		TEXT("Coordinate"),
 		TEXT("Coordinates"),
 		TEXT("D"),
+		TEXT("Default"),
 		TEXT("Exponent"),
 		TEXT("Factor"),
 		TEXT("Fraction"),
@@ -200,6 +206,7 @@ FMaterialXManager::FMaterialXManager()
 		TEXT("Position"),
 		TEXT("RotationAngle"),
 		TEXT("Scale"),
+		TEXT("SwitchValue"),
 		TEXT("TargetLow"),
 		TEXT("TargetHigh"),
 		TEXT("Temp"),
diff --git a/Engine/Plugins/Interchange/Runtime/Source/Import/Private/MaterialX/MaterialXUtils/MaterialXSurfaceShaderAbstract.cpp b/Engine/Plugins/Interchange/Runtime/Source/Import/Private/MaterialX/MaterialXUtils/MaterialXSurfaceShaderAbstract.cpp
--- a/Engine/Plugins/Interchange/Runtime/Source/Import/Private/MaterialX/MaterialXUtils/MaterialXSurfaceShaderAbstract.cpp
+++ b/Engine/Plugins/Interchange/Runtime/Source/Import/Private/MaterialX/MaterialXUtils/MaterialXSurfaceShaderAbstract.cpp
@@ -1199,6 +1199,42 @@ void FMaterialXSurfaceShaderAbstract::ConnectSeparateInputToOutput(const FConnec
 	}
 }
 
+void FMaterialXSurfaceShaderAbstract::ConnectSwitchInputToOutput(const FConnectNode& Connect)
+{
+	using namespace UE::Interchange::Materials::Standard::Nodes;
+	UInterchangeShaderNode* SwitchNode = CreateShaderNode(Connect.UpstreamNode, Connect.UpstreamNode->getName().c_str(), Switch::Name.ToString());
+
+	int32 Index = 0;
+
+	for (mx::InputPtr Input : Connect.UpstreamNode->getInputs())
+	{
+		if (Input->getName() != "which")
+		{
+			SwitchNode->AddStringAttribute(Switch::Inputs::InputName.ToString() + FString::FromInt(Index), Input->getName().c_str());
+			AddAttributeFromValueOrInterface(Input, Input->getName().c_str(), SwitchNode);
+			++Index;
+		}
+		else
+		{
+			AddAttributeFromValueOrInterface(Input, Switch::Inputs::Value.ToString(), SwitchNode);
+		}
+	}
+
+	//In that case we also need to add an attribute to Default	
+	if (mx::InputPtr Input = Connect.UpstreamNode->getInput("in1"))
+	{
+		AddAttributeFromValueOrInterface(Input, Switch::Inputs::Default.ToString(), SwitchNode);
+
+		//Let's add a new input that is a copy of in1 to connect it to the Default input of the MaterialExpressionSwitch
+		mx::InputPtr InputDefault = Connect.UpstreamNode->addInput("default");
+		InputDefault->copyContentFrom(Input);
+		SetAttributeNewName(InputDefault, TCHAR_TO_ANSI(*Switch::Inputs::Default.ToString()));
+	}
+
+	SwitchNode->AddInt32Attribute(Switch::Attributes::InputCount.ToString(), Index);
+	UInterchangeShaderPortsAPI::ConnectDefaultOuputToInput(Connect.ParentShaderNode, Connect.InputChannelName, SwitchNode->GetUniqueID());
+}
+
 void FMaterialXSurfaceShaderAbstract::ConnectSwizzleInputToOutput(const FConnectNode& Connect)
 {
 	using namespace UE::Interchange::Materials::Standard::Nodes;
@@ -1496,6 +1532,7 @@ void FMaterialXSurfaceShaderAbstract::RegisterConnectNodeOutputToInputDelegates(
 	MatchingConnectNodeDelegates.Add(mx::Category::Separate4,		FOnConnectNodeOutputToInput::CreateSP(this, &FMaterialXSurfaceShaderAbstract::ConnectSeparateInputToOutput));
 	MatchingConnectNodeDelegates.Add(mx::Category::TexCoord,		FOnConnectNodeOutputToInput::CreateSP(this, &FMaterialXSurfaceShaderAbstract::ConnectTexCoordInputToOutput));
 	MatchingConnectNodeDelegates.Add(mx::Category::Swizzle,			FOnConnectNodeOutputToInput::CreateSP(this, &FMaterialXSurfaceShaderAbstract::ConnectSwizzleInputToOutput));
+	MatchingConnectNodeDelegates.Add(mx::Category::Switch,			FOnConnectNodeOutputToInput::CreateSP(this, &FMaterialXSurfaceShaderAbstract::ConnectSwitchInputToOutput));
 }
 
 void FMaterialXSurfaceShaderAbstract::SetMatchingInputsNames(MaterialX::NodePtr Node) const
diff --git a/Engine/Plugins/Interchange/Runtime/Source/Import/Public/MaterialX/MaterialXUtils/MaterialXSurfaceShaderAbstract.h b/Engine/Plugins/Interchange/Runtime/Source/Import/Public/MaterialX/MaterialXUtils/MaterialXSurfaceShaderAbstract.h
--- a/Engine/Plugins/Interchange/Runtime/Source/Import/Public/MaterialX/MaterialXUtils/MaterialXSurfaceShaderAbstract.h
+++ b/Engine/Plugins/Interchange/Runtime/Source/Import/Public/MaterialX/MaterialXUtils/MaterialXSurfaceShaderAbstract.h
@@ -280,6 +280,9 @@ protected:
 	/** <separate2/3/4> */
 	void ConnectSeparateInputToOutput(const FConnectNode& Connect);
 
+	/** <switch> */
+	void ConnectSwitchInputToOutput(const FConnectNode& Connect);
+
 	/** <swizzle> */
 	void ConnectSwizzleInputToOutput(const FConnectNode& Connect);
 
diff --git a/Engine/Plugins/Interchange/Runtime/Source/Nodes/Public/InterchangeMaterialDefinitions.h b/Engine/Plugins/Interchange/Runtime/Source/Nodes/Public/InterchangeMaterialDefinitions.h
--- a/Engine/Plugins/Interchange/Runtime/Source/Nodes/Public/InterchangeMaterialDefinitions.h
+++ b/Engine/Plugins/Interchange/Runtime/Source/Nodes/Public/InterchangeMaterialDefinitions.h
@@ -155,6 +155,23 @@ namespace Interchange
 					}
 				}
 
+				namespace Switch
+				{
+					const FName Name = TEXT("Switch");
+
+					namespace Inputs
+					{
+						const FName Value = TEXT("SwitchValue");
+						const FName Default = TEXT("SwitchDefault");
+						const FName InputName = TEXT("SwitchInputName"); // The caller should concatenate this input with an index to get the InputName in the array of the MaterialExpressionSwitch
+					}
+
+					namespace Attributes
+					{
+						const FName InputCount = TEXT("SwitchInputCount"); // Used to know the number of inputs in the Switch
+					}
+				}
+
 				namespace Swizzle
 				{
 					const FName Name = TEXT("MaterialXSwizzle");
diff --git a/Engine/Plugins/Interchange/Runtime/Source/Pipelines/Private/InterchangeGenericMaterialPipeline.cpp b/Engine/Plugins/Interchange/Runtime/Source/Pipelines/Private/InterchangeGenericMaterialPipeline.cpp
--- a/Engine/Plugins/Interchange/Runtime/Source/Pipelines/Private/InterchangeGenericMaterialPipeline.cpp
+++ b/Engine/Plugins/Interchange/Runtime/Source/Pipelines/Private/InterchangeGenericMaterialPipeline.cpp
@@ -36,6 +36,7 @@
 #include "Materials/MaterialExpressionStaticBoolParameter.h"
 #include "Materials/MaterialExpressionConstant.h"
 #include "Materials/MaterialExpressionConstant3Vector.h"
+#include "Materials/MaterialExpressionSwitch.h"
 #include "Materials/MaterialExpressionTextureCoordinate.h"
 #include "Materials/MaterialExpressionTextureObject.h"
 #include "Materials/MaterialExpressionTextureSample.h"
@@ -2688,6 +2689,55 @@ void UInterchangeGenericMaterialPipeline::HandleSwizzleNode(const UInterchangeSh
 	}
 }
 
+void UInterchangeGenericMaterialPipeline::HandleSwitchNode(const UInterchangeShaderNode* ShaderNode, UInterchangeBaseMaterialFactoryNode* MaterialFactoryNode, UInterchangeMaterialExpressionFactoryNode* SwitchFactoryNode)
+{
+	using namespace UE::Interchange::Materials::Standard::Nodes;
+	SwitchFactoryNode->SetCustomExpressionClassName(UMaterialExpressionSwitch::StaticClass()->GetName());
+
+	// SwitchValue
+	{
+		TTuple<UInterchangeMaterialExpressionFactoryNode*, FString> SwitchValueExpression =
+			CreateMaterialExpressionForInput(MaterialFactoryNode, ShaderNode, Switch::Inputs::Value.ToString(), SwitchFactoryNode->GetUniqueID());
+		if (SwitchValueExpression.Get<0>())
+		{
+			UInterchangeShaderPortsAPI::ConnectOuputToInputByName(SwitchFactoryNode, GET_MEMBER_NAME_CHECKED(UMaterialExpressionSwitch, SwitchValue).ToString(),
+																  SwitchValueExpression.Get<0>()->GetUniqueID(), SwitchValueExpression.Get<1>());
+		}
+	}
+
+	// Default
+	{
+		TTuple<UInterchangeMaterialExpressionFactoryNode*, FString> DefaultExpression =
+			CreateMaterialExpressionForInput(MaterialFactoryNode, ShaderNode, Switch::Inputs::Default.ToString(), SwitchFactoryNode->GetUniqueID());
+		if (DefaultExpression.Get<0>())
+		{
+			UInterchangeShaderPortsAPI::ConnectOuputToInputByName(SwitchFactoryNode, GET_MEMBER_NAME_CHECKED(UMaterialExpressionSwitch, Default).ToString(),
+																  DefaultExpression.Get<0>()->GetUniqueID(), DefaultExpression.Get<1>());
+		}
+	}
+
+	// Inputs
+	if (int32 InputCount; ShaderNode->GetInt32Attribute(Switch::Attributes::InputCount.ToString(), InputCount))
+	{
+		SwitchFactoryNode->AddInt32Attribute(Switch::Attributes::InputCount.ToString(), InputCount);
+
+		for (int32 Index = 0; Index < InputCount; ++Index)
+		{
+			if (FString InputName; ShaderNode->GetStringAttribute(Switch::Inputs::InputName.ToString() + FString::FromInt(Index), InputName))
+			{
+				SwitchFactoryNode->AddStringAttribute(Switch::Inputs::InputName.ToString() + FString::FromInt(Index), InputName);
+
+				TTuple<UInterchangeMaterialExpressionFactoryNode*, FString> InputExpression =
+					CreateMaterialExpressionForInput(MaterialFactoryNode, ShaderNode, InputName, SwitchFactoryNode->GetUniqueID());
+				if (InputExpression.Get<0>())
+				{
+					UInterchangeShaderPortsAPI::ConnectOuputToInputByName(SwitchFactoryNode, InputName, InputExpression.Get<0>()->GetUniqueID(), InputExpression.Get<1>());
+				}
+			}
+		}
+	}
+}
+
 void UInterchangeGenericMaterialPipeline::HandleScalarParameterNode(const UInterchangeShaderNode* ShaderNode, UInterchangeMaterialExpressionFactoryNode* ScalarParameterFactoryNode)
 {
 	using namespace UE::Interchange::Materials::Standard::Nodes;
@@ -2800,6 +2850,10 @@ UInterchangeMaterialExpressionFactoryNode* UInterchangeGenericMaterialPipeline::
 	{
 		HandleRotatorNode(ShaderNode, MaterialFactoryNode, MaterialExpression);
 	}
+	else if (*ShaderType == Nodes::Switch::Name)
+	{
+		HandleSwitchNode(ShaderNode, MaterialFactoryNode, MaterialExpression);
+	}
 	else if(*ShaderType == Nodes::Swizzle::Name)
 	{
 		HandleSwizzleNode(ShaderNode, MaterialFactoryNode, MaterialExpression);
diff --git a/Engine/Plugins/Interchange/Runtime/Source/Pipelines/Public/InterchangeGenericMaterialPipeline.h b/Engine/Plugins/Interchange/Runtime/Source/Pipelines/Public/InterchangeGenericMaterialPipeline.h
--- a/Engine/Plugins/Interchange/Runtime/Source/Pipelines/Public/InterchangeGenericMaterialPipeline.h
+++ b/Engine/Plugins/Interchange/Runtime/Source/Pipelines/Public/InterchangeGenericMaterialPipeline.h
@@ -218,6 +218,7 @@ private:
 	void HandleNoiseNode(const UInterchangeShaderNode* ShaderNode, UInterchangeBaseMaterialFactoryNode* MaterialFactoryNode, UInterchangeMaterialExpressionFactoryNode* NoiseFactoryNode);
 	void HandleVectorNoiseNode(const UInterchangeShaderNode* ShaderNode, UInterchangeBaseMaterialFactoryNode* MaterialFactoryNode, UInterchangeMaterialExpressionFactoryNode* NoiseFactoryNode);
 	void HandleSwizzleNode(const UInterchangeShaderNode* ShaderNode, UInterchangeBaseMaterialFactoryNode* MaterialFactoryNode, UInterchangeMaterialExpressionFactoryNode* SwizzleFactoryNode);
+	void HandleSwitchNode(const UInterchangeShaderNode* ShaderNode, UInterchangeBaseMaterialFactoryNode* MaterialFactoryNode, UInterchangeMaterialExpressionFactoryNode* SwitchFactoryNode);
 	void HandleScalarParameterNode(const UInterchangeShaderNode* ShaderNode, UInterchangeMaterialExpressionFactoryNode* ScalarParameterFactoryNode);
 	void HandleVectorParameterNode(const UInterchangeShaderNode* ShaderNode, UInterchangeMaterialExpressionFactoryNode* VectorParameterFactoryNode);
 	void HandleStaticBooleanParameterNode(const UInterchangeShaderNode* ShaderNode, UInterchangeMaterialExpressionFactoryNode* StaticBooleanParameterFactoryNode);
