Interchange MaterialX:
- Fix missing inputs for SurfaceShader <surface>
- Allow multioutputs type to be validated during import
- Fixing a crashing during the call of the CreateNode function, the parent is not necessarily a NodeGraph

Moving it from 5.5 to 5.6: https://p4-swarm.epicgames.net/reviews/36969549

diff --git a/Engine/Plugins/Interchange/Runtime/Source/Import/Private/MaterialX/InterchangeMaterialXTranslator.cpp b/Engine/Plugins/Interchange/Runtime/Source/Import/Private/MaterialX/InterchangeMaterialXTranslator.cpp
--- a/Engine/Plugins/Interchange/Runtime/Source/Import/Private/MaterialX/InterchangeMaterialXTranslator.cpp
+++ b/Engine/Plugins/Interchange/Runtime/Source/Import/Private/MaterialX/InterchangeMaterialXTranslator.cpp
@@ -80,8 +80,9 @@ namespace
 												  FText::FromString(Node->getName().c_str()));
 					return false;
 				}
-				
-				if(!Node->getTypeDef())
+
+				//verify first if it's not a multioutput node (weird that there's no typedef for it)
+				if (!Node->getTypeDef() && Node->getType() != "multioutput")
 				{
 					UInterchangeResultError_Generic* Message = Translator->AddMessage<UInterchangeResultError_Generic>();
 					Message->Text = FText::Format(LOCTEXT("TypeDefNotFound", "<{0}> has no matching TypeDef, aborting import..."),
diff --git a/Engine/Plugins/Interchange/Runtime/Source/Import/Private/MaterialX/MaterialXUtils/MaterialXSurfaceShaderAbstract.cpp b/Engine/Plugins/Interchange/Runtime/Source/Import/Private/MaterialX/MaterialXUtils/MaterialXSurfaceShaderAbstract.cpp
--- a/Engine/Plugins/Interchange/Runtime/Source/Import/Private/MaterialX/MaterialXUtils/MaterialXSurfaceShaderAbstract.cpp
+++ b/Engine/Plugins/Interchange/Runtime/Source/Import/Private/MaterialX/MaterialXUtils/MaterialXSurfaceShaderAbstract.cpp
@@ -500,7 +500,7 @@ void FMaterialXSurfaceShaderAbstract::ConnectRotate3DInputToOutput(const FConnec
 		if(!AddAttributeFromValueOrInterface(Input, TEXT("RotationAngle"), Rotate3DNode))
 		{
 			//copy "in" input into "in1" and create "amount" input as "in2" with the value 360 because UE's angle is a value between [0,1]
-			mx::NodePtr NewDivideNode = CreateNode(Connect.UpstreamNode->getParent()->asA<mx::NodeGraph>(),
+			mx::NodePtr NewDivideNode = CreateNode(Connect.UpstreamNode->getParent(),
 												   Connect.UpstreamNode->getName().c_str(),
 												   mx::Category::Divide,
 												   { {"in1", Input} }, // rename input to match <divide> input "in1"
@@ -667,7 +667,7 @@ void FMaterialXSurfaceShaderAbstract::ConnectConvertInputToOutput(const FConnect
 		}
 
 		//copy "in" input into "in1" and create "amount" input as "in2" with the value 360 because UE's angle is a value between [0,1]
-		mx::NodePtr CombineNode = CreateNode(Connect.UpstreamNode->getParent()->asA<mx::NodeGraph>(),
+		mx::NodePtr CombineNode = CreateNode(Connect.UpstreamNode->getParent(),
 											 Connect.UpstreamNode->getName().c_str(),
 											 Category.c_str(),
 											 InputsToCopy, // rename input to match <divide> input "in1"
@@ -1673,6 +1673,10 @@ UInterchangeShaderNode* FMaterialXSurfaceShaderAbstract::Translate(EInterchangeM
 		{
 			ConnectNodeOutputToInput(Input->getName().c_str(), FunctionCallShaderNode, Input->getName().c_str(), mx::fromValueString<int32>(DefaultValue->getValueString()));
 		}
+		else // this is most likely a BSDF
+		{
+			ConnectNodeOutputToInput(Input->getName().c_str(), FunctionCallShaderNode, Input->getName().c_str(), nullptr);
+		}
 	}
 
 	return FunctionCallShaderNode;
