[Interchange] Fix import into level static mesh LOD which are not importing.

#rnx

diff --git a/Engine/Plugins/Interchange/Runtime/Source/Import/Private/Scene/InterchangeStaticMeshActorFactory.cpp b/Engine/Plugins/Interchange/Runtime/Source/Import/Private/Scene/InterchangeStaticMeshActorFactory.cpp
--- a/Engine/Plugins/Interchange/Runtime/Source/Import/Private/Scene/InterchangeStaticMeshActorFactory.cpp
+++ b/Engine/Plugins/Interchange/Runtime/Source/Import/Private/Scene/InterchangeStaticMeshActorFactory.cpp
@@ -2,6 +2,7 @@
 #include "Scene/InterchangeStaticMeshActorFactory.h"
 
 #include "InterchangeMeshActorFactoryNode.h"
+#include "InterchangeStaticMeshFactoryNode.h"
 #include "Nodes/InterchangeBaseNodeContainer.h"
 #include "Scene/InterchangeActorHelper.h"
 
@@ -129,6 +130,33 @@ UObject* UInterchangeStaticMeshActorFactory::ProcessActor(AActor& SpawnedActor,
 	return nullptr;
 }
 
+void UInterchangeStaticMeshActorFactory::SetupObject_GameThread(const FSetupObjectParams& Arguments)
+{
+	if (AStaticMeshActor* StaticMeshActor = Cast<AStaticMeshActor>(Arguments.ImportedObject))
+	{
+		if (UStaticMeshComponent* StaticMeshComponent = StaticMeshActor->GetStaticMeshComponent())
+		{
+			TArray<FString> TargetNodeUids;
+			Arguments.FactoryNode->GetTargetNodeUids(TargetNodeUids);
+			const UInterchangeStaticMeshFactoryNode* StaticMeshFactoryNode = TargetNodeUids.IsEmpty() ? nullptr : Cast<UInterchangeStaticMeshFactoryNode>(Arguments.NodeContainer->GetFactoryNode(TargetNodeUids[0]));
+			if (StaticMeshFactoryNode)
+			{
+				FSoftObjectPath ReferenceObject;
+				StaticMeshFactoryNode->GetCustomReferenceObject(ReferenceObject);
+				if (UStaticMesh* StaticMesh = Cast<UStaticMesh>(ReferenceObject.TryLoad()))
+				{
+					StaticMeshComponent->SetStaticMesh(StaticMesh);
+
+					if (const UInterchangeMeshActorFactoryNode* MeshActorFactoryNode = Cast<UInterchangeMeshActorFactoryNode>(Arguments.FactoryNode))
+					{
+						UE::Interchange::ActorHelper::ApplySlotMaterialDependencies(*Arguments.NodeContainer, *MeshActorFactoryNode, *StaticMeshComponent);
+					}
+				}
+			}
+		}
+	}
+}
+
 void UInterchangeStaticMeshActorFactory::ExecuteResetObjectProperties(const UInterchangeBaseNodeContainer* BaseNodeContainer, UInterchangeFactoryBaseNode* FactoryNode, UObject* ImportedObject)
 {
 	Super::ExecuteResetObjectProperties(BaseNodeContainer, FactoryNode, ImportedObject);
diff --git a/Engine/Plugins/Interchange/Runtime/Source/Import/Public/Scene/InterchangeStaticMeshActorFactory.h b/Engine/Plugins/Interchange/Runtime/Source/Import/Public/Scene/InterchangeStaticMeshActorFactory.h
--- a/Engine/Plugins/Interchange/Runtime/Source/Import/Public/Scene/InterchangeStaticMeshActorFactory.h
+++ b/Engine/Plugins/Interchange/Runtime/Source/Import/Public/Scene/InterchangeStaticMeshActorFactory.h
@@ -24,6 +24,8 @@ public:
 
 	virtual UClass* GetFactoryClass() const override;
 
+	virtual void SetupObject_GameThread(const FSetupObjectParams& Arguments) override;
+
 	// Interchange factory base interface end
 	//////////////////////////////////////////////////////////////////////////
 
diff --git a/Engine/Plugins/Interchange/Runtime/Source/Pipelines/Private/InterchangeGenericScenesPipeline.cpp b/Engine/Plugins/Interchange/Runtime/Source/Pipelines/Private/InterchangeGenericScenesPipeline.cpp
--- a/Engine/Plugins/Interchange/Runtime/Source/Pipelines/Private/InterchangeGenericScenesPipeline.cpp
+++ b/Engine/Plugins/Interchange/Runtime/Source/Pipelines/Private/InterchangeGenericScenesPipeline.cpp
@@ -28,6 +28,8 @@
 #include "InterchangeSkeletalMeshFactoryNode.h"
 #include "InterchangeSkeletalMeshLodDataNode.h"
 #include "InterchangeSkeletonFactoryNode.h"
+#include "InterchangeStaticMeshFactoryNode.h"
+#include "InterchangeStaticMeshLodDataNode.h"
 #include "InterchangeVariantSetNode.h"
 #include "Nodes/InterchangeUserDefinedAttribute.h"
 
@@ -354,14 +356,58 @@ void UInterchangeGenericLevelPipeline::ExecutePipeline(UInterchangeBaseNodeConta
 	/* Find all scene node that are active joint. Non active joint should be convert to actor if they are in a static mesh hierarchy */
 	CacheActiveJointUids();
 
+	//Cache any lod group data, this way we can add actor only for the lod group
+	TMap< const UInterchangeSceneNode*, TArray<const UInterchangeSceneNode*>> SceneNodesPerLodGroupNode;
+	for (const UInterchangeSceneNode* SceneNode : SceneNodes)
+	{
+		if (!SceneNode)
+		{
+			continue;
+		}
+		TArray<FString> SpecializeTypes;
+		SceneNode->GetSpecializedTypes(SpecializeTypes);
+
+		if (SpecializeTypes.Num() > 0)
+		{
+			if (SpecializeTypes.Contains(UE::Interchange::FSceneNodeStaticData::GetLodGroupSpecializeTypeString()))
+			{
+				TArray<const UInterchangeSceneNode*>& LodGroupChildren = SceneNodesPerLodGroupNode.FindOrAdd(SceneNode);
+				BaseNodeContainer->IterateNodeChildren(SceneNode->GetUniqueID(), [&LodGroupChildren, &SceneNode](const UInterchangeBaseNode* ChildNode)
+					{
+						if (const UInterchangeSceneNode* ChildSceneNode = Cast<UInterchangeSceneNode>(ChildNode))
+						{
+							//Avoid adding self (first iterative call is self)
+							if (SceneNode != ChildSceneNode)
+							{
+								LodGroupChildren.Add(ChildSceneNode);
+							}
+						}
+					});
+			}
+		}
+	}
+
+	auto GetParentLodGroup = [&SceneNodesPerLodGroupNode](const UInterchangeSceneNode* Node)->const UInterchangeSceneNode*
+		{
+			for (TPair< const UInterchangeSceneNode*, TArray<const UInterchangeSceneNode*>> LodSceneNodes : SceneNodesPerLodGroupNode)
+			{
+				if (LodSceneNodes.Value.Contains(Node))
+				{
+					return LodSceneNodes.Key;
+				}
+			}
+			return nullptr;
+		};
+
 	for (const UInterchangeSceneNode* SceneNode : SceneNodes)
 	{
 		if (SceneNode)
 		{
-			if (SceneNode->GetSpecializedTypeCount() > 0)
+			TArray<FString> SpecializeTypes;
+			SceneNode->GetSpecializedTypes(SpecializeTypes);
+
+			if (SpecializeTypes.Num() > 0)
 			{
-				TArray<FString> SpecializeTypes;
-				SceneNode->GetSpecializedTypes(SpecializeTypes);
 				if (!SpecializeTypes.Contains(UE::Interchange::FSceneNodeStaticData::GetTransformSpecializeTypeString()))
 				{
 					bool bSkipNode = true;
@@ -385,6 +431,11 @@ void UInterchangeGenericLevelPipeline::ExecutePipeline(UInterchangeBaseNodeConta
 							}
 						}
 					}
+					else if (SpecializeTypes.Contains(UE::Interchange::FSceneNodeStaticData::GetLodGroupSpecializeTypeString()))
+					{
+						//Do not skip lod group, we always treat it has a mesh (import LODs option control if we import all lod or only the first)
+						bSkipNode = false;
+					}
 
 					if (bSkipNode)
 					{
@@ -393,6 +444,13 @@ void UInterchangeGenericLevelPipeline::ExecutePipeline(UInterchangeBaseNodeConta
 					}
 				}
 			}
+
+			if (GetParentLodGroup(SceneNode))
+			{
+				//Ignore all lod hierarchy after the lod group
+				continue;
+			}
+
 			ExecuteSceneNodePreImport(GlobalOffsetTransform, SceneNode);
 		}
 	}
@@ -492,6 +550,60 @@ void UInterchangeGenericLevelPipeline::ExecuteSceneNodePreImport(const FTransfor
 		}
 	}
 
+	const bool bIsLodGroup = [&SceneNode]()
+		{
+			TArray<FString> SpecializeTypes;
+			SceneNode->GetSpecializedTypes(SpecializeTypes);
+			return SpecializeTypes.Contains(UE::Interchange::FSceneNodeStaticData::GetLodGroupSpecializeTypeString());
+		}();
+
+	UInterchangeStaticMeshFactoryNode* LodGroupStaticMeshFactoryNode = nullptr;
+	if (bIsLodGroup)
+	{
+		//Find the staticmesh factory node create for this lod group
+		BaseNodeContainer->BreakableIterateNodesOfType<UInterchangeStaticMeshFactoryNode>([this, &LodGroupStaticMeshFactoryNode, &SceneNode, &TranslatedAssetNode](const FString& NodeUid, UInterchangeStaticMeshFactoryNode* StaticMeshFactoryNode)
+			{
+				TArray<FString> LodDataUids;
+				StaticMeshFactoryNode->GetLodDataUniqueIds(LodDataUids);
+				if (!LodDataUids.IsEmpty())
+				{
+					if (UInterchangeStaticMeshLodDataNode* LodDataNode = Cast<UInterchangeStaticMeshLodDataNode>(BaseNodeContainer->GetFactoryNode(LodDataUids[0])))
+					{
+						TArray<FString> LodDataMeshUids;
+						LodDataNode->GetMeshUids(LodDataMeshUids);
+						if (!LodDataMeshUids.IsEmpty())
+						{
+							if (BaseNodeContainer->GetIsAncestor(LodDataMeshUids[0], SceneNode->GetUniqueID()))
+							{
+								LodGroupStaticMeshFactoryNode = StaticMeshFactoryNode;
+								//Set the first lod meshuid has the TranslatedAssetNode so it create the mesh actor factory node,
+								//and setup it properly
+								if (const UInterchangeSceneNode* LodDataMeshNode = Cast<UInterchangeSceneNode>(BaseNodeContainer->GetNode(LodDataMeshUids[0])))
+								{
+									FString AssetUid;
+									if (LodDataMeshNode->GetCustomAssetInstanceUid(AssetUid))
+									{
+										TranslatedAssetNode = BaseNodeContainer->GetNode(AssetUid);
+									}
+								}
+								else
+								{
+									TranslatedAssetNode = BaseNodeContainer->GetNode(LodDataMeshUids[0]);
+								}
+								return true;
+							}
+						}
+					}
+				}
+				return false;
+			});
+		if (!LodGroupStaticMeshFactoryNode || !TranslatedAssetNode)
+		{
+			//Skip this lod group if there is no associated mesh
+			return;
+		}
+	}
+
 	UInterchangeActorFactoryNode* ActorFactoryNode = CreateActorFactoryNode(SceneNode, TranslatedAssetNode);
 
 	if (!ensure(ActorFactoryNode))
@@ -549,6 +661,10 @@ void UInterchangeGenericLevelPipeline::ExecuteSceneNodePreImport(const FTransfor
 	{
 		ActorFactoryNode->AddTargetNodeUid(SkeletalMeshFactoryNodeUid);
 	}
+	else if (LodGroupStaticMeshFactoryNode)
+	{
+		ActorFactoryNode->AddTargetNodeUid(LodGroupStaticMeshFactoryNode->GetUniqueID());
+	}
 	else
 	{
 		ActorFactoryNode->AddTargetNodeUid(SceneNode->GetUniqueID());
@@ -1228,9 +1344,24 @@ void UInterchangeGenericLevelPipeline::ExecutePostBroadcastPipeline(const UInter
 #if WITH_EDITORONLY_DATA
 	if (ALevelInstance* LevelInstanceActor = Cast<ALevelInstance>(CreatedAsset))
 	{
-		//Commit an edit to properly finish the setup for the new Level instance actor and new level.
-		LevelInstanceActor->EnterEdit();
-		LevelInstanceActor->ExitEdit();
+		if (UWorld* WorldAsset = LevelInstanceActor->GetWorldAsset().Get())
+		{
+			//We cannot call EnterEdit on a dirty world.
+			if (WorldAsset->GetPackage()->IsDirty())
+			{
+				if (UInterchangeEditorUtilitiesBase* EditorUtilities = UInterchangeManager::GetInterchangeManager().GetEditorUtilities())
+				{
+					if (!EditorUtilities->SaveAsset(WorldAsset))
+					{
+						UE_LOG(LogInterchangePipeline, Warning, TEXT("UInterchangeGenericAssetsPipeline: Cannot save the level instance actor (%s) referenced world (%s)"), *LevelInstanceActor->GetName(), *WorldAsset->GetName());
+					}
+				}
+			}
+
+			//Commit an edit to properly finish the setup for the new Level instance actor and new level.
+			LevelInstanceActor->EnterEdit();
+			LevelInstanceActor->ExitEdit();
+		}
 	}
 #endif
 }
diff --git a/Engine/Plugins/Interchange/Runtime/Source/Pipelines/Private/InterchangeGenericStaticMeshPipeline.cpp b/Engine/Plugins/Interchange/Runtime/Source/Pipelines/Private/InterchangeGenericStaticMeshPipeline.cpp
--- a/Engine/Plugins/Interchange/Runtime/Source/Pipelines/Private/InterchangeGenericStaticMeshPipeline.cpp
+++ b/Engine/Plugins/Interchange/Runtime/Source/Pipelines/Private/InterchangeGenericStaticMeshPipeline.cpp
@@ -244,101 +244,77 @@ void UInterchangeGenericMeshPipeline::ExecutePreImportPipelineStaticMesh()
 		{
 			// Do not combine static meshes
 
-			bool bFoundMeshes = false;
-			if(CommonMeshesProperties->bBakeMeshes || CommonMeshesProperties->bBakePivotMeshes)
-			{
-				TArray<FString> MeshUids;
-				PipelineMeshesUtilities->GetAllStaticMeshInstance(MeshUids);
-
-				// Work out which meshes are collision meshes which correspond to another mesh
-				TMap<FString, TArray<FString>> MeshToCollisionMeshMap;
-				if (bImportCollisionAccordingToMeshName)
-				{
-					BuildMeshToCollisionMeshMap(PipelineMeshesUtilities, *BaseNodeContainer, MeshUids, MeshToCollisionMeshMap);
-				}
-
-				// Now iterate through each mesh UID, creating a new factory for each one
-				for (const FString& MeshUid : MeshUids)
+			auto CreateMeshFactoryNodeUnCombined = [this](const TArray<FString>& MeshUids, const bool bInstancedMesh)->bool
 				{
-					if (bImportCollisionAccordingToMeshName && IsCollisionMeshUid(PipelineMeshesUtilities, *BaseNodeContainer, MeshUid, MeshUids))
+					constexpr int32 LodIndexZero = 0;
+					bool bFoundMeshes = false;
+					// Work out which meshes are collision meshes which correspond to another mesh
+					TMap<FString, TArray<FString>> MeshToCollisionMeshMap;
+					if (bImportCollisionAccordingToMeshName)
 					{
-						// If this is a collision mesh, don't add a factory; it will be added as part of another factory
-						continue;
+						BuildMeshToCollisionMeshMap(PipelineMeshesUtilities, *BaseNodeContainer, MeshUids, MeshToCollisionMeshMap);
 					}
-
-					TMap<int32, TArray<FString>> MeshUidsPerLodIndex;
-
-					const FInterchangeMeshInstance& MeshInstance = PipelineMeshesUtilities->GetMeshInstanceByUid(MeshUid);
-					for (const auto& LodIndexAndSceneNodeContainer : MeshInstance.SceneNodePerLodIndex)
+					// Now iterate through each mesh UID, creating a new factory for each one
+					for (const FString& MeshUid : MeshUids)
 					{
-						const int32 LodIndex = LodIndexAndSceneNodeContainer.Key;
-						const FInterchangeLodSceneNodeContainer& SceneNodeContainer = LodIndexAndSceneNodeContainer.Value;
-
-						TArray<FString>& TranslatedNodes = MeshUidsPerLodIndex.FindOrAdd(LodIndex);
-						for (const UInterchangeSceneNode* SceneNode : SceneNodeContainer.SceneNodes)
+						if (bImportCollisionAccordingToMeshName && IsCollisionMeshUid(PipelineMeshesUtilities, *BaseNodeContainer, MeshUid, MeshUids))
 						{
-							TranslatedNodes.Add(SceneNode->GetUniqueID());
+							// If this is a collision mesh, don't add a factory; it will be added as part of another factory
+							continue;
 						}
-					}
-
-					if (MeshUidsPerLodIndex.Num() > 0)
-					{
-						if (bCollision && bImportCollisionAccordingToMeshName)
+						
+						TArray<FString> ReferencingMeshInstanceUids; //This data is use only when bInstancedMesh is false
+						TMap<int32, TArray<FString>> MeshUidsPerLodIndex;
+						if (bInstancedMesh)
 						{
-							if (const TArray<FString>* CorrespondingCollisionMeshes = MeshToCollisionMeshMap.Find(MeshUid))
+							//Instanced geometry can have lods
+							const FInterchangeMeshInstance& MeshInstance = PipelineMeshesUtilities->GetMeshInstanceByUid(MeshUid);
+							for (const auto& LodIndexAndSceneNodeContainer : MeshInstance.SceneNodePerLodIndex)
 							{
-								MeshUidsPerLodIndex.FindOrAdd(0).Append(*CorrespondingCollisionMeshes);
+								const int32 LodIndex = LodIndexAndSceneNodeContainer.Key;
+								const FInterchangeLodSceneNodeContainer& SceneNodeContainer = LodIndexAndSceneNodeContainer.Value;
+
+								TArray<FString>& TranslatedNodes = MeshUidsPerLodIndex.FindOrAdd(LodIndex);
+								for (const UInterchangeSceneNode* SceneNode : SceneNodeContainer.SceneNodes)
+								{
+									TranslatedNodes.Add(SceneNode->GetUniqueID());
+								}
 							}
 						}
+						else
+						{
+							//Non instanced geometry cannot have lods
+							const FInterchangeMeshGeometry& MeshGeometry = PipelineMeshesUtilities->GetMeshGeometryByUid(MeshUid);
+							TArray<FString>& TranslatedNodes = MeshUidsPerLodIndex.FindOrAdd(LodIndexZero);
+							TranslatedNodes.Add(MeshGeometry.MeshUid);
+							ReferencingMeshInstanceUids = MeshGeometry.ReferencingMeshInstanceUids;
+						}
 
-						UInterchangeStaticMeshFactoryNode* StaticMeshFactoryNode = CreateStaticMeshFactoryNode(MeshUidsPerLodIndex);
-						StaticMeshFactoryNodes.Add(StaticMeshFactoryNode);
-						bFoundMeshes = true;
-					}
-				}
-			}
-
-			if (!bFoundMeshes)
-			{
-				TArray<FString> MeshUids;
-				PipelineMeshesUtilities->GetAllStaticMeshGeometry(MeshUids);
-
-				// Work out which meshes are collision meshes which correspond to another mesh
-				TMap<FString, TArray<FString>> MeshToCollisionMeshMap;
-				if (bImportCollisionAccordingToMeshName)
-				{
-					BuildMeshToCollisionMeshMap(PipelineMeshesUtilities, *BaseNodeContainer, MeshUids, MeshToCollisionMeshMap);
-				}
-
-				for (const FString& MeshUid : MeshUids)
-				{
-					if (bImportCollisionAccordingToMeshName && IsCollisionMeshUid(PipelineMeshesUtilities, *BaseNodeContainer, MeshUid, MeshUids))
-					{
-						// If this is a collision mesh, don't add a factory; it will be added as part of another factory
-						continue;
-					}
-
-					TMap<int32, TArray<FString>> MeshUidsPerLodIndex;
-
-					const FInterchangeMeshGeometry& MeshGeometry = PipelineMeshesUtilities->GetMeshGeometryByUid(MeshUid);
-					const int32 LodIndex = 0;
-					TArray<FString>& TranslatedNodes = MeshUidsPerLodIndex.FindOrAdd(LodIndex);
-					TranslatedNodes.Add(MeshGeometry.MeshUid);
-
-					if (MeshUidsPerLodIndex.Num() > 0)
-					{
-						if (bCollision && bImportCollisionAccordingToMeshName)
+						if (MeshUidsPerLodIndex.Num() > 0)
 						{
-							if (const TArray<FString>* CorrespondingCollisionMeshes = MeshToCollisionMeshMap.Find(MeshUid))
+							if (bCollision && bImportCollisionAccordingToMeshName)
 							{
-								MeshUidsPerLodIndex.FindOrAdd(0).Append(*CorrespondingCollisionMeshes);
+								if (const TArray<FString>* CorrespondingCollisionMeshes = MeshToCollisionMeshMap.Find(MeshUid))
+								{
+									MeshUidsPerLodIndex.FindOrAdd(LodIndexZero).Append(*CorrespondingCollisionMeshes);
+								}
 							}
-						}
 
-						UInterchangeStaticMeshFactoryNode* StaticMeshFactoryNode = CreateStaticMeshFactoryNode(MeshUidsPerLodIndex, MeshGeometry.ReferencingMeshInstanceUids);
-						StaticMeshFactoryNodes.Add(StaticMeshFactoryNode);
+							UInterchangeStaticMeshFactoryNode* StaticMeshFactoryNode = CreateStaticMeshFactoryNode(MeshUidsPerLodIndex, ReferencingMeshInstanceUids);
+							StaticMeshFactoryNodes.Add(StaticMeshFactoryNode);
+							bFoundMeshes = true;
+						}
 					}
-				}
+					return bFoundMeshes;
+				};
+
+			TArray<FString> MeshUids;
+			PipelineMeshesUtilities->GetAllStaticMeshInstance(MeshUids);
+			if (!CreateMeshFactoryNodeUnCombined(MeshUids, true/*bInstancedMesh*/))
+			{
+				MeshUids.Reset();
+				PipelineMeshesUtilities->GetAllStaticMeshGeometry(MeshUids);
+				CreateMeshFactoryNodeUnCombined(MeshUids, false/*bInstancedMesh*/);
 			}
 		}
 	}
