[Interchange] Fix the skeletalmesh import when Geometric Transform is present and we use Time zero to import the skinning bind pose.

#rnx

diff --git a/Engine/Plugins/Interchange/Runtime/Source/Import/Private/Mesh/InterchangeSkeletalMeshFactory.cpp b/Engine/Plugins/Interchange/Runtime/Source/Import/Private/Mesh/InterchangeSkeletalMeshFactory.cpp
--- a/Engine/Plugins/Interchange/Runtime/Source/Import/Private/Mesh/InterchangeSkeletalMeshFactory.cpp
+++ b/Engine/Plugins/Interchange/Runtime/Source/Import/Private/Mesh/InterchangeSkeletalMeshFactory.cpp
@@ -408,6 +408,13 @@ namespace UE
 					FTransform TransformMatrixTransform;
 					TransformMatrixTransform.SetFromMatrix(TransformMatrix);
 
+					FTransform GeometricTransform;
+					if (SceneNode->GetCustomGeometricTransform(GeometricTransform))
+					{
+						//We must add the geometric node transform
+						TransformMatrixTransform = GeometricTransform * TransformMatrixTransform;
+					}
+
 					FMatrix VertexTransformMatrix = (SceneNodeTransform.Inverse() * ((TransformMatrixTransform * BindPose.Inverse()) * (T0 * SceneNodeTransform.Inverse()))).ToMatrixWithScale();
 
 					//Iterate all bone vertices
diff --git a/Engine/Plugins/Interchange/Runtime/Source/Parsers/Fbx/Private/FbxScene.cpp b/Engine/Plugins/Interchange/Runtime/Source/Parsers/Fbx/Private/FbxScene.cpp
--- a/Engine/Plugins/Interchange/Runtime/Source/Parsers/Fbx/Private/FbxScene.cpp
+++ b/Engine/Plugins/Interchange/Runtime/Source/Parsers/Fbx/Private/FbxScene.cpp
@@ -30,6 +30,22 @@ namespace UE
 	{
 		namespace Private
 		{
+			namespace RecursiveHelper
+			{
+				void RecursiveFillChildrenFbxNode(FbxNode* Parent, TArray<FbxNode*>& NodeArray)
+				{
+					if (!Parent)
+					{
+						return;
+					}
+					NodeArray.Add(Parent);
+					int32 ChildCount = Parent->GetChildCount();
+					for (int32 ChildIndex = 0; ChildIndex < ChildCount; ++ChildIndex)
+					{
+						RecursiveFillChildrenFbxNode(Parent->GetChild(ChildIndex), NodeArray);
+					}
+				}
+			} // ns RecursiveHelper
 			void FFbxScene::CreateMeshNodeReference(UInterchangeSceneNode* UnrealSceneNode, FbxNodeAttribute* NodeAttribute, UInterchangeBaseNodeContainer& NodeContainer, const FTransform& GeometricTransform, const FTransform& PivotNodeTransform)
 			{
 				const UInterchangeMeshNode* MeshNode = nullptr;
@@ -91,7 +107,7 @@ namespace UE
 				CreateAssetNodeReference(Parser, UnrealSceneNode, NodeAttribute, NodeContainer, UInterchangeLightNode::StaticAssetTypeName());
 			}
 
-			bool IsNodeUnderCommonJointRootNode(FbxNode* Node, TSet<FbxNode*>& CommonJointRootNodes)
+			bool IsNodeUnderCommonJointRootNode(FbxNode* Node, TMap<FbxNode*, FFbxScene::FRootJointInfo>& CommonJointRootNodes)
 			{
 				if (!Node || CommonJointRootNodes.IsEmpty())
 				{
@@ -170,17 +186,20 @@ namespace UE
 
 				auto ApplySkeletonAttribute = [this, &SDKScene, &UnrealNode, &Node, &NodeContainer, &bResetCache, &GetConvertedTransform, &bBadBindPoseMessageDisplay]()
 				{
+					if (FRootJointInfo* RootJointInfo = CommonJointRootNodes.Find(Node))
+					{
+						if (!RootJointInfo->bValidBindPose)
+						{
+							UnrealNode->SetCustomHasBindPose(false);
+						}
+					}
 					//Add the joint specialized type
 					UnrealNode->AddSpecializedType(FSceneNodeStaticData::GetJointSpecializeTypeString());
 					//Get the bind pose transform for this joint
 					FbxAMatrix GlobalBindPoseJointMatrix = SDKScene->GetAnimationEvaluator()->GetNodeGlobalTransform(Node, 0);
 					TMap<FString, FMatrix> MeshIdToGlobalBindPoseReferenceMap;
 
-					if (!FFbxMesh::GetGlobalJointBindPoseTransform(&Parser, SDKScene, Node, GlobalBindPoseJointMatrix, MeshIdToGlobalBindPoseReferenceMap, bBadBindPoseMessageDisplay))
-					{
-						//Bind Pose not present for the Joint, we should use T0 for the entire skeleton:
-						UnrealNode->SetCustomHasBindPose(false);
-					}
+					FFbxMesh::GetGlobalJointBindPoseTransform(&Parser, SDKScene, Node, GlobalBindPoseJointMatrix, MeshIdToGlobalBindPoseReferenceMap, bBadBindPoseMessageDisplay);
 
 					FTransform GlobalBindPoseJointTransform = GetConvertedTransform(GlobalBindPoseJointMatrix);
 					UnrealNode->SetGlobalBindPoseReferenceForMeshUIDs(MeshIdToGlobalBindPoseReferenceMap);
@@ -472,6 +491,99 @@ namespace UE
 				NodeContainer.AddNode(TransformNode);
 				return TransformNode;
 			}
+			
+			bool FFbxScene::IsValidBindPose(FbxScene* SDKScene, FbxNode* RootJoint) const
+			{
+				if (CommonJointRootNodes.IsEmpty())
+				{
+					return false;
+				}
+
+				int32 PoseCount = SDKScene->GetPoseCount();
+				if (PoseCount == 0)
+				{
+					SDKScene->GetFbxManager()->CreateMissingBindPoses(SDKScene);
+					PoseCount = SDKScene->GetPoseCount();
+				}
+				
+				TArray<FbxNode*> NodeArray;
+				RecursiveHelper::RecursiveFillChildrenFbxNode(RootJoint, NodeArray);
+
+				for (int32 PoseIndex = 0; PoseIndex < PoseCount; PoseIndex++)
+				{
+					FbxPose* CurrentPose = SDKScene->GetPose(PoseIndex);
+
+					// current pose is bind pose, 
+					if (CurrentPose && CurrentPose->IsBindPose())
+					{
+						// IsValidBindPose doesn't work reliably
+						// It checks all the parent chain(regardless root given), and if the parent doesn't have correct bind pose, it fails
+						// It causes more false positive issues than the real issue we have to worry about
+						// If you'd like to try this, set CHECK_VALID_BIND_POSE to 1, and try the error message
+						// when Autodesk fixes this bug, then we might be able to re-open this
+						FString PoseName = CurrentPose->GetName();
+						// all error report status
+						FbxStatus Status;
+
+						// it does not make any difference of checking with different node
+						for(FbxNode* Current : NodeArray)
+						{
+							FString CurrentName = Current->GetName();
+							FbxArray<FbxNode*> pMissingAncestors, pMissingDeformers, pMissingDeformersAncestors, pWrongMatrices;
+
+							if (CurrentPose->IsValidBindPoseVerbose(Current, pMissingAncestors, pMissingDeformers, pMissingDeformersAncestors, pWrongMatrices, 0.0001, &Status))
+							{
+								return true;
+							}
+							else
+							{
+								// first try to fix up
+								// add missing ancestors
+								for (int i = 0; i < pMissingAncestors.GetCount(); i++)
+								{
+									FbxAMatrix mat = pMissingAncestors.GetAt(i)->EvaluateGlobalTransform(FBXSDK_TIME_ZERO);
+									CurrentPose->Add(pMissingAncestors.GetAt(i), mat);
+								}
+
+								pMissingAncestors.Clear();
+								pMissingDeformers.Clear();
+								pMissingDeformersAncestors.Clear();
+								pWrongMatrices.Clear();
+
+								// check it again
+								if (CurrentPose->IsValidBindPose(Current))
+								{
+									return true;
+								}
+								else
+								{
+									// first try to find parent who is null group and see if you can try test it again
+									FbxNode* ParentNode = Current->GetParent();
+									while (ParentNode)
+									{
+										FbxNodeAttribute* Attr = ParentNode->GetNodeAttribute();
+										if (Attr && Attr->GetAttributeType() == FbxNodeAttribute::eNull)
+										{
+											// found it 
+											break;
+										}
+
+										// find next parent
+										ParentNode = ParentNode->GetParent();
+									}
+
+									if (ParentNode && CurrentPose->IsValidBindPose(ParentNode))
+									{
+										return true;
+									}
+								}
+							}
+						}
+					}
+				}
+				return false;
+			}
+
 
 			void FFbxScene::AddHierarchy(FbxScene* SDKScene, UInterchangeBaseNodeContainer& NodeContainer, TMap<FString, TSharedPtr<FPayloadContextBase, ESPMode::ThreadSafe>>& PayloadContexts)
 			{
@@ -485,6 +597,11 @@ namespace UE
 				//Cache the common root joint
 				FindCommonJointRootNode(SDKScene, ForceJointNodes);
 
+				for (TPair<FbxNode*, FRootJointInfo>& RootJointInfo : CommonJointRootNodes)
+				{
+					RootJointInfo.Value.bValidBindPose = IsValidBindPose(SDKScene, RootJointInfo.Key);
+				}
+
 				bool bBadBindPoseMessageDisplay = false;
 				AddHierarchyRecursively(nullptr, RootNode, SDKScene, NodeContainer, PayloadContexts, ForceJointNodes, bBadBindPoseMessageDisplay);
 
@@ -662,7 +779,7 @@ namespace UE
 						{
 							if (FbxNode* Root = Internal_GetRootSkeleton(SDKScene, Node))
 							{
-								CommonJointRootNodes.Add(Root);
+								CommonJointRootNodes.FindOrAdd(Root);
 							}
 						}
 					}
diff --git a/Engine/Plugins/Interchange/Runtime/Source/Parsers/Fbx/Private/FbxScene.h b/Engine/Plugins/Interchange/Runtime/Source/Parsers/Fbx/Private/FbxScene.h
--- a/Engine/Plugins/Interchange/Runtime/Source/Parsers/Fbx/Private/FbxScene.h
+++ b/Engine/Plugins/Interchange/Runtime/Source/Parsers/Fbx/Private/FbxScene.h
@@ -35,6 +35,11 @@ namespace UE
 				void AddMorphTargetAnimations(FbxScene* SDKScene, UInterchangeBaseNodeContainer& NodeContainer, TMap<FString, TSharedPtr<FPayloadContextBase, ESPMode::ThreadSafe>>& PayloadContexts, const TArray<FMorphTargetAnimationBuildingData>& MorphTargetAnimationsBuildingData);
 				UInterchangeSceneNode* CreateTransformNode(UInterchangeBaseNodeContainer& NodeContainer, const FString& NodeName, const FString& NodeUniqueID);
 
+				struct FRootJointInfo
+				{
+					bool bValidBindPose = false;
+				};
+
 			protected:
 				void CreateMeshNodeReference(UInterchangeSceneNode* UnrealSceneNode, FbxNodeAttribute* NodeAttribute, UInterchangeBaseNodeContainer& NodeContainer, const FTransform& GeometricTransform, const FTransform& PivotNodeTransform);
 				void CreateCameraNodeReference(UInterchangeSceneNode* UnrealSceneNode, FbxNodeAttribute* NodeAttribute, UInterchangeBaseNodeContainer& NodeContainer);
@@ -67,8 +72,9 @@ namespace UE
 
 				void FindForceJointNode(FbxScene* SDKScene, TArray<FbxNode*>& ForceJointNodes);
 
+				bool IsValidBindPose(FbxScene* SDKScene, FbxNode* RootJoint) const;
 
-				TSet<FbxNode*> CommonJointRootNodes;
+				TMap<FbxNode*, FRootJointInfo> CommonJointRootNodes;
 
 				FFbxParser& Parser;
 			};
