[Interchange] Fix ensure with morph target payload key. Fix performance issue when importing a lot of morph target with a lot of joints in fbx.

#rnx

diff --git a/Engine/Plugins/Interchange/Runtime/Source/Parsers/Fbx/Private/FbxAnimation.cpp b/Engine/Plugins/Interchange/Runtime/Source/Parsers/Fbx/Private/FbxAnimation.cpp
--- a/Engine/Plugins/Interchange/Runtime/Source/Parsers/Fbx/Private/FbxAnimation.cpp
+++ b/Engine/Plugins/Interchange/Runtime/Source/Parsers/Fbx/Private/FbxAnimation.cpp
@@ -883,7 +883,12 @@ namespace UE::Interchange::Private
 		, TMap<FString, TSharedPtr<FPayloadContextBase>>& PayloadContexts
 		, const FMorphTargetAnimationBuildingData& MorphTargetAnimationBuildingData)
 	{
-		FString PayLoadKey = MorphTargetAnimationBuildingData.MorphTargetNodeUid + TEXT("\\") + FString::FromInt(MorphTargetAnimationBuildingData.AnimationIndex) + TEXT("\\") + FString::FromInt(MorphTargetAnimationBuildingData.MorphTargetIndex) + TEXT("\\") + FString::FromInt(MorphTargetAnimationBuildingData.ChannelIndex) + TEXT("_CurveAnimationPayloadKey");
+		FString PayLoadKey = MorphTargetAnimationBuildingData.MorphTargetNodeUid
+			+ TEXT("\\") + (MorphTargetAnimationBuildingData.InterchangeMeshNode ? MorphTargetAnimationBuildingData.InterchangeMeshNode->GetUniqueID() : FString()) //Same shape can be animated on different mesh node
+			+ TEXT("\\") + FString::FromInt(MorphTargetAnimationBuildingData.AnimationIndex)
+			+ TEXT("\\") + FString::FromInt(MorphTargetAnimationBuildingData.MorphTargetIndex)
+			+ TEXT("\\") + FString::FromInt(MorphTargetAnimationBuildingData.ChannelIndex)
+			+ TEXT("_CurveAnimationPayloadKey");
 
 		if (ensure(!PayloadContexts.Contains(PayLoadKey)))
 		{
diff --git a/Engine/Plugins/Interchange/Runtime/Source/Parsers/Fbx/Private/FbxScene.cpp b/Engine/Plugins/Interchange/Runtime/Source/Parsers/Fbx/Private/FbxScene.cpp
--- a/Engine/Plugins/Interchange/Runtime/Source/Parsers/Fbx/Private/FbxScene.cpp
+++ b/Engine/Plugins/Interchange/Runtime/Source/Parsers/Fbx/Private/FbxScene.cpp
@@ -1086,6 +1086,7 @@ namespace UE
 				//Group the Morph Target animations based on SkeletonNodeUid and AnimationIndex
 				TMap<const UInterchangeSceneNode*, TMap<int32, TArray<FMorphTargetAnimationBuildingData>>> MorphTargetAnimationsBuildingDataGrouped;
 
+				TMap<FString, FString> EvaluatedJoints;
 				for (const FMorphTargetAnimationBuildingData& MorphTargetAnimationBuildingData : MorphTargetAnimationsBuildingData)
 				{
 					if (MorphTargetAnimationBuildingData.StartTime == MorphTargetAnimationBuildingData.StopTime)
@@ -1095,6 +1096,7 @@ namespace UE
 					}
 					
 					TSet<FString> SkeletonUids;
+					
 					if (MorphTargetAnimationBuildingData.InterchangeMeshNode->IsSkinnedMesh())
 					{
 						//Find the root joint(s) for this MeshGeometry
@@ -1103,32 +1105,40 @@ namespace UE
 						for (const FString& SkeletonDependency : SkeletonDependencies)
 						{
 							FString JointNodeUid = SkeletonDependency;
-							FString ParentNodeUid = SkeletonDependency;
-
-							while (!JointNodeUid.Equals(UInterchangeBaseNode::InvalidNodeUid()))
+							FString& RootJointNodeForJoint = EvaluatedJoints.FindOrAdd(JointNodeUid);
+							if (!RootJointNodeForJoint.IsEmpty())
+							{
+								SkeletonUids.Add(RootJointNodeForJoint);
+							}
+							else
 							{
-								if (const UInterchangeSceneNode* Node = Cast< UInterchangeSceneNode >(NodeContainer.GetNode(ParentNodeUid)))
+								FString ParentNodeUid = SkeletonDependency;
+								while (!JointNodeUid.Equals(UInterchangeBaseNode::InvalidNodeUid()))
 								{
-									if (Node->IsSpecializedTypeContains(FSceneNodeStaticData::GetJointSpecializeTypeString()))
+									if (const UInterchangeSceneNode* Node = Cast< UInterchangeSceneNode >(NodeContainer.GetNode(ParentNodeUid)))
 									{
-										JointNodeUid = ParentNodeUid;
-										ParentNodeUid = Node->GetParentUid();
+										if (Node->IsSpecializedTypeContains(FSceneNodeStaticData::GetJointSpecializeTypeString()))
+										{
+											JointNodeUid = ParentNodeUid;
+											ParentNodeUid = Node->GetParentUid();
+										}
+										else
+										{
+											break;
+										}
 									}
 									else
 									{
 										break;
 									}
 								}
-								else
+
+								if (!JointNodeUid.Equals(UInterchangeBaseNode::InvalidNodeUid()))
 								{
-									break;
+									RootJointNodeForJoint = JointNodeUid;
+									SkeletonUids.Add(JointNodeUid);
 								}
 							}
-
-							if (!JointNodeUid.Equals(UInterchangeBaseNode::InvalidNodeUid()))
-							{
-								SkeletonUids.Add(JointNodeUid);
-							}
 						}
 					}
 					else
