Interchange Level Sequence:
- Based on discussion with TwinMotion team, update the Playback range to be in [TickTime(Frame0), TickTime(LastFrame) + TickTime(Duration(LastFrame))] range instead of [TickTime(Frame0), TickTime(LastFrame)], fixing sequences that are one frame short

EngineTest:
- Update tests

#virtualized

diff --git a/Engine/Plugins/Interchange/Runtime/Source/Import/Private/Animation/InterchangeLevelSequenceFactory.cpp b/Engine/Plugins/Interchange/Runtime/Source/Import/Private/Animation/InterchangeLevelSequenceFactory.cpp
--- a/Engine/Plugins/Interchange/Runtime/Source/Import/Private/Animation/InterchangeLevelSequenceFactory.cpp
+++ b/Engine/Plugins/Interchange/Runtime/Source/Import/Private/Animation/InterchangeLevelSequenceFactory.cpp
@@ -142,10 +142,10 @@ namespace UE::Interchange::Private
 		AActor* GetActor(const UInterchangeAnimationTrackNode& TrackNode);
 		FGuid BindActorToLevelSequence(AActor* Actor);
 		FGuid BindComponentToLevelSequence(AActor* Actor);
-		void UpdateTransformChannels(TArrayView<FMovieSceneDoubleChannel*>& Channels, int32 IndexOffset, const TArray<FRichCurve>& Curves);
+		void UpdateTransformChannels(TArrayView<FMovieSceneDoubleChannel*>& Channels, int32 IndexOffset, const TArray<FRichCurve>& Curves, const FFrameNumber& FrameDuration);
 		
 		template<typename ChannelType, typename ValueType> 
-		void UpdateStepChannel(ChannelType& Channel, const TArray<float>& KeyTimes, const TArray<ValueType>& Values)
+		void UpdateStepChannel(ChannelType& Channel, const TArray<float>& KeyTimes, const TArray<ValueType>& Values, const FFrameNumber & FrameDuration)
 		{
 			const FFrameRate& FrameRate = MovieScene->GetTickResolution();
 
@@ -153,18 +153,23 @@ namespace UE::Interchange::Private
 
 			Data.Reset();
 
+			TRange<FFrameNumber> MinMaxFrameNumber{ TRange<FFrameNumber>::Empty() };
+
 			for (int32 KeyIndex = 0; KeyIndex < KeyTimes.Num(); ++KeyIndex)
 			{
 				FFrameNumber FrameNumber = FrameRate.AsFrameNumber(KeyTimes[KeyIndex]);
 
-				if (FrameNumber < this->MinFrameNumber)
+				if (MinMaxFrameNumber.IsEmpty())
 				{
-					this->MinFrameNumber = FrameNumber;
+					MinMaxFrameNumber = TRange<FFrameNumber>{ FrameNumber };
 				}
-
-				if (FrameNumber > this->MaxFrameNumber)
+				else if (FrameNumber < MinMaxFrameNumber.GetLowerBoundValue())
+				{
+					MinMaxFrameNumber.SetLowerBoundValue(FrameNumber);
+				}
+				else if (FrameNumber > MinMaxFrameNumber.GetUpperBoundValue())
 				{
-					this->MaxFrameNumber = FrameNumber;
+					MinMaxFrameNumber.SetUpperBoundValue(FrameNumber);
 				}
 
 				if constexpr(std::is_same_v<ChannelType, FMovieSceneObjectPathChannel>)
@@ -176,6 +181,12 @@ namespace UE::Interchange::Private
 					Data.AddKey(FrameNumber, Values[KeyIndex]);
 				}
 			}
+
+			if (!MinMaxFrameNumber.IsEmpty())
+			{
+				MinMaxFrameNumber.SetUpperBound(TRangeBound<FFrameNumber>::Exclusive(MinMaxFrameNumber.GetUpperBoundValue() + FrameDuration));
+				PlaybackRange = TRange<FFrameNumber>::Hull(PlaybackRange, MinMaxFrameNumber);
+			}
 		}
 
 	private:
@@ -185,8 +196,7 @@ namespace UE::Interchange::Private
 		const UInterchangeBaseNodeContainer& NodeContainer;
 		const IInterchangeAnimationPayloadInterface& PayloadInterface;
 
-		FFrameNumber MinFrameNumber;
-		FFrameNumber MaxFrameNumber;
+		TRange<FFrameNumber> PlaybackRange = TRange<FFrameNumber>::Empty(); // [StartFrame, EndFrame[ in TickResolution
 		bool ClearSubsequenceTrack = true;
 	};
 
@@ -235,7 +245,7 @@ namespace UE::Interchange::Private
 			}
 		}
 
-		LevelSequence.MovieScene->SetPlaybackRange(TRange<FFrameNumber>::Inclusive(MinFrameNumber, MaxFrameNumber));
+		LevelSequence.MovieScene->SetPlaybackRange(PlaybackRange.IsEmpty() ? TRange<FFrameNumber>(0, 1) : PlaybackRange);
 		LevelSequence.MovieScene->SetEvaluationType(EMovieSceneEvaluationType::FrameLocked);
 	}
 
@@ -319,9 +329,10 @@ namespace UE::Interchange::Private
 
 		TArrayView<FMovieSceneDoubleChannel*> Channels = TransformSection->GetChannelProxy().GetChannels<FMovieSceneDoubleChannel>();
 
-		UpdateTransformChannels(Channels, 0, PayloadData->Curves); // Translation
-		UpdateTransformChannels(Channels, 3, PayloadData->Curves); // Rotation
-		UpdateTransformChannels(Channels, 6, PayloadData->Curves); // Scaling
+		const FFrameNumber FrameDuration = FFrameRate::TransformTime(1, MovieScene->GetDisplayRate(), MovieScene->GetTickResolution()).RoundToFrame();
+		UpdateTransformChannels(Channels, 0, PayloadData->Curves, FrameDuration); // Translation
+		UpdateTransformChannels(Channels, 3, PayloadData->Curves, FrameDuration); // Rotation
+		UpdateTransformChannels(Channels, 6, PayloadData->Curves, FrameDuration); // Scaling
 
 		// Remove unnecessary keys
 		FKeyDataOptimizationParams OptimParams;
@@ -383,10 +394,11 @@ namespace UE::Interchange::Private
 	}
 
 	template <class T>
-	void ProcessRichCurveKeys(const FFrameRate& TargetFrameRate, const FRichCurve& Curve, TArray<FFrameNumber>& OutFrameNumbers, TArray<T>& OutValues, FFrameNumber& MinFrameNumber, FFrameNumber& MaxFrameNumber)
+	void ProcessRichCurveKeys(const FFrameRate& TargetFrameRate, const FRichCurve& Curve, FFrameNumber FrameDuration, TArray<FFrameNumber>& OutFrameNumbers, TArray<T>& OutValues, TRange<FFrameNumber>& PlaybackRange)
 	{
 		const TArray<FRichCurveKey>& CurveKeys = Curve.GetConstRefOfKeys();
 		const int32 NumCurveKeys = CurveKeys.Num();
+		TRange<FFrameNumber> MinMaxFrameNumber{ TRange<FFrameNumber>::Empty() };
 
 		OutFrameNumbers.Reserve(NumCurveKeys);
 		OutValues.Reserve(NumCurveKeys);
@@ -399,14 +411,17 @@ namespace UE::Interchange::Private
 
 			FFrameNumber& FrameNumber = OutFrameNumbers.Add_GetRef(TargetFrameRate.AsFrameNumber(RichCurveKey.Time));
 
-			if (FrameNumber < MinFrameNumber)
+			if (MinMaxFrameNumber.IsEmpty())
 			{
-				MinFrameNumber = FrameNumber;
+				MinMaxFrameNumber = TRange<FFrameNumber>{ FrameNumber };
 			}
-
-			if (FrameNumber > MaxFrameNumber)
+			else if (FrameNumber < MinMaxFrameNumber.GetLowerBoundValue())
+			{
+				MinMaxFrameNumber.SetLowerBoundValue(FrameNumber);
+			}
+			else if (FrameNumber > MinMaxFrameNumber.GetUpperBoundValue())
 			{
-				MaxFrameNumber = FrameNumber;
+				MinMaxFrameNumber.SetUpperBoundValue(FrameNumber);
 			}
 
 			const float SecondsDelta = [&]() -> float
@@ -452,6 +467,12 @@ namespace UE::Interchange::Private
 
 			ConvertRichCurveKeyToFloatValue(RichCurveKey, Value, KeyTimingRatio, TargetFrameRate.AsInterval());
 		}
+
+		if (!MinMaxFrameNumber.IsEmpty())
+		{
+			MinMaxFrameNumber.SetUpperBound(TRangeBound<FFrameNumber>::Exclusive(MinMaxFrameNumber.GetUpperBoundValue() + FrameDuration));
+			PlaybackRange = TRange<FFrameNumber>::Hull(PlaybackRange, MinMaxFrameNumber);
+		}
 	}
 
 	void FLevelSequenceHelper::PopulateSubsequenceTrack(const UInterchangeAnimationTrackSetInstanceNode& InstanceNode)
@@ -540,8 +561,7 @@ namespace UE::Interchange::Private
 			NewSection->EvalOptions.CompletionMode = EMovieSceneCompletionMode::ProjectDefault;
 		}
 
-		MinFrameNumber = FMath::Min(MinFrameNumber, NewSection->GetRange().GetLowerBoundValue());
-		MaxFrameNumber = FMath::Max(MaxFrameNumber, NewSection->GetRange().GetUpperBoundValue());
+		PlaybackRange = TRange<FFrameNumber>::Hull(PlaybackRange, NewSection->GetRange());
 
 		ClearSubsequenceTrack = false;
 	}
@@ -606,14 +626,16 @@ namespace UE::Interchange::Private
 		const FName EnumChannelTypeName = FMovieSceneByteChannel::StaticStruct()->GetFName();
 		const FName ObjectPathChannelTypeName = FMovieSceneObjectPathChannel::StaticStruct()->GetFName();
 
-		auto CopyToChannel = [this](auto Channel, const FRichCurve& Curve)
+		const FFrameNumber FrameDuration = FFrameRate::TransformTime(1, MovieScene->GetDisplayRate(), MovieScene->GetTickResolution()).RoundToFrame();
+
+		auto CopyToChannel = [this, &FrameDuration](auto Channel, const FRichCurve& Curve)
 		{
 			TArray<FFrameNumber> FrameNumbers;
 
 			using FMovieSceneValue = typename std::remove_pointer_t<decltype(Channel)>::ChannelValueType;
 			TArray<FMovieSceneValue> Values;
 
-			ProcessRichCurveKeys(this->MovieScene->GetTickResolution(), Curve, FrameNumbers, Values, this->MinFrameNumber, this->MaxFrameNumber);
+			ProcessRichCurveKeys(MovieScene->GetTickResolution(), Curve, FrameDuration, FrameNumbers, Values, PlaybackRange);
 
 			if(!Values.IsEmpty())
 			{
@@ -656,19 +678,19 @@ namespace UE::Interchange::Private
 				FMovieSceneChannelHandle Channel = ChannelProxy.MakeHandle(ChannelTypeName, Index);
 				if(bIsBoolChannel)
 				{
-					UpdateStepChannel(*(Channel.Cast<FMovieSceneBoolChannel>().Get()), PayloadData->StepCurves[Index].KeyTimes, PayloadData->StepCurves[Index].BooleanKeyValues.GetValue());
+					UpdateStepChannel(*(Channel.Cast<FMovieSceneBoolChannel>().Get()), PayloadData->StepCurves[Index].KeyTimes, PayloadData->StepCurves[Index].BooleanKeyValues.GetValue(), FrameDuration);
 				}
 				else if(bIsEnumChannel)
 				{
-					UpdateStepChannel(*(Channel.Cast<FMovieSceneByteChannel>().Get()), PayloadData->StepCurves[Index].KeyTimes, PayloadData->StepCurves[Index].ByteKeyValues.GetValue());
+					UpdateStepChannel(*(Channel.Cast<FMovieSceneByteChannel>().Get()), PayloadData->StepCurves[Index].KeyTimes, PayloadData->StepCurves[Index].ByteKeyValues.GetValue(), FrameDuration);
 				}
 				else if(bIsIntegerChannel)
 				{
-					UpdateStepChannel(*(Channel.Cast<FMovieSceneIntegerChannel>().Get()), PayloadData->StepCurves[Index].KeyTimes, PayloadData->StepCurves[Index].IntegerKeyValues.GetValue());
+					UpdateStepChannel(*(Channel.Cast<FMovieSceneIntegerChannel>().Get()), PayloadData->StepCurves[Index].KeyTimes, PayloadData->StepCurves[Index].IntegerKeyValues.GetValue(), FrameDuration);
 				}
 				else if(bIsObjectPathChannel)
 				{
-					UpdateStepChannel(*(Channel.Cast<FMovieSceneObjectPathChannel>().Get()), PayloadData->StepCurves[Index].KeyTimes, PayloadData->StepCurves[Index].StringKeyValues.GetValue());
+					UpdateStepChannel(*(Channel.Cast<FMovieSceneObjectPathChannel>().Get()), PayloadData->StepCurves[Index].KeyTimes, PayloadData->StepCurves[Index].StringKeyValues.GetValue(), FrameDuration);
 				}
 				else if(bIsFloatChannel)
 				{
@@ -749,14 +771,14 @@ namespace UE::Interchange::Private
 		return ComponentBinding;
 	}
 
-	void FLevelSequenceHelper::UpdateTransformChannels(TArrayView<FMovieSceneDoubleChannel*>& Channels, int32 IndexOffset, const TArray<FRichCurve>& Curves)
+	void FLevelSequenceHelper::UpdateTransformChannels(TArrayView<FMovieSceneDoubleChannel*>& Channels, int32 IndexOffset, const TArray<FRichCurve>& Curves, const FFrameNumber& FrameDuration)
 	{
-		auto CopyToChannel = [this](FMovieSceneDoubleChannel* Channel, const FRichCurve& Curve)
+		auto CopyToChannel = [this, &FrameDuration](FMovieSceneDoubleChannel* Channel, const FRichCurve& Curve)
 		{
 			TArray<FFrameNumber> FrameNumbers;
 			TArray<FMovieSceneDoubleValue> Values;
 
-			ProcessRichCurveKeys(this->MovieScene->GetTickResolution(), Curve, FrameNumbers, Values, this->MinFrameNumber, this->MaxFrameNumber);
+			ProcessRichCurveKeys(MovieScene->GetTickResolution(), Curve, FrameDuration, FrameNumbers, Values, PlaybackRange);
 
 			if (!Values.IsEmpty())
 			{
