[Interchange] SkeletalMeshActor position issue fixes.

#rnx

diff --git a/Engine/Plugins/Interchange/Runtime/Source/Import/Private/Scene/InterchangeSkeletalMeshActorFactory.cpp b/Engine/Plugins/Interchange/Runtime/Source/Import/Private/Scene/InterchangeSkeletalMeshActorFactory.cpp
--- a/Engine/Plugins/Interchange/Runtime/Source/Import/Private/Scene/InterchangeSkeletalMeshActorFactory.cpp
+++ b/Engine/Plugins/Interchange/Runtime/Source/Import/Private/Scene/InterchangeSkeletalMeshActorFactory.cpp
@@ -28,24 +28,21 @@ UObject* UInterchangeSkeletalMeshActorFactory::ProcessActor(AActor& SpawnedActor
 
 	if (USkeletalMeshComponent* SkeletalMeshComponent = SkeletalMeshActor->GetSkeletalMeshComponent())
 	{
-		if (USkeletalMesh* SkeletalMesh = SkeletalMeshComponent->GetSkeletalMeshAsset())
+		USkeletalMesh* SkeletalMesh = SkeletalMeshComponent->GetSkeletalMeshAsset();
+		if (SkeletalMesh && Params.ImportAssets.Contains(SkeletalMesh))
 		{
-			if (Params.ImportAssets.Contains(SkeletalMesh))
-			{
 #if WITH_EDITOR
-				//If we are importing the skeletal mesh, we need to use a async build scope to unregister the component to avoid a deadlock
-				FSkinnedAssetAsyncBuildScope AsyncBuildScope(SkeletalMesh);
+			//If we are importing the skeletal mesh, we need to use a async build scope to unregister the component to avoid a deadlock
+			FSkinnedAssetAsyncBuildScope AsyncBuildScope(SkeletalMesh);
 #endif
-				SkeletalMeshComponent->UnregisterComponent();
-			}
-			else
-			{
-				SkeletalMeshComponent->UnregisterComponent();
-			}
-			
-
-			return SkeletalMeshComponent;
+			SkeletalMeshComponent->UnregisterComponent();
+		}
+		else
+		{
+			SkeletalMeshComponent->UnregisterComponent();
 		}
+
+		return SkeletalMeshComponent;
 	}
 
 	return nullptr;
