[Interchange] Make sure we add non lod geo to all lod when importing with lod

#rnx

diff --git a/Engine/Plugins/Interchange/Runtime/Source/Pipelines/Private/InterchangeGenericSkeletalMeshPipeline.cpp b/Engine/Plugins/Interchange/Runtime/Source/Pipelines/Private/InterchangeGenericSkeletalMeshPipeline.cpp
--- a/Engine/Plugins/Interchange/Runtime/Source/Pipelines/Private/InterchangeGenericSkeletalMeshPipeline.cpp
+++ b/Engine/Plugins/Interchange/Runtime/Source/Pipelines/Private/InterchangeGenericSkeletalMeshPipeline.cpp
@@ -121,6 +121,53 @@ void UInterchangeGenericMeshPipeline::ExecutePreImportPipelineSkeletalMesh()
 		}
 	};
 
+	auto FinishSkeletalMeshSetup = [this, &SetSkeletalMeshDependencies](const bool bUseInstanceMesh
+		, bool& bFoundInstances
+		, const FString& SkeletonRootUid
+		, UInterchangeSkeletonFactoryNode* SkeletonFactoryNode
+		, TMap<int32, TArray<FString>>& MeshUidsPerLodIndex
+		, const TArray<FString>& MeshUids)
+		{
+			//We must add all nodes that are not part of any lod to all lods upper then 0
+			if (bUseInstanceMesh && MeshUidsPerLodIndex.Num() > 1)
+			{
+				constexpr int32 LodIndexZero = 0;
+				for (const FString& MeshUid : MeshUids)
+				{
+					const FInterchangeMeshInstance& MeshInstance = PipelineMeshesUtilities->GetMeshInstanceByUid(MeshUid);
+					if (MeshInstance.LodGroupNode == nullptr
+						&& MeshInstance.SceneNodePerLodIndex.Num() == 1
+						&& MeshInstance.SceneNodePerLodIndex.Contains(LodIndexZero)
+						&& MeshInstance.SceneNodePerLodIndex[LodIndexZero].SceneNodes.Num() == 1)
+					{
+						const UInterchangeSceneNode* SceneNode = MeshInstance.SceneNodePerLodIndex[LodIndexZero].SceneNodes[0];
+						//This mesh must be added to all existing LODs over 0
+						for (TPair<int32, TArray<FString>>& MeshUidsPerLodIndexPair : MeshUidsPerLodIndex)
+						{
+							if (MeshUidsPerLodIndexPair.Key > LodIndexZero)
+							{
+								TArray<FString>& TranslatedNodes = MeshUidsPerLodIndexPair.Value;
+								TranslatedNodes.AddUnique(SceneNode->GetUniqueID());
+							}
+						}
+					}
+				}
+			}
+
+			if (MeshUidsPerLodIndex.Num() > 0)
+			{
+				UInterchangeSkeletalMeshFactoryNode* SkeletalMeshFactoryNode = CreateSkeletalMeshFactoryNode(SkeletonRootUid, MeshUidsPerLodIndex);
+				if (SkeletalMeshFactoryNode != nullptr)
+				{
+					SetSkeletalMeshDependencies(SkeletonRootUid, SkeletalMeshFactoryNode);
+					SkeletonFactoryNodes.Add(SkeletonFactoryNode);
+					SkeletalMeshFactoryNodes.Add(SkeletalMeshFactoryNode);
+					bFoundInstances = true;
+				}
+			}
+		};
+
+
 	PRAGMA_DISABLE_DEPRECATION_WARNINGS
 	const bool bCombineSkeletalMeshes = bCombineSkeletalMeshes_DEPRECATED;
 	PRAGMA_ENABLE_DEPRECATION_WARNINGS
@@ -129,7 +176,7 @@ void UInterchangeGenericMeshPipeline::ExecutePreImportPipelineSkeletalMesh()
 		//////////////////////////////////////////////////////////////////////////
 		//Combined everything we can
 		TMap<FString, TArray<FString>> MeshUidsPerSkeletonRootUid;
-		auto CreatePerSkeletonRootUidCombinedSkinnedMesh = [this, &MeshUidsPerSkeletonRootUid, &SetSkeletalMeshDependencies](const bool bUseInstanceMesh)
+		auto CreatePerSkeletonRootUidCombinedSkinnedMesh = [this, &MeshUidsPerSkeletonRootUid, &SetSkeletalMeshDependencies, &FinishSkeletalMeshSetup](const bool bUseInstanceMesh)
 		{
 #if !WITH_EDITOR || !WITH_EDITORONLY_DATA
 			if (MeshUidsPerSkeletonRootUid.Num() > 0)
@@ -176,17 +223,7 @@ void UInterchangeGenericMeshPipeline::ExecutePreImportPipelineSkeletalMesh()
 						}
 					}
 
-					if (MeshUidsPerLodIndex.Num() > 0)
-					{
-						UInterchangeSkeletalMeshFactoryNode* SkeletalMeshFactoryNode = CreateSkeletalMeshFactoryNode(SkeletonRootUid, MeshUidsPerLodIndex);
-						if (SkeletalMeshFactoryNode != nullptr)
-						{
-							SetSkeletalMeshDependencies(SkeletonRootUid, SkeletalMeshFactoryNode);
-							SkeletonFactoryNodes.Add(SkeletonFactoryNode);
-							SkeletalMeshFactoryNodes.Add(SkeletalMeshFactoryNode);
-							bFoundInstances = true;
-						}
-					}
+					FinishSkeletalMeshSetup(bUseInstanceMesh, bFoundInstances, SkeletonRootUid, SkeletonFactoryNode, MeshUidsPerLodIndex, MeshUids);
 				}
 				return bFoundInstances;
 			}
@@ -204,7 +241,7 @@ void UInterchangeGenericMeshPipeline::ExecutePreImportPipelineSkeletalMesh()
 		//////////////////////////////////////////////////////////////////////////
 		//Do not combined meshes
 		TArray<FString> MeshUids;
-		auto CreatePerSkeletonRootUidSkinnedMesh = [this, &MeshUids, &SetSkeletalMeshDependencies](const bool bUseInstanceMesh)
+		auto CreatePerSkeletonRootUidSkinnedMesh = [this, &MeshUids, &SetSkeletalMeshDependencies, &FinishSkeletalMeshSetup](const bool bUseInstanceMesh)
 		{
 #if !WITH_EDITOR || !WITH_EDITORONLY_DATA
 			if (MeshUids.Num() > 0)
@@ -261,18 +298,8 @@ void UInterchangeGenericMeshPipeline::ExecutePreImportPipelineSkeletalMesh()
 						TArray<FString>& TranslatedNodes = MeshUidsPerLodIndex.FindOrAdd(LodIndex);
 						TranslatedNodes.Add(MeshGeometry.MeshUid);
 					}
-					if (MeshUidsPerLodIndex.Num() > 0)
-					{
-						UInterchangeSkeletalMeshFactoryNode* SkeletalMeshFactoryNode = CreateSkeletalMeshFactoryNode(SkeletonRootUid, MeshUidsPerLodIndex);
-						if (SkeletalMeshFactoryNode != nullptr)
-						{
-							SetSkeletalMeshDependencies(SkeletonRootUid, SkeletalMeshFactoryNode);
-							SkeletonFactoryNodes.Add(SkeletonFactoryNode);
-							SkeletalMeshFactoryNodes.Add(SkeletalMeshFactoryNode);
-							bFoundInstances = true;
-						}
 
-					}
+					FinishSkeletalMeshSetup(bUseInstanceMesh, bFoundInstances, SkeletonRootUid, SkeletonFactoryNode, MeshUidsPerLodIndex, MeshUids);
 				}
 				return bFoundInstances;
 			}
diff --git a/Engine/Plugins/Interchange/Runtime/Source/Pipelines/Private/InterchangeGenericStaticMeshPipeline.cpp b/Engine/Plugins/Interchange/Runtime/Source/Pipelines/Private/InterchangeGenericStaticMeshPipeline.cpp
--- a/Engine/Plugins/Interchange/Runtime/Source/Pipelines/Private/InterchangeGenericStaticMeshPipeline.cpp
+++ b/Engine/Plugins/Interchange/Runtime/Source/Pipelines/Private/InterchangeGenericStaticMeshPipeline.cpp
@@ -206,6 +206,32 @@ void UInterchangeGenericMeshPipeline::ExecutePreImportPipelineStaticMesh()
 					}
 				}
 
+				//We must add all nodes that are not part of any lod to all lods upper then 0
+				if (MeshUidsPerLodIndex.Num() > 1)
+				{
+					constexpr int32 LodIndexZero = 0;
+					for (const FString& MeshUid : MeshUids)
+					{
+						const FInterchangeMeshInstance& MeshInstance = PipelineMeshesUtilities->GetMeshInstanceByUid(MeshUid);
+						if (MeshInstance.LodGroupNode == nullptr
+							&& MeshInstance.SceneNodePerLodIndex.Num() == 1
+							&& MeshInstance.SceneNodePerLodIndex.Contains(LodIndexZero)
+							&& MeshInstance.SceneNodePerLodIndex[LodIndexZero].SceneNodes.Num() == 1)
+						{
+							const UInterchangeSceneNode* SceneNode = MeshInstance.SceneNodePerLodIndex[LodIndexZero].SceneNodes[0];
+							//This mesh must be added to all existing LODs over 0
+							for (TPair<int32, TArray<FString>>& MeshUidsPerLodIndexPair : MeshUidsPerLodIndex)
+							{
+								if (MeshUidsPerLodIndexPair.Key > LodIndexZero)
+								{
+									TArray<FString>& TranslatedNodes = MeshUidsPerLodIndexPair.Value;
+									TranslatedNodes.AddUnique(SceneNode->GetUniqueID());
+								}
+							}
+						}
+					}
+				}
+
 				// If we got some instances, create a static mesh factory node
 				if (MeshUidsPerLodIndex.Num() > 0)
 				{
