[Interchange] fix the build setting so they apply to all imported LODs.

#rnx

diff --git a/Engine/Plugins/Interchange/Runtime/Source/FactoryNodes/Private/InterchangeMeshFactoryNode.cpp b/Engine/Plugins/Interchange/Runtime/Source/FactoryNodes/Private/InterchangeMeshFactoryNode.cpp
--- a/Engine/Plugins/Interchange/Runtime/Source/FactoryNodes/Private/InterchangeMeshFactoryNode.cpp
+++ b/Engine/Plugins/Interchange/Runtime/Source/FactoryNodes/Private/InterchangeMeshFactoryNode.cpp
@@ -11,27 +11,34 @@
 #if WITH_EDITOR
 
 #define IMPLEMENT_MESH_BUILD_VALUE_TO_ASSET(AttributeName, AttributeType, PropertyName)	\
+bool bResult = false;																\
 AttributeType ValueData;															\
 if (GetAttribute<AttributeType>(Macro_Custom##AttributeName##Key.Key, ValueData))	\
 {																					\
 	if (USkeletalMesh* SkeletalMesh = Cast<USkeletalMesh>(Asset))					\
 	{																				\
-		if (FSkeletalMeshLODInfo* LodInfo = SkeletalMesh->GetLODInfo(0))			\
+		for(int32 LodIndex = 0; LodIndex < SkeletalMesh->GetLODNum(); ++LodIndex)	\
 		{																			\
-			LodInfo->BuildSettings.PropertyName = ValueData;						\
-			return true;															\
+			if (FSkeletalMeshLODInfo* LodInfo = SkeletalMesh->GetLODInfo(LodIndex))	\
+			{																		\
+				LodInfo->BuildSettings.PropertyName = ValueData;					\
+				bResult = true;														\
+			}																		\
 		}																			\
 	}																				\
 	else if (UStaticMesh* StaticMesh = Cast<UStaticMesh>(Asset))					\
 	{																				\
-		if (StaticMesh->GetNumSourceModels() > 0)									\
+		for(int32 LodIndex = 0; LodIndex < StaticMesh->GetNumSourceModels(); ++LodIndex) \
 		{																			\
-			StaticMesh->GetSourceModel(0).BuildSettings.PropertyName = ValueData;	\
-			return true;															\
+			if (StaticMesh->IsSourceModelValid(LodIndex))							\
+			{																		\
+				StaticMesh->GetSourceModel(LodIndex).BuildSettings.PropertyName = ValueData;	\
+				bResult = true;														\
+			}																		\
 		}																			\
 	}																				\
 }																					\
-return false;
+return bResult;
 
 #define IMPLEMENT_MESH_BUILD_ASSET_TO_VALUE(AttributeName, PropertyName)	\
 if (UStaticMesh* StaticMesh = Cast<UStaticMesh>(Asset))								\
diff --git a/Engine/Plugins/Interchange/Runtime/Source/FactoryNodes/Private/InterchangeSkeletalMeshFactoryNode.cpp b/Engine/Plugins/Interchange/Runtime/Source/FactoryNodes/Private/InterchangeSkeletalMeshFactoryNode.cpp
--- a/Engine/Plugins/Interchange/Runtime/Source/FactoryNodes/Private/InterchangeSkeletalMeshFactoryNode.cpp
+++ b/Engine/Plugins/Interchange/Runtime/Source/FactoryNodes/Private/InterchangeSkeletalMeshFactoryNode.cpp
@@ -11,19 +11,23 @@
 #if WITH_EDITOR
 
 #define IMPLEMENT_SKELETAL_BUILD_VALUE_TO_ASSET(AttributeName, AttributeType, PropertyName)	\
+bool bResult = false;																\
 AttributeType ValueData;															\
 if (GetAttribute<AttributeType>(Macro_Custom##AttributeName##Key.Key, ValueData))	\
 {																					\
 	if (USkeletalMesh* SkeletalMesh = Cast<USkeletalMesh>(Asset))					\
 	{																				\
-		if (FSkeletalMeshLODInfo* LodInfo = SkeletalMesh->GetLODInfo(0))			\
+		for(int32 LodIndex = 0; LodIndex < SkeletalMesh->GetLODNum(); ++LodIndex)	\
 		{																			\
-			LodInfo->BuildSettings.PropertyName = ValueData;						\
-			return true;															\
+			if (FSkeletalMeshLODInfo* LodInfo = SkeletalMesh->GetLODInfo(LodIndex))	\
+			{																		\
+				LodInfo->BuildSettings.PropertyName = ValueData;					\
+				bResult = true;														\
+			}																		\
 		}																			\
 	}																				\
 }																					\
-return false;
+return bResult;
 
 #define IMPLEMENT_SKELETALMESH_BUILD_ASSET_TO_VALUE(AttributeName, PropertyName)					\
 if (USkeletalMesh* SkeletalMesh = Cast<USkeletalMesh>(Asset))						\
diff --git a/Engine/Plugins/Interchange/Runtime/Source/FactoryNodes/Private/InterchangeStaticMeshFactoryNode.cpp b/Engine/Plugins/Interchange/Runtime/Source/FactoryNodes/Private/InterchangeStaticMeshFactoryNode.cpp
--- a/Engine/Plugins/Interchange/Runtime/Source/FactoryNodes/Private/InterchangeStaticMeshFactoryNode.cpp
+++ b/Engine/Plugins/Interchange/Runtime/Source/FactoryNodes/Private/InterchangeStaticMeshFactoryNode.cpp
@@ -8,21 +8,27 @@
 
 #if WITH_EDITOR
 
+//We set the build value to all LODs
 #define IMPLEMENT_STATICMESH_BUILD_VALUE_TO_ASSET(AttributeName, AttributeType, PropertyName)	\
+bool bResult = false;																\
 AttributeType ValueData;															\
 if (GetAttribute<AttributeType>(Macro_Custom##AttributeName##Key.Key, ValueData))	\
 {																					\
 	if (UStaticMesh* StaticMesh = Cast<UStaticMesh>(Asset))							\
 	{																				\
-		if (StaticMesh->GetNumSourceModels() > 0)									\
+		for(int32 LodIndex = 0; LodIndex < StaticMesh->GetNumSourceModels(); ++LodIndex) \
 		{																			\
-			StaticMesh->GetSourceModel(0).BuildSettings.PropertyName = ValueData;	\
-			return true;															\
+			if (StaticMesh->IsSourceModelValid(LodIndex))							\
+			{																		\
+				StaticMesh->GetSourceModel(LodIndex).BuildSettings.PropertyName = ValueData;	\
+				bResult = true;														\
+			}																		\
 		}																			\
 	}																				\
 }																					\
-return false;
+return bResult;
 
+//We take only the lod zero to get back the value
 #define IMPLEMENT_STATICMESH_BUILD_ASSET_TO_VALUE(AttributeName, PropertyName)	\
 if (UStaticMesh* StaticMesh = Cast<UStaticMesh>(Asset))								\
 {																					\
