[Interchange - glTF] IES Light support.

#rnx

diff --git a/Engine/Plugins/Interchange/Runtime/Source/FactoryNodes/Public/InterchangeTextureLightProfileFactoryNode.h b/Engine/Plugins/Interchange/Runtime/Source/FactoryNodes/Public/InterchangeTextureLightProfileFactoryNode.h
--- a/Engine/Plugins/Interchange/Runtime/Source/FactoryNodes/Public/InterchangeTextureLightProfileFactoryNode.h
+++ b/Engine/Plugins/Interchange/Runtime/Source/FactoryNodes/Public/InterchangeTextureLightProfileFactoryNode.h
@@ -5,6 +5,7 @@
 #include "CoreMinimal.h"
 #include "InterchangeTexture2DFactoryNode.h"
 #include "Nodes/InterchangeBaseNode.h"
+#include "Misc/Base64.h"
 
 #if WITH_ENGINE
 #include "Engine/Texture.h"
@@ -20,6 +21,10 @@ class INTERCHANGEFACTORYNODES_API UInterchangeTextureLightProfileFactoryNode : p
 	GENERATED_BODY()
 
 public:
+	UInterchangeTextureLightProfileFactoryNode()
+	{
+		AttributeHelperIESSourceFileContentChunks.Initialize(Attributes, TEXT("IESSourceFileContentChunks"));
+	}
 
 	UFUNCTION(BlueprintCallable, Category = "Interchange | Node | Texture")
 	bool GetCustomBrightness(float& AttributeValue) const
@@ -54,6 +59,50 @@ public:
 #endif
 	}
 
+	void StoreIESSourceFileContents(const TArray64<uint8>& IESSourceFileContents) 
+	{
+		int64 NumberOfElements = IESSourceFileContents.Num();
+		constexpr int64 SubPartSize = 65536;// 2 ^ 16;
+		int64 ElementIndex = 0;
+		while (ElementIndex < NumberOfElements)
+		{
+			int64 RawSubPartSize = ((ElementIndex + SubPartSize) < NumberOfElements) ? SubPartSize : (NumberOfElements - ElementIndex);
+			TArray<uint8> RawSubPart;
+			RawSubPart.SetNum(RawSubPartSize);
+
+			FPlatformMemory::Memcpy(RawSubPart.GetData(), IESSourceFileContents.GetData() + ElementIndex * sizeof(uint8), RawSubPartSize * sizeof(uint8));
+
+			AttributeHelperIESSourceFileContentChunks.AddItem(FBase64::Encode(RawSubPart));
+
+			ElementIndex += SubPartSize;
+		}
+	}
+
+	void GetIESSourceFileContents(TArray64<uint8>& IESSourceFileContent) const
+	{
+		TArray<FString> IESSourceFileContentChunks;
+		AttributeHelperIESSourceFileContentChunks.GetItems(IESSourceFileContentChunks);
+
+		int64 NumberOfElements = 0;
+		for (const FString& Chunk : IESSourceFileContentChunks)
+		{
+			NumberOfElements += Chunk.Len();
+		}
+
+		IESSourceFileContent.Reserve(NumberOfElements);
+
+		int64 ElementIndex = 0;
+		for (const FString& Chunk : IESSourceFileContentChunks)
+		{
+			TArray<uint8> ChunkDecoded;
+			FBase64::Decode(Chunk, ChunkDecoded);
+
+			IESSourceFileContent.SetNum(ElementIndex + ChunkDecoded.Num());
+
+			FPlatformMemory::Memcpy(IESSourceFileContent.GetData() + ElementIndex * sizeof(uint8), ChunkDecoded.GetData(), ChunkDecoded.Num() * sizeof(uint8));
+		}
+	}
+
 	virtual void CopyWithObject(const UInterchangeFactoryBaseNode* SourceNode, UObject* Object) override
 	{
 		Super::CopyWithObject(SourceNode, Object);
@@ -67,6 +116,8 @@ public:
 #endif
 	}
 
+	
+
 private:
 	virtual UClass* GetObjectClass() const override
 	{
@@ -77,4 +128,6 @@ private:
 	IMPLEMENT_NODE_ATTRIBUTE_KEY(Brightness);
 	IMPLEMENT_NODE_ATTRIBUTE_KEY(TextureMultiplier);
 
+	//Note: Base64 encoded chunks of 2^16 sized strings.
+	UE::Interchange::TArrayAttributeHelper<FString> AttributeHelperIESSourceFileContentChunks; 
 };
\ No newline at end of file
diff --git a/Engine/Plugins/Interchange/Runtime/Source/Import/Private/Gltf/InterchangeGltfTranslator.cpp b/Engine/Plugins/Interchange/Runtime/Source/Import/Private/Gltf/InterchangeGltfTranslator.cpp
--- a/Engine/Plugins/Interchange/Runtime/Source/Import/Private/Gltf/InterchangeGltfTranslator.cpp
+++ b/Engine/Plugins/Interchange/Runtime/Source/Import/Private/Gltf/InterchangeGltfTranslator.cpp
@@ -23,6 +23,7 @@
 #include "InterchangeTexture2DNode.h"
 #include "InterchangeTranslatorHelper.h"
 #include "InterchangeVariantSetNode.h"
+#include "InterchangeTextureLightProfileNode.h"
 #include "Nodes/InterchangeSourceNode.h"
 #include "Nodes/InterchangeUserDefinedAttribute.h"
 
@@ -45,6 +46,9 @@
 #include "EngineAnalytics.h"
 #include "Engine/RendererSettings.h"
 
+#include "IESConverter.h"
+#include "Texture/TextureTranslatorUtilities.h"
+
 #include UE_INLINE_GENERATED_CPP_BY_NAME(InterchangeGltfTranslator)
 
 #define LOCTEXT_NAMESPACE "InterchangeGLTFTranslator"
@@ -53,6 +57,7 @@ static const TArray<FString> ImporterSupportedExtensions = {
 	/* Lights */
 	GLTF::ToString(GLTF::EExtension::KHR_LightsPunctual),
 	GLTF::ToString(GLTF::EExtension::KHR_Lights),
+	GLTF::ToString(GLTF::EExtension::EXT_LightsIES),
 	/* Variants */
 	GLTF::ToString(GLTF::EExtension::KHR_MaterialsVariants),
 	/* Materials */
@@ -351,7 +356,73 @@ void UInterchangeGLTFTranslator::HandleGltfNode( UInterchangeBaseNodeContainer&
 		{
 			Transform.ConcatenateRotation(FRotator(0, -90, 0).Quaternion());
 
-			if ( GltfAsset.Lights.IsValidIndex( GltfNode.LightIndex ) )
+			if (GltfAsset.LightsIES.IsValidIndex(GltfNode.LightIES.Index))
+			{
+				FString LightIESInstanceHash = GltfNode.LightIES.GetHash();
+				const GLTF::FLightIES& LightIES = GltfAsset.LightsIES[GltfNode.LightIES.Index];
+				//As LightIES values can be instance specific, we create LightIES on Demand
+
+				const FString LightIESUid = TEXT("\\LightIES\\") + GltfAsset.LightsIES[GltfNode.LightIES.Index].UniqueId;
+
+				auto SetupLightIESInstanceNode = [&GltfNode, &LightIESUid, &NodeContainer](UInterchangeLightNode* LightNode, const FString& LightIESInstanceUid, const FString& DisplayLabel, const EInterchangeNodeContainerType& ContainerType)
+					{
+						LightNode->InitializeNode(LightIESInstanceUid, DisplayLabel, ContainerType);
+
+						if (GltfNode.LightIES.bHasColor)
+						{
+							LightNode->SetCustomLightColor(FLinearColor(GltfNode.LightIES.Color));
+						}
+
+						if (GltfNode.LightIES.bHasIntensityMultiplier)
+						{
+							LightNode->SetCustomUseIESBrightness(true);
+							LightNode->SetCustomIESBrightnessScale(GltfNode.LightIES.IntensityMultipler);
+						}
+
+						LightNode->SetCustomIESTexture(LightIESUid);
+
+						NodeContainer.AddNode(LightNode);
+					};
+
+				if (!GltfAsset.Lights.IsValidIndex(GltfNode.LightIndex))
+				{//No Base Light
+					const FString LightIESInstanceUid = TEXT("\\LightIESInstance\\LightIES-") + LightIES.UniqueId + TEXT("\\") + LightIESInstanceHash;
+				
+					if (!NodeContainer.GetNode(LightIESInstanceUid))
+					{
+						UInterchangePointLightNode* LightIESInstanceNode = NewObject< UInterchangePointLightNode >(&NodeContainer);
+						SetupLightIESInstanceNode(LightIESInstanceNode, LightIESInstanceUid, GltfAsset.LightsIES[GltfNode.LightIES.Index].Name, EInterchangeNodeContainerType::TranslatedAsset);
+
+						//Default is Candelas with 1 intensity:
+						LightIESInstanceNode->SetCustomIntensityUnits(EInterchangeLightUnits::Candelas);
+						LightIESInstanceNode->SetCustomIntensity(1.f);
+					}
+
+					InterchangeSceneNode->SetCustomAssetInstanceUid(LightIESInstanceUid);
+				}
+				else
+				{//Base Light defined, needs duplication in order to be able to set the IES Instance.
+					const FString LightIESInstanceUid = TEXT("\\LightIESInstance\\Light-") + GltfAsset.Lights[GltfNode.LightIndex].UniqueId + TEXT("_LightIES-") + LightIES.UniqueId + TEXT("\\") + LightIESInstanceHash;
+
+					if (!NodeContainer.GetNode(LightIESInstanceUid))
+					{
+						const FString LightNodeUid = TEXT("\\Light\\") + GltfAsset.Lights[GltfNode.LightIndex].UniqueId;
+
+						const UInterchangeBaseNode* OriginalLightBaseNode = NodeContainer.GetNode(LightNodeUid);
+						if (const UInterchangeLightNode* OriginalLightNode = Cast<UInterchangeLightNode>(OriginalLightBaseNode))
+						{
+							//Duplicate the originalLightNode so we can set the LightProfile on it:
+							if (UInterchangeLightNode* DuplicatedObject = Cast<UInterchangeLightNode>(StaticDuplicateObject(OriginalLightNode, &NodeContainer)))
+							{
+								SetupLightIESInstanceNode(DuplicatedObject, LightIESInstanceUid, DuplicatedObject->GetDisplayLabel(), DuplicatedObject->GetNodeContainerType());
+							}
+						}
+					}
+
+					InterchangeSceneNode->SetCustomAssetInstanceUid(LightIESInstanceUid);
+				}
+			}
+			else if ( GltfAsset.Lights.IsValidIndex( GltfNode.LightIndex ) )
 			{
 				const FString LightNodeUid = TEXT("\\Light\\") + GltfAsset.Lights[ GltfNode.LightIndex ].UniqueId;
 				InterchangeSceneNode->SetCustomAssetInstanceUid( LightNodeUid );
@@ -583,6 +654,21 @@ bool UInterchangeGLTFTranslator::Translate( UInterchangeBaseNodeContainer& NodeC
 		}
 	}
 
+	//IES
+	{
+		for (const GLTF::FLightIES& LightIES : GltfAsset.LightsIES)
+		{
+			UInterchangeTextureLightProfileNode* TextureLightProfileNode = NewObject< UInterchangeTextureLightProfileNode>(&NodeContainer, UInterchangeTextureLightProfileNode::StaticClass());
+
+			FString LightIESUid = TEXT("\\LightIES\\") + LightIES.UniqueId;
+
+			TextureLightProfileNode->InitializeNode(LightIESUid, LightIES.Name, EInterchangeNodeContainerType::TranslatedAsset);
+			TextureLightProfileNode->SetPayLoadKey(FString::FromInt(LightIES.Index));
+
+			NodeContainer.AddNode(TextureLightProfileNode);
+		}
+	}
+
 	// Meshes
 	TSet<FString> MaterialsUsedOnMeshesWithVertexColor;
 	TSet<int> UnusedGltfMeshIndices;
@@ -1501,4 +1587,44 @@ UInterchangeGLTFTranslator::UInterchangeGLTFTranslator()
 	}
 }
 
+TOptional<UE::Interchange::FImportLightProfile> UInterchangeGLTFTranslator::GetLightProfilePayloadData(const FString& PayloadKey, TOptional<FString>& AlternateTexturePath) const
+{
+	using namespace UE::Interchange;
+
+	int32 IndexLightIES = INDEX_NONE;
+	LexFromString(IndexLightIES, *PayloadKey);
+
+	if (!GltfAsset.LightsIES.IsValidIndex(IndexLightIES))
+	{
+		return {};
+	}
+	
+	const bool bUseBuffer = GltfAsset.LightsIES[IndexLightIES].BufferViewIndex != INDEX_NONE;
+	
+	//faking the filepath for buffered version
+	const FString FilePath = bUseBuffer ? TEXT("FakeFor_glTF_IES_Buffer_Import.ies") : FPaths::ConvertRelativePathToFull(GltfAsset.LightsIES[IndexLightIES].FilePath);
+
+	UInterchangeSourceData* PayloadSourceData = UInterchangeManager::GetInterchangeManager().CreateSourceData(FilePath);
+	FGCObjectScopeGuard ScopedSourceData(PayloadSourceData);
+	if (!PayloadSourceData)
+	{
+		return TOptional<UE::Interchange::FImportLightProfile>();
+	}
+
+	PayloadSourceData->ClearInternalFlags(EInternalObjectFlags::Async);
+
+	UInterchangeTranslatorBase* SourceTranslator = UInterchangeManager::GetInterchangeManager().GetTranslatorForSourceData(PayloadSourceData);
+	FGCObjectScopeGuard ScopedSourceTranslator(SourceTranslator);
+	const IInterchangeTextureLightProfilePayloadInterface* TextureTranslator = Cast< IInterchangeTextureLightProfilePayloadInterface >(SourceTranslator);
+	if (!ensure(TextureTranslator))
+	{
+		return TOptional<UE::Interchange::FImportLightProfile>();
+	}
+
+	SourceTranslator->ClearInternalFlags(EInternalObjectFlags::Async);
+	SourceTranslator->SetResultsContainer(Results);
+
+	return bUseBuffer ? TextureTranslator->GetLightProfilePayloadData(GltfAsset.LightsIES[IndexLightIES].Data, GltfAsset.LightsIES[IndexLightIES].DataByteLength) : TextureTranslator->GetLightProfilePayloadData(FilePath, AlternateTexturePath);
+}
+
 #undef LOCTEXT_NAMESPACE
diff --git a/Engine/Plugins/Interchange/Runtime/Source/Import/Private/Texture/InterchangeIESTranslator.cpp b/Engine/Plugins/Interchange/Runtime/Source/Import/Private/Texture/InterchangeIESTranslator.cpp
--- a/Engine/Plugins/Interchange/Runtime/Source/Import/Private/Texture/InterchangeIESTranslator.cpp
+++ b/Engine/Plugins/Interchange/Runtime/Source/Import/Private/Texture/InterchangeIESTranslator.cpp
@@ -50,27 +50,34 @@ TOptional<UE::Interchange::FImportLightProfile> UInterchangeIESTranslator::GetLi
 		return {};
 	}
 
-	const uint8* Buffer = SourceDataBuffer.GetData();
+	return GetLightProfilePayloadData(SourceDataBuffer.GetData(), SourceDataBuffer.Num());
+}
 
-	FIESConverter IESConverter(Buffer, SourceDataBuffer.Num());
+TOptional<UE::Interchange::FImportLightProfile> UInterchangeIESTranslator::GetLightProfilePayloadData(const uint8* Buffer, uint32 BufferLength) const
+{
+	UE::Interchange::FImportLightProfile Payload;
 
-	if(IESConverter.IsValid())
-	{
-		UE::Interchange::FImportLightProfile Payload;
+	Payload.SourceDataBuffer.SetNum(BufferLength);
 
+	FPlatformMemory::Memcpy(Payload.SourceDataBuffer.GetData(), Buffer, BufferLength);
+
+	FIESConverter IESConverter(Payload.SourceDataBuffer.GetData(), Payload.SourceDataBuffer.Num());
+
+	if (IESConverter.IsValid())
+	{
 		Payload.Init2DWithParams(
 			IESConverter.GetWidth(),
 			IESConverter.GetHeight(),
 			TSF_RGBA16F,
 			false
-			);
+		);
 
 		Payload.CompressionSettings = TC_HDR;
 		Payload.Brightness = IESConverter.GetBrightness();
 		Payload.TextureMultiplier = IESConverter.GetMultiplier();
 
 
-		const TArray<uint8>&  RawData = IESConverter.GetRawData();;
+		const TArray<uint8>& RawData = IESConverter.GetRawData();;
 
 		FPlatformMemory::Memcpy(Payload.RawData.GetData(), RawData.GetData(), Payload.RawData.GetSize());
 
@@ -78,7 +85,7 @@ TOptional<UE::Interchange::FImportLightProfile> UInterchangeIESTranslator::GetLi
 	}
 	else
 	{
-		FTextureTranslatorUtilities::LogError(*this, FText::Format(LOCTEXT("IESImportFailure_Generic", "IES import failed: \"{0}\""), FText::FromString(IESConverter.GetError())));
+		UE::Interchange::FTextureTranslatorUtilities::LogError(*this, FText::Format(LOCTEXT("IESImportFailure_Generic", "IES import failed: \"{0}\""), FText::FromString(IESConverter.GetError())));
 	}
 
 	return {};
diff --git a/Engine/Plugins/Interchange/Runtime/Source/Import/Private/Texture/InterchangeTextureFactory.cpp b/Engine/Plugins/Interchange/Runtime/Source/Import/Private/Texture/InterchangeTextureFactory.cpp
--- a/Engine/Plugins/Interchange/Runtime/Source/Import/Private/Texture/InterchangeTextureFactory.cpp
+++ b/Engine/Plugins/Interchange/Runtime/Source/Import/Private/Texture/InterchangeTextureFactory.cpp
@@ -1697,7 +1697,7 @@ namespace UE::Interchange::Private::InterchangeTextureFactory
 #endif
 	}
 
-	void SetupTexture2DSourceDataFromPayload(UTexture2D* Texture2D, FProcessedPayload& ProcessedPayload, bool bIsReimport)
+	void SetupTexture2DSourceDataFromPayload(UInterchangeFactoryBaseNode* FactorBaseNode, UTexture2D* Texture2D, FProcessedPayload& ProcessedPayload, bool bIsReimport)
 	{
 #if WITH_EDITOR
 		if (TOptional<FImportBlockedImage>* BlockedImage = ProcessedPayload.SettingsFromPayload.TryGet<TOptional<FImportBlockedImage>>())
@@ -1726,16 +1726,15 @@ namespace UE::Interchange::Private::InterchangeTextureFactory
 				{
 					SetupTextureSourceDataFromBulkData(TextureLightProfile, LightProfile, MoveTemp(ProcessedPayload.PayloadAndId), bIsReimport);
 
-					if (const TOptional<FImportLightProfile>* ImportLightProfilePtr = ProcessedPayload.SettingsFromPayload.TryGet<TOptional<FImportLightProfile>>())
-					{
-						if (ImportLightProfilePtr->IsSet())
-						{
-							const UE::Interchange::FImportLightProfile& ImportLightProfile = ImportLightProfilePtr->GetValue();
+					const UE::Interchange::FImportLightProfile& ImportLightProfile = OptionalLightProfile->GetValue();
 
-							// The legacy factory updated these setting on each reimport (leave it as it is for 5.1 but we should look into moving that decision into the pipelines)
-							TextureLightProfile->Brightness = ImportLightProfile.Brightness;
-							TextureLightProfile->TextureMultiplier = ImportLightProfile.TextureMultiplier;
-						}
+					// The legacy factory updated these setting on each reimport (leave it as it is for 5.1 but we should look into moving that decision into the pipelines)
+					TextureLightProfile->Brightness = ImportLightProfile.Brightness;
+					TextureLightProfile->TextureMultiplier = ImportLightProfile.TextureMultiplier;
+
+					if (UInterchangeTextureLightProfileFactoryNode* TextureProfileFactoryNode = Cast<UInterchangeTextureLightProfileFactoryNode>(FactorBaseNode))
+					{
+						TextureProfileFactoryNode->StoreIESSourceFileContents(ImportLightProfile.SourceDataBuffer);
 					}
 				}
 				else
@@ -1761,11 +1760,11 @@ namespace UE::Interchange::Private::InterchangeTextureFactory
 #endif // WITH_EDITORONLY_DATA
 	}
 
-	void SetupTextureSourceDataFromPayload(UTexture* Texture, FProcessedPayload& ProcessedPayload, bool bIsReimport)
+	void SetupTextureSourceDataFromPayload(UInterchangeFactoryBaseNode* FactorBaseNode, UTexture* Texture, FProcessedPayload& ProcessedPayload, bool bIsReimport)
 	{
 		if (UTexture2D* Texture2D = Cast<UTexture2D>(Texture))
 		{
-			SetupTexture2DSourceDataFromPayload(Texture2D, ProcessedPayload, bIsReimport);
+			SetupTexture2DSourceDataFromPayload(FactorBaseNode, Texture2D, ProcessedPayload, bIsReimport);
 		}
 #if WITH_EDITOR
 		else if (UTextureCube* TextureCube = Cast<UTextureCube>(Texture))
@@ -2110,10 +2109,10 @@ void UInterchangeTextureFactory::SetupObject_GameThread(const FSetupObjectParams
 		using namespace UE::Interchange::Private::InterchangeTextureFactory;
 
 		// Setup source data
-		SetupTextureSourceDataFromPayload(Texture, ProcessedPayload, Arguments.bIsReimport);
+		UInterchangeFactoryBaseNode* TextureFactoryNode = Arguments.FactoryNode;
+		SetupTextureSourceDataFromPayload(TextureFactoryNode, Texture, ProcessedPayload, Arguments.bIsReimport);
 
 #if WITH_EDITOR
-		UInterchangeFactoryBaseNode* TextureFactoryNode = Arguments.FactoryNode;
 		if (!Arguments.bIsReimport)
 		{
 			// Finalize texture properties, before ApplyAllCustomAttributeToObject :
diff --git a/Engine/Plugins/Interchange/Runtime/Source/Import/Public/Gltf/InterchangeGltfTranslator.h b/Engine/Plugins/Interchange/Runtime/Source/Import/Public/Gltf/InterchangeGltfTranslator.h
--- a/Engine/Plugins/Interchange/Runtime/Source/Import/Public/Gltf/InterchangeGltfTranslator.h
+++ b/Engine/Plugins/Interchange/Runtime/Source/Import/Public/Gltf/InterchangeGltfTranslator.h
@@ -10,6 +10,7 @@
 #include "Nodes/InterchangeBaseNodeContainer.h"
 #include "Scene/InterchangeVariantSetPayloadInterface.h"
 #include "Texture/InterchangeTexturePayloadInterface.h"
+#include "Texture/InterchangeTextureLightProfilePayloadInterface.h"
 #include "UObject/Object.h"
 #include "UObject/ObjectMacros.h"
 
@@ -27,6 +28,7 @@ class UInterchangeGLTFTranslator : public UInterchangeTranslatorBase
 	, public IInterchangeTexturePayloadInterface
 	, public IInterchangeAnimationPayloadInterface
 	, public IInterchangeVariantSetPayloadInterface
+	, public IInterchangeTextureLightProfilePayloadInterface
 {
 	GENERATED_BODY()
 
@@ -61,6 +63,10 @@ public:
 	virtual TOptional<UE::Interchange::FVariantSetPayloadData> GetVariantSetPayloadData(const FString& PayloadKey) const override;
 	/* IInterchangeVariantSetPayloadInterface End */
 
+	/* IInterchangeTextureLightProfilePayloadInterface Begin */
+	virtual TOptional<UE::Interchange::FImportLightProfile> GetLightProfilePayloadData(const FString& PayloadKey, TOptional<FString>& AlternateTexturePath) const override;
+	/* IInterchangeTextureLightProfilePayloadInterface End */
+
 protected:
 	using FNodeUidMap = TMap<const GLTF::FNode*, FString>;
 
diff --git a/Engine/Plugins/Interchange/Runtime/Source/Import/Public/Texture/InterchangeIESTranslator.h b/Engine/Plugins/Interchange/Runtime/Source/Import/Public/Texture/InterchangeIESTranslator.h
--- a/Engine/Plugins/Interchange/Runtime/Source/Import/Public/Texture/InterchangeIESTranslator.h
+++ b/Engine/Plugins/Interchange/Runtime/Source/Import/Public/Texture/InterchangeIESTranslator.h
@@ -42,5 +42,7 @@ public:
 	 */
 	virtual TOptional<UE::Interchange::FImportLightProfile> GetLightProfilePayloadData(const FString& PayloadKey, TOptional<FString>& AlternateTexturePath) const override;
 
+	virtual TOptional<UE::Interchange::FImportLightProfile> GetLightProfilePayloadData(const uint8* Buffer, uint32 BufferLength) const override;
+
 	/* IInterchangeTexturePayloadInterface End */
 };
\ No newline at end of file
diff --git a/Engine/Plugins/Interchange/Runtime/Source/Import/Public/Texture/InterchangeTextureLightProfilePayloadData.h b/Engine/Plugins/Interchange/Runtime/Source/Import/Public/Texture/InterchangeTextureLightProfilePayloadData.h
--- a/Engine/Plugins/Interchange/Runtime/Source/Import/Public/Texture/InterchangeTextureLightProfilePayloadData.h
+++ b/Engine/Plugins/Interchange/Runtime/Source/Import/Public/Texture/InterchangeTextureLightProfilePayloadData.h
@@ -11,5 +11,7 @@ namespace UE::Interchange
 		float Brightness = 0.0f;
 
 		float TextureMultiplier = 0.0f;
+
+		TArray64<uint8> SourceDataBuffer;
 	};
 }
\ No newline at end of file
diff --git a/Engine/Plugins/Interchange/Runtime/Source/Import/Public/Texture/InterchangeTextureLightProfilePayloadInterface.h b/Engine/Plugins/Interchange/Runtime/Source/Import/Public/Texture/InterchangeTextureLightProfilePayloadInterface.h
--- a/Engine/Plugins/Interchange/Runtime/Source/Import/Public/Texture/InterchangeTextureLightProfilePayloadInterface.h
+++ b/Engine/Plugins/Interchange/Runtime/Source/Import/Public/Texture/InterchangeTextureLightProfilePayloadInterface.h
@@ -39,4 +39,9 @@ public:
 	 * @return a PayloadData containing the import light profile data. The TOptional will not be set if there is an error.
 	 */
 	virtual TOptional<UE::Interchange::FImportLightProfile> GetLightProfilePayloadData(const FString& PayloadKey, TOptional<FString>& AlternateTexturePath) const PURE_VIRTUAL(IInterchangeTextureLightProfilePayloadInterface::GetLightProfilePayloadData, return{};);
+
+	/**
+	* IESTranslator uses this, but we also expose it here in case otrher translators want to use directly.
+	*/
+	virtual TOptional<UE::Interchange::FImportLightProfile> GetLightProfilePayloadData(const uint8* Buffer, uint32 BufferLength) const { return {}; }
 };
diff --git a/Engine/Plugins/Interchange/Runtime/Source/Parsers/GLTFCore/Private/GLTF/ExtensionsHandler.cpp b/Engine/Plugins/Interchange/Runtime/Source/Parsers/GLTFCore/Private/GLTF/ExtensionsHandler.cpp
--- a/Engine/Plugins/Interchange/Runtime/Source/Parsers/GLTFCore/Private/GLTF/ExtensionsHandler.cpp
+++ b/Engine/Plugins/Interchange/Runtime/Source/Parsers/GLTFCore/Private/GLTF/ExtensionsHandler.cpp
@@ -6,6 +6,8 @@
 #include "GLTFAsset.h"
 #include "GLTFNode.h"
 #include "MaterialUtilities.h"
+#include "GenericPlatform/GenericPlatformHttp.h"
+#include "Misc/Paths.h"
 
 #if USE_DRACO_LIBRARY
 #include "draco/mesh/mesh.h"
@@ -115,25 +117,7 @@ namespace GLTF
 
 	namespace
 	{
-		static const TArray<FString> LightExtensions = { GLTF::ToString(GLTF::EExtension::KHR_LightsPunctual), GLTF::ToString(GLTF::EExtension::KHR_Lights) };
 		static const TArray<FString> MaterialsExtensions = { GLTF::ToString(GLTF::EExtension::KHR_MaterialsVariants) };
-	
-		TSharedPtr<FJsonObject> GetLightExtension(const TSharedPtr<FJsonObject>& Object)
-		{
-			if (!Object)
-				return nullptr;
-
-			TSharedPtr<FJsonObject> LightsObj;
-			if (Object->HasTypedField<EJson::Object>(LightExtensions[0]))
-			{
-				LightsObj = Object->GetObjectField(LightExtensions[0]);
-			}
-			else if (Object->HasTypedField<EJson::Object>(LightExtensions[1]))
-			{
-				LightsObj = Object->GetObjectField(LightExtensions[1]);
-			}
-			return LightsObj;
-		}
 
 		TSharedPtr<FJsonObject> GetExtensions(const FJsonObject& Object)
 		{
@@ -152,32 +136,114 @@ namespace GLTF
 	{
 	}
 
-	void FExtensionsHandler::SetupAssetExtensions(const FJsonObject& Object) const
+	void FExtensionsHandler::SetupLightExtensions(const FJsonObject& Object, const FString& InResourcesPath) const
 	{
-		const TSharedPtr<FJsonObject>& ExtensionsObj = GetExtensions(Object);
-
-		if (!ExtensionsObj.IsValid())
+		static const TArray<EExtension> Extensions = {
+			EExtension::KHR_LightsPunctual,
+			EExtension::KHR_Lights,
+			EExtension::EXT_LightsIES
+		};
+		TArray<FString> ExtensionsStringified;
+		for (size_t ExtensionIndex = 0; ExtensionIndex < Extensions.Num(); ExtensionIndex++)
 		{
-			return;
+			ExtensionsStringified.Add(GLTF::ToString(Extensions[ExtensionIndex]));
 		}
 
-		// lights
-
-		if (TSharedPtr<FJsonObject> LightsObj = GetLightExtension(ExtensionsObj))
+		const FJsonObject& ExtensionsObj = *Object.GetObjectField(TEXT("extensions"));
+		for (int32 Index = 0; Index < Extensions.Num(); ++Index)
 		{
-			Asset->ProcessedExtensions.Add(EExtension::KHR_LightsPunctual);
+			const FString ExtensionName = ExtensionsStringified[Index];
+			if (!ExtensionsObj.HasTypedField<EJson::Object>(ExtensionName))
+				continue;
+
+			const FJsonObject& ExtObj = *ExtensionsObj.GetObjectField(ExtensionName);
 
-			uint32 LightCount = ArraySize(*LightsObj, TEXT("lights"));
-			if (LightCount > 0)
+			const EExtension Extension = Extensions[Index];
+			
+			switch (Extension)
 			{
-				Asset->Lights.Reserve(LightCount);
-				for (const TSharedPtr<FJsonValue>& Value : LightsObj->GetArrayField(TEXT("lights")))
-				{
-					SetupLightPunctual(*Value->AsObject());
-				}
+				case EExtension::KHR_Lights:
+				case EExtension::KHR_LightsPunctual:
+					{
+						const FJsonObject& LightsObj = ExtObj;
+
+						Asset->ProcessedExtensions.Add(EExtension::KHR_LightsPunctual);
+
+						uint32 LightCount = ArraySize(LightsObj, TEXT("lights"));
+						if (LightCount > 0)
+						{
+							Asset->Lights.Reserve(LightCount);
+							for (const TSharedPtr<FJsonValue>& Value : LightsObj.GetArrayField(TEXT("lights")))
+							{
+								SetupLightPunctual(*Value->AsObject());
+							}
+						}
+					}
+					break;
+				case EExtension::EXT_LightsIES:
+					{
+						const FJsonObject& LightsIESObj = ExtObj;
+
+						Asset->ProcessedExtensions.Add(EExtension::EXT_LightsIES);
+
+						uint32 LightIESCount = ArraySize(LightsIESObj, TEXT("lights"));
+						if (LightIESCount > 0)
+						{
+							Asset->LightsIES.Reserve(LightIESCount);
+							
+							for (const TSharedPtr<FJsonValue>& Value : LightsIESObj.GetArrayField(TEXT("lights")))
+							{
+								const FJsonObject& ObjectLightIES = *Value->AsObject();
+								
+								FLightIES& LightIES = Asset->LightsIES.Emplace_GetRef();
+								LightIES.Index = Asset->LightsIES.Num() - 1;
+								
+								LightIES.URI = GetString(ObjectLightIES, TEXT("uri"));
+								if (LightIES.URI.Len() > 0)
+								{
+									LightIES.URI = FGenericPlatformHttp::UrlDecode(LightIES.URI);
+									LightIES.FilePath = InResourcesPath / LightIES.URI;
+								}
+
+								LightIES.BufferViewIndex = GetIndex(ObjectLightIES, TEXT("bufferView"));
+								if (Asset->BufferViews.IsValidIndex(LightIES.BufferViewIndex))
+								{
+									const FBufferView& BufferView = Asset->BufferViews[LightIES.BufferViewIndex];
+									LightIES.DataByteLength = BufferView.ByteLength;
+									LightIES.Data = static_cast<const uint8*>(BufferView.DataAt(0));
+								}
+
+								LightIES.MimeType = GetString(ObjectLightIES, TEXT("mimeType"));
+								LightIES.Name = GetString(ObjectLightIES, TEXT("name"));
+
+								if (LightIES.Name.Len() == 0 && LightIES.URI.Len() > 0)
+								{
+									LightIES.Name = FPaths::GetBaseFilename(LightIES.URI);
+								}
+							}
+						}
+					}
+					break;
+				default:
+					break;
 			}
 		}
 
+		CheckExtensions(Object, ExtensionsStringified);
+	}
+
+	void FExtensionsHandler::SetupAssetExtensions(const FJsonObject& Object, const FString& InResourcesPath) const
+	{
+		const TSharedPtr<FJsonObject>& ExtensionsObj = GetExtensions(Object);
+
+		if (!ExtensionsObj.IsValid())
+		{
+			return;
+		}
+
+		// lights
+		SetupLightExtensions(Object, InResourcesPath);
+
 		// variants
 
 		if (ExtensionsObj->HasTypedField<EJson::Object>(GLTF::ToString(GLTF::EExtension::KHR_MaterialsVariants)))
@@ -199,8 +265,10 @@ namespace GLTF
 		}
 		
 		TArray<FString> SupportedExtensions;
-		SupportedExtensions.Append(LightExtensions);
 		SupportedExtensions.Append(MaterialsExtensions);
+		SupportedExtensions.Add(GLTF::ToString(GLTF::EExtension::KHR_Lights));
+		SupportedExtensions.Add(GLTF::ToString(GLTF::EExtension::KHR_LightsPunctual));
+		SupportedExtensions.Add(GLTF::ToString(GLTF::EExtension::EXT_LightsIES));
 
 		CheckExtensions(Object, SupportedExtensions);
 	}
@@ -587,14 +655,65 @@ namespace GLTF
 
 	void FExtensionsHandler::SetupNodeExtensions(const FJsonObject& Object, FNode& Node) const
 	{
-		const TSharedPtr<FJsonObject>& ExtensionsObj = GetExtensions(Object);
-		if (TSharedPtr<FJsonObject> LightsObj = GetLightExtension(ExtensionsObj))
+		static const TArray<EExtension> Extensions = {
+			EExtension::KHR_LightsPunctual,
+			EExtension::KHR_Lights,
+			EExtension::EXT_LightsIES
+		};
+		TArray<FString> ExtensionsStringified;
+		for (size_t ExtensionIndex = 0; ExtensionIndex < Extensions.Num(); ExtensionIndex++)
+		{
+			ExtensionsStringified.Add(GLTF::ToString(Extensions[ExtensionIndex]));
+		}
+
+		if (Object.HasField(TEXT("extensions")))
 		{
-			Node.LightIndex = GetIndex(*LightsObj, TEXT("light"));
+			const FJsonObject& ExtensionsObj = *Object.GetObjectField(TEXT("extensions"));
+			for (int32 Index = 0; Index < Extensions.Num(); ++Index)
+			{
+				const FString ExtensionName = ExtensionsStringified[Index];
+				if (!ExtensionsObj.HasTypedField<EJson::Object>(ExtensionName))
+					continue;
+
+				const FJsonObject& ExtObj = *ExtensionsObj.GetObjectField(ExtensionName);
+
+				const EExtension Extension = Extensions[Index];
+
+				switch (Extension)
+				{
+					case EExtension::KHR_Lights:
+					case EExtension::KHR_LightsPunctual:
+					{
+						Node.LightIndex = GetIndex(ExtObj, TEXT("light"));
+					}
+					break;
+					case EExtension::EXT_LightsIES:
+					{
+						const FJsonObject& IESNodeLight = ExtObj;
+
+						Node.LightIES.Index = GetIndex(IESNodeLight, TEXT("light"));
+
+						if (IESNodeLight.HasField(TEXT("multiplier")))
+						{
+							Node.LightIES.bHasIntensityMultiplier = true;
+							Node.LightIES.IntensityMultipler = GetScalar(IESNodeLight, TEXT("multiplier"), 1.0f);
+						}
+
+						if (IESNodeLight.HasField(TEXT("color")))
+						{
+							Node.LightIES.bHasColor = true;
+							Node.LightIES.Color = GetVec3(IESNodeLight, TEXT("color"), FVector(1.f));
+						}
+
+					}
+					break;
+					default:
+						break;
+				}
+			}
 		}
 
-		static const TArray<FString> Extensions = LightExtensions;
-		CheckExtensions(Object, Extensions);
+		CheckExtensions(Object, ExtensionsStringified);
 	}
 
 	void FExtensionsHandler::SetupCameraExtensions(const FJsonObject& Object, FCamera& Camera) const
diff --git a/Engine/Plugins/Interchange/Runtime/Source/Parsers/GLTFCore/Private/GLTF/ExtensionsHandler.h b/Engine/Plugins/Interchange/Runtime/Source/Parsers/GLTFCore/Private/GLTF/ExtensionsHandler.h
--- a/Engine/Plugins/Interchange/Runtime/Source/Parsers/GLTFCore/Private/GLTF/ExtensionsHandler.h
+++ b/Engine/Plugins/Interchange/Runtime/Source/Parsers/GLTFCore/Private/GLTF/ExtensionsHandler.h
@@ -31,7 +31,7 @@ namespace GLTF
 
 		void SetAsset(GLTF::FAsset& InAsset);
 
-		void SetupAssetExtensions(const FJsonObject& Object) const;
+		void SetupAssetExtensions(const FJsonObject& Object, const FString& InResourcesPath) const;
 		void SetupBufferExtensions(const FJsonObject& Object, FBuffer& Buffer) const;
 		void SetupBufferViewExtensions(const FJsonObject& Object, FBufferView& BufferView) const;
 		void SetupAccessorExtensions(const FJsonObject& Object, FAccessor& Accessor) const;
@@ -46,6 +46,7 @@ namespace GLTF
 		void SetupSamplerExtensions(const FJsonObject& Object, FSampler& Sampler) const;
 		void SetupTextureExtensions(const FJsonObject& Object, FTexture& Texture) const;
 		void SetupMaterialExtensions(const FJsonObject& Object, FMaterial& Material) const;
+		void SetupLightExtensions(const FJsonObject& Object, const FString& InResourcesPath) const;
 
 		TArray<FLogMessage>& GetMessages()
 		{
diff --git a/Engine/Plugins/Interchange/Runtime/Source/Parsers/GLTFCore/Private/GLTF/GLTFAsset.cpp b/Engine/Plugins/Interchange/Runtime/Source/Parsers/GLTFCore/Private/GLTF/GLTFAsset.cpp
--- a/Engine/Plugins/Interchange/Runtime/Source/Parsers/GLTFCore/Private/GLTF/GLTFAsset.cpp
+++ b/Engine/Plugins/Interchange/Runtime/Source/Parsers/GLTFCore/Private/GLTF/GLTFAsset.cpp
@@ -249,12 +249,14 @@ namespace GLTF
 		GLTF::GenerateNames(Name + TEXT("_animation_"), Animations);
 		GLTF::GenerateNames(Name + TEXT("_image_"), Images);
 		GLTF::GenerateNames(Name + TEXT("_mesh_"), Meshes);
+		GLTF::GenerateNames(Name + TEXT("_IES_"), LightsIES);
 
 		// Validate FName criteria for Materials Textures Meshes Animations (assets)
 		GLTF::ValidateFNameCriteria(TEXT("Material_"), Materials);
 		GLTF::ValidateFNameCriteria(TEXT("Texture_"), Textures);
 		GLTF::ValidateFNameCriteria(TEXT("Mesh_"), Meshes);
 		GLTF::ValidateFNameCriteria(TEXT("Animation_"), Animations);
+		GLTF::ValidateFNameCriteria(TEXT("TextureIES_"), LightsIES);
 
 		//Generate MorphTarget Names:
 		//MorphTargetNames have to be unique across meshes as well:
@@ -421,6 +423,8 @@ namespace GLTF
 				return TEXT("KHR_lights_punctual");
 			case GLTF::EExtension::KHR_Lights:
 				return TEXT("KHR_lights");
+			case GLTF::EExtension::EXT_LightsIES:
+				return TEXT("EXT_lights_ies");
 			case GLTF::EExtension::KHR_Blend:
 				return TEXT("KHR_blend");
 			case GLTF::EExtension::MSFT_TextureDDS:
diff --git a/Engine/Plugins/Interchange/Runtime/Source/Parsers/GLTFCore/Private/GLTF/GLTFReader.cpp b/Engine/Plugins/Interchange/Runtime/Source/Parsers/GLTFCore/Private/GLTF/GLTFReader.cpp
--- a/Engine/Plugins/Interchange/Runtime/Source/Parsers/GLTFCore/Private/GLTF/GLTFReader.cpp
+++ b/Engine/Plugins/Interchange/Runtime/Source/Parsers/GLTFCore/Private/GLTF/GLTFReader.cpp
@@ -1207,7 +1207,7 @@ namespace GLTF
 			}
 		}
 
-		ExtensionsHandler->SetupAssetExtensions(*JsonRoot);
+		ExtensionsHandler->SetupAssetExtensions(*JsonRoot, InResourcesPath);
 	}
 
 	bool FFileReader::CheckForErrors(int32 StartIndex) const
@@ -1294,7 +1294,7 @@ namespace GLTF
 			{
 				Node.Type = FNode::EType::Camera;
 			}
-			else if (Node.LightIndex != INDEX_NONE)
+			else if (Node.LightIndex != INDEX_NONE || Node.LightIES.Index != INDEX_NONE)
 			{
 				Node.Type = FNode::EType::Light;
 			}
diff --git a/Engine/Plugins/Interchange/Runtime/Source/Parsers/GLTFCore/Public/GLTFAsset.h b/Engine/Plugins/Interchange/Runtime/Source/Parsers/GLTFCore/Public/GLTFAsset.h
--- a/Engine/Plugins/Interchange/Runtime/Source/Parsers/GLTFCore/Public/GLTFAsset.h
+++ b/Engine/Plugins/Interchange/Runtime/Source/Parsers/GLTFCore/Public/GLTFAsset.h
@@ -43,6 +43,7 @@ namespace GLTF
 		KHR_DracoMeshCompression,
 		KHR_LightsPunctual,
 		KHR_Lights,
+		EXT_LightsIES,
 		KHR_Blend,
 		MSFT_TextureDDS,
 		MSFT_PackingNormalRoughnessMetallic,
@@ -95,6 +96,7 @@ namespace GLTF
 		TArray<FNode>      Nodes;
 		TArray<FCamera>    Cameras;
 		TArray<FLight>     Lights;
+		TArray<FLightIES>  LightsIES;
 		TArray<FSkinInfo>  Skins;
 		TArray<FAnimation> Animations;
 
diff --git a/Engine/Plugins/Interchange/Runtime/Source/Parsers/GLTFCore/Public/GLTFNode.h b/Engine/Plugins/Interchange/Runtime/Source/Parsers/GLTFCore/Public/GLTFNode.h
--- a/Engine/Plugins/Interchange/Runtime/Source/Parsers/GLTFCore/Public/GLTFNode.h
+++ b/Engine/Plugins/Interchange/Runtime/Source/Parsers/GLTFCore/Public/GLTFNode.h
@@ -13,6 +13,39 @@ namespace GLTF
 {
 	struct GLTFCORE_API FNode
 	{
+		struct GLTFCORE_API FLightInstanceIES
+		{
+			int32 Index = INDEX_NONE; // index into FAsset::LightsIES
+			
+			float IntensityMultipler = 1.f;
+			bool bHasIntensityMultiplier = false; //IES brightness scale
+			
+			FVector Color = FVector(1.f);
+			bool bHasColor = false;
+
+			FString GetHash() const
+			{
+				
+				if (!bHasIntensityMultiplier && !bHasColor)
+				{
+					return TEXT("");
+				}
+
+				FMD5 MD5;
+				
+				MD5.Update(reinterpret_cast<const uint8*>(&Index), sizeof(Index));
+				MD5.Update(reinterpret_cast<const uint8*>(&IntensityMultipler), sizeof(IntensityMultipler));
+				MD5.Update(reinterpret_cast<const uint8*>(&Color), sizeof(Color));
+
+				FMD5Hash Hash;
+				Hash.Set(MD5);
+
+				FString HashString = LexToString(Hash);
+
+				return HashString;
+			}
+		};
+
 		enum class EType
 		{
 			None,
@@ -36,6 +69,7 @@ namespace GLTF
 		int32 Skindex;      // index into FAsset::Skins
 		int32 CameraIndex;  // index into FAsset::Cameras
 		int32 LightIndex;   // index into FAsset::Lights
+		FLightInstanceIES LightIES;
 
 		int32 Index;		// index of FNode in GLTFAsset.Nodes array.
 		int32 ParentIndex;
@@ -103,6 +137,7 @@ namespace GLTF
 		FCamera(const FNode& Node)
 		    : Node(Node)
 		    , ZNear(0.f)
+			, ZFar(100.f)
 		    , bIsPerspective(true)
 		{
 			Perspective.Fov         = 0.f;
@@ -110,6 +145,23 @@ namespace GLTF
 		}
 	};
 
+	struct GLTFCORE_API FLightIES
+	{
+		int32 Index; //serves as idintifer as well.
+		
+		FString URI;
+		FString FilePath;
+		
+		int32 BufferViewIndex = INDEX_NONE;
+		uint32 DataByteLength;
+		const uint8* Data;
+		
+		FString MimeType;
+		
+		FString Name;
+		FString UniqueId; //will be generated in FAsset::GenerateNames
+	};
+
 	struct GLTFCORE_API FLight
 	{
 		enum class EType
@@ -144,6 +196,7 @@ namespace GLTF
 
 		FLight(const FNode* Node)
 		    : Node(Node)
+			, Type(EType::Point)
 		    , Color(1.f)
 		    , Intensity(1.f)
 		    , Range(1e+20)
