[Interchange] Fixing Rigid AnimSequence Import's SampleRate setting.

#rnx

diff --git a/Engine/Plugins/Interchange/Runtime/Source/Import/Private/Animation/InterchangeAnimSequenceFactory.cpp b/Engine/Plugins/Interchange/Runtime/Source/Import/Private/Animation/InterchangeAnimSequenceFactory.cpp
--- a/Engine/Plugins/Interchange/Runtime/Source/Import/Private/Animation/InterchangeAnimSequenceFactory.cpp
+++ b/Engine/Plugins/Interchange/Runtime/Source/Import/Private/Animation/InterchangeAnimSequenceFactory.cpp
@@ -1221,6 +1221,9 @@ UInterchangeFactoryBase::FImportAssetResult UInterchangeAnimSequenceFactory::Imp
 				FTransform LocalTransform;
 				SkeletonSceneNode->GetCustomLocalTransform(LocalTransform);
 
+				//As non-baked transforms do not have BakeFrequency concept we set the BakeFrequency here, in order for the CalculateDataFor can pick it up correctly.
+				AnimationTransformPayload.BakeFrequency = SampleRate;
+
 				//Currently only Curve -> Baked conversion (for LevelSequence->AnimSequence conversion by ForceMeshType Skeletal use case)
 				//and Curve -> Step Curve conversion (for custom attributes)
 				AnimationTransformPayload.CalculateDataFor(EInterchangeAnimationPayLoadType::BAKED, LocalTransform);
