MaterialX:
- Update Material Expression Overlay to 1.38.10 spec
- Add support of <refract> using Refract Engine material function
- Add support of ColorTransform group nodes

MaterialXSurfaceShaderAbstract:
- Fix a bug where we end up with an empty output string, it's happening when we flatten subgraphs that are the first node in the hierarchy (bug in the SDK)
- Always add the inputs of the nodedefs on a node, because we don't have a perfect 1:1 match with UE on the default values (MaterialX sets for min [0,0] whereas UE [0, 1])
- Change in the <position> to now use the LocalPosition expression instead of transforming a WorldSpace Position into local space
- Change the <noise3d> node to use VectorNoise expression instead of noise as the algorithm is closer to what the spec expects
- Changing the logic of the GetAttributeParentName function, we now put the name of the ShaderNode on the input, useful for nodes such as <noise3d> where we need to connect to a different input

InterchangeMaterialFactory:
- Add a test against the base method of UMaterialExpression::GetInputName, since sometimes underlying expressions override this function depending on different factors we can't control (for example UMaterialExpression sets the name of Position either as World Position or Camera Position)

Material Functions:
- Place2D now uses an operation order input to choose between SRT and TRS, adding the input for now to avoid errors during import (the logic is not in place for the moment)
- Add ColorTransforms Material Functions

Tests:
- Update screenshots of some materials
- Fix failing tests (geometric_nodes and procedural3d) and re-enable them

#virtualized

diff --git a/Engine/Plugins/Interchange/Runtime/Source/Common/Public/MaterialX/InterchangeMaterialXDefinitions.h b/Engine/Plugins/Interchange/Runtime/Source/Common/Public/MaterialX/InterchangeMaterialXDefinitions.h
--- a/Engine/Plugins/Interchange/Runtime/Source/Common/Public/MaterialX/InterchangeMaterialXDefinitions.h
+++ b/Engine/Plugins/Interchange/Runtime/Source/Common/Public/MaterialX/InterchangeMaterialXDefinitions.h
@@ -482,6 +482,7 @@ MATERIALX_NAMESPACE_BEGIN
 		static constexpr const char* CrossProduct = "crossproduct";
 		static constexpr const char* Determinant = "determinant";
 		static constexpr const char* Divide = "divide";
+		static constexpr const char* Distance = "distance";
 		static constexpr const char* DotProduct = "dotproduct";
 		static constexpr const char* Exp = "exp";
 		static constexpr const char* Floor = "floor";
@@ -497,8 +498,12 @@ MATERIALX_NAMESPACE_BEGIN
 		static constexpr const char* NormalMap = "normalmap";
 		static constexpr const char* Place2D = "place2d";
 		static constexpr const char* Power = "power";
+		static constexpr const char* Reflect = "reflect";
+		static constexpr const char* Refract = "refract";
 		static constexpr const char* Rotate2D = "rotate2d";
 		static constexpr const char* Rotate3D = "rotate3d";
+		static constexpr const char* Round = "round";
+		static constexpr const char* SafePower = "safepower";
 		static constexpr const char* Sign = "sign";
 		static constexpr const char* Sin = "sin";
 		static constexpr const char* Sqrt = "sqrt";
@@ -508,7 +513,8 @@ MATERIALX_NAMESPACE_BEGIN
 		static constexpr const char* TransformNormal = "transformnormal";
 		static constexpr const char* TransformPoint = "transformpoint";
 		static constexpr const char* TransformVector = "transformvector";
-		static constexpr const char* Transpose = "transpose";		
+		static constexpr const char* Transpose = "transpose";
+		static constexpr const char* TriangleWave = "trianglewave";
 		// Compositing nodes
 		static constexpr const char* Burn = "burn";
 		static constexpr const char* Difference = "difference";
@@ -620,6 +626,16 @@ MATERIALX_NAMESPACE_BEGIN
 		// Convolution
 		static constexpr const char* Blur = "blur";
 		static constexpr const char* HeightToNormal = "heighttonormal";
+		// ColorTransform
+		static constexpr const char* ACEScgToLinRec709 = "acescg_to_lin_rec709";
+		static constexpr const char* AdobeRgbToLinRec709 = "adobergb_to_lin_rec709";
+		static constexpr const char* G18Rec709ToLinRec709 = "g18_rec709_to_lin_rec709";
+		static constexpr const char* G22ApP1ToLinRec709 = "g22_ap1_to_lin_rec709";
+		static constexpr const char* LinAdobeRgbToLinRec709 = "lin_adobergb_to_lin_rec709";
+		static constexpr const char* LinDisplayP3ToLinRec709 = "lin_displayp3_to_lin_rec709";
+		static constexpr const char* Rec709DisplayToLinRec709 = "rec709_display_to_lin_rec709";
+		static constexpr const char* SrgbDisplayP3ToLinRec709 = "srgb_displayp3_to_lin_rec709";
+		static constexpr const char* SrgbTextureToLinRec709 = "srgb_texture_to_lin_rec709";
 	}
 
 	namespace NodeDefinition
diff --git a/Engine/Plugins/Interchange/Runtime/Source/Import/Private/Material/InterchangeMaterialFactory.cpp b/Engine/Plugins/Interchange/Runtime/Source/Import/Private/Material/InterchangeMaterialFactory.cpp
--- a/Engine/Plugins/Interchange/Runtime/Source/Import/Private/Material/InterchangeMaterialFactory.cpp
+++ b/Engine/Plugins/Interchange/Runtime/Source/Import/Private/Material/InterchangeMaterialFactory.cpp
@@ -86,7 +86,10 @@ namespace UE::Interchange::MaterialFactory::Internal
 		{
 			// MaterialFuncCall appends the type to the input name when calling GetInputName
 			// and the InputName in FExpressionInput is optional so we'll check both here to be safe
-			if (MaterialExpression.GetInputName(It.Index) == *InputName || It->InputName == *InputName)
+			// Sometimes we have to retrieve the input name as is because some Expressions override it based on different factors and we have no mechanism to keep tracks of it
+			if (MaterialExpression.GetInputName(It.Index) == *InputName || 
+				It->InputName == *InputName || 
+				MaterialExpression.UMaterialExpression::GetInputName(It.Index) == *InputName)
 			{
 				return It.Index;
 			}
diff --git a/Engine/Plugins/Interchange/Runtime/Source/Import/Private/MaterialX/MaterialExpressions/MaterialExpressionOverlay.cpp b/Engine/Plugins/Interchange/Runtime/Source/Import/Private/MaterialX/MaterialExpressions/MaterialExpressionOverlay.cpp
--- a/Engine/Plugins/Interchange/Runtime/Source/Import/Private/MaterialX/MaterialExpressions/MaterialExpressionOverlay.cpp
+++ b/Engine/Plugins/Interchange/Runtime/Source/Import/Private/MaterialX/MaterialExpressions/MaterialExpressionOverlay.cpp
@@ -42,6 +42,7 @@ int32 UMaterialExpressionMaterialXOverlay::Compile(FMaterialCompiler* Compiler,
 	int32 IndexAlpha = Alpha.GetTracedInput().Expression ? Alpha.Compile(Compiler) : Compiler->Constant(ConstAlpha);
 
 	int32 OneIndex = Compiler->Constant(1.f);
+	int32 TwoIndex = Compiler->Constant(2.f);
 	auto OneMinus = [&](int32 index)
 	{
 		return Compiler->Sub(OneIndex, index);
@@ -49,13 +50,13 @@ int32 UMaterialExpressionMaterialXOverlay::Compile(FMaterialCompiler* Compiler,
 
 	int32 IndexA = A.Compile(Compiler);
 	int32 IndexB = B.Compile(Compiler);
-	int32 Screen = OneMinus(Compiler->Mul(OneMinus(IndexA), OneMinus(IndexB)));
+	int32 Screen = OneMinus(Compiler->Mul(TwoIndex, Compiler->Mul(OneMinus(IndexA), OneMinus(IndexB))));
 
-	int32 Overlay = Compiler->If(IndexA,
+	int32 Overlay = Compiler->If(IndexB,
 								 Compiler->Constant(0.5f),
 								 Screen,
 								 Screen,
-								 Compiler->Mul(Compiler->Mul(Compiler->Constant(2.f), IndexA), IndexB),
+								 Compiler->Mul(Compiler->Mul(TwoIndex, IndexA), IndexB),
 								 Compiler->Constant(0.00001f));
 
 	return Compiler->Lerp(IndexB, Overlay, IndexAlpha);
@@ -81,6 +82,7 @@ bool UMaterialExpressionMaterialXOverlay::GenerateHLSLExpression(FMaterialHLSLGe
 
 	FTree& Tree = Generator.GetTree();
 	const FExpression* ExpressionOne = Tree.NewConstant(1.f);
+	const FExpression* ExpressionTwo = Tree.NewConstant(2.f);
 
 	auto NewOneMinus = [&](const FExpression* Input)
 	{
@@ -88,9 +90,9 @@ bool UMaterialExpressionMaterialXOverlay::GenerateHLSLExpression(FMaterialHLSLGe
 	};
 
 	const FExpression* ExpressionOverlay = Generator.GenerateBranch(Scope,
-																	Tree.NewLess(ExpressionA, Tree.NewConstant(0.5f)),
-																	Tree.NewMul(Tree.NewMul(Tree.NewConstant(2.f), ExpressionA), ExpressionB),
-																	NewOneMinus(Tree.NewMul(NewOneMinus(ExpressionA), NewOneMinus(ExpressionB))));
+																	Tree.NewLess(ExpressionB, Tree.NewConstant(0.5f)),
+																	Tree.NewMul(Tree.NewMul(ExpressionTwo, ExpressionA), ExpressionB),
+																	NewOneMinus(Tree.NewMul(ExpressionTwo, Tree.NewMul(NewOneMinus(ExpressionA), NewOneMinus(ExpressionB)))));
 
 	OutExpression = Tree.NewLerp(ExpressionB, ExpressionOverlay, ExpressionAlpha);
 
diff --git a/Engine/Plugins/Interchange/Runtime/Source/Import/Private/MaterialX/MaterialExpressions/MaterialExpressionOverlay.h b/Engine/Plugins/Interchange/Runtime/Source/Import/Private/MaterialX/MaterialExpressions/MaterialExpressionOverlay.h
--- a/Engine/Plugins/Interchange/Runtime/Source/Import/Private/MaterialX/MaterialExpressions/MaterialExpressionOverlay.h
+++ b/Engine/Plugins/Interchange/Runtime/Source/Import/Private/MaterialX/MaterialExpressions/MaterialExpressionOverlay.h
@@ -10,8 +10,8 @@
  * Blend nodes take two 1-4 channel inputs and apply the same operator to all channels.
  * Blend nodes support an optional float input mix , which can be used
  * to mix the original B value with the result of the blend operation.
- * Operation: 2*A*B          if A < 0.5;
- *            1-(1-A)(1-B) if A >= 0.5
+ * Operation: 2*A*B         if B < 0.5
+ *            1-2(1-A)(1-B) if B >= 0.5
  * Result: Lerp(B, Op, Alpha)
  */
 UCLASS(MinimalAPI, collapsecategories, hidecategories = Object, meta = (Private))
diff --git a/Engine/Plugins/Interchange/Runtime/Source/Import/Private/MaterialX/MaterialXUtils/MaterialXBase.cpp b/Engine/Plugins/Interchange/Runtime/Source/Import/Private/MaterialX/MaterialXUtils/MaterialXBase.cpp
--- a/Engine/Plugins/Interchange/Runtime/Source/Import/Private/MaterialX/MaterialXUtils/MaterialXBase.cpp
+++ b/Engine/Plugins/Interchange/Runtime/Source/Import/Private/MaterialX/MaterialXUtils/MaterialXBase.cpp
@@ -18,14 +18,24 @@ void FMaterialXBase::UpdateDocumentRecursively(MaterialX::GraphElementPtr Graph)
 	AddTexCoordToTiledImageNodes(Graph);
 
 	static std::set<std::string> SkippedCategories{
+		mx::Category::ACEScgToLinRec709,
+		mx::Category::AdobeRgbToLinRec709,
+		mx::Category::Distance,
 		mx::Category::Extract,
+		mx::Category::LinAdobeRgbToLinRec709,
+		mx::Category::LinDisplayP3ToLinRec709,
+		mx::Category::Noise3D,
 		mx::Category::OpenPBRSurface,
+		mx::Category::Overlay,
 		mx::Category::Place2D,
 		mx::Category::Ramp4,
+		mx::Category::Refract,
 		mx::Category::Saturate,
 		mx::Category::Separate2,
 		mx::Category::Separate3,
 		mx::Category::Separate4,
+		mx::Category::SrgbDisplayP3ToLinRec709,
+		mx::Category::SrgbTextureToLinRec709,
 		mx::Category::StandardSurface,
 		mx::Category::UsdPreviewSurface
 	};
diff --git a/Engine/Plugins/Interchange/Runtime/Source/Import/Private/MaterialX/MaterialXUtils/MaterialXManager.cpp b/Engine/Plugins/Interchange/Runtime/Source/Import/Private/MaterialX/MaterialXUtils/MaterialXManager.cpp
--- a/Engine/Plugins/Interchange/Runtime/Source/Import/Private/MaterialX/MaterialXUtils/MaterialXManager.cpp
+++ b/Engine/Plugins/Interchange/Runtime/Source/Import/Private/MaterialX/MaterialXUtils/MaterialXManager.cpp
@@ -89,7 +89,10 @@ FMaterialXManager::FMaterialXManager()
 		{{MaterialX::Category::Outside,        TEXT("in")},          TEXT("A")},			  // Outside is treated as Multiply node
 		{{MaterialX::Category::Outside,        TEXT("mask")},        TEXT("B")},			  // Outside is treated as Multiply node
 		{{MaterialX::Category::Power,          TEXT("in1")},         TEXT("Base")},
-		{{MaterialX::Category::Power,	       TEXT("in2")},         TEXT("Exponent")},
+		{{MaterialX::Category::Power,          TEXT("in2")},         TEXT("Exponent")},
+		{{MaterialX::Category::Refract,        TEXT("in")},          TEXT("Ray Direction")},
+		{{MaterialX::Category::Refract,        TEXT("normal")},      TEXT("Surface Normal")},
+		{{MaterialX::Category::Refract,        TEXT("ior")},         TEXT("Refractive Index Origin")},
 		{{MaterialX::Category::Rotate2D,       TEXT("in")},          TEXT("Coordinate")},
 		{{MaterialX::Category::Rotate2D,       TEXT("amount")},      TEXT("Time")},
 		{{MaterialX::Category::Rotate3D,       TEXT("amount")},      TEXT("RotationAngle")},
@@ -115,6 +118,7 @@ FMaterialXManager::FMaterialXManager()
 		{MaterialX::Category::Cos,          TEXT("Cosine")},
 		{MaterialX::Category::CrossProduct, TEXT("Crossproduct")},
 		{MaterialX::Category::Divide,       TEXT("Divide")},
+		{MaterialX::Category::Distance,     TEXT("Distance")},
 		{MaterialX::Category::DotProduct,   TEXT("Dotproduct")},
 		{MaterialX::Category::Exp,          TEXT("Exponential")},
 		{MaterialX::Category::Floor,        TEXT("Floor")},
@@ -129,6 +133,7 @@ FMaterialXManager::FMaterialXManager()
 		{MaterialX::Category::Power,        TEXT("Power")},
 		{MaterialX::Category::RampLR,       TEXT("MaterialXRampLeftRight")},
 		{MaterialX::Category::RampTB,       TEXT("MaterialXRampTopBottom")},
+		{MaterialX::Category::Round,        TEXT("Round")},
 		{MaterialX::Category::Sign,         TEXT("Sign")},
 		{MaterialX::Category::Sin,          TEXT("Sine")},
 		{MaterialX::Category::SplitLR,      TEXT("MaterialXSplitLeftRight")},
@@ -214,7 +219,13 @@ FMaterialXManager::FMaterialXManager()
 		TEXT("Value"),
 		TEXT("VectorInput"),
 		TEXT("X"),
-		TEXT("Y")
+		TEXT("Y"),
+
+		// case for Refract material function
+		TEXT("Ray Direction"),
+		TEXT("Surface Normal"),
+		TEXT("Refractive Index Origin"),
+		TEXT("Refractive Index Target"),
 	}
 	, MaterialXContainerDelegates{
 		{mx::Category::OpenPBRSurface, FMaterialXManager::FOnGetMaterialXInstance::CreateStatic(&FMaterialXOpenPBRSurfaceShader::MakeInstance)},
@@ -254,6 +265,14 @@ FMaterialXManager::FMaterialXManager()
 			{mx::Category::RoughnessDual,			FMaterialXMaterialFunction{TInPlaceType<FString>{}, TEXT("/InterchangeAssets/Functions/MX_Roughness_Dual.MX_Roughness_Dual")}},
 			// Math
 			{mx::Category::Place2D,					FMaterialXMaterialFunction{TInPlaceType<FString>{}, TEXT("/InterchangeAssets/Functions/MX_Place2D.MX_Place2D")}},
+			{mx::Category::Refract,					FMaterialXMaterialFunction{TInPlaceType<FString>{}, TEXT("/Engine/Functions/Engine_MaterialFunctions01/Vectors/Refract.Refract")}},
+			//ColorTransform
+			{mx::Category::ACEScgToLinRec709,		FMaterialXMaterialFunction{TInPlaceType<FString>{}, TEXT("/InterchangeAssets/Functions/MX_ACEScgToRec709.MX_ACEScgToRec709")}},
+			{mx::Category::AdobeRgbToLinRec709,		FMaterialXMaterialFunction{TInPlaceType<FString>{}, TEXT("/InterchangeAssets/Functions/MX_AdobeRGBToRec709.MX_AdobeRGBToRec709")}},
+			{mx::Category::LinAdobeRgbToLinRec709,	FMaterialXMaterialFunction{TInPlaceType<FString>{}, TEXT("/InterchangeAssets/Functions/MX_LinearAdobeRGBToRec709.MX_LinearAdobeRGBToRec709")}},
+			{mx::Category::LinDisplayP3ToLinRec709,	FMaterialXMaterialFunction{TInPlaceType<FString>{}, TEXT("/InterchangeAssets/Functions/MX_LinearDisplayP3ToRec709.MX_LinearDisplayP3ToRec709")}},
+			{mx::Category::SrgbDisplayP3ToLinRec709,FMaterialXMaterialFunction{TInPlaceType<FString>{}, TEXT("/InterchangeAssets/Functions/MX_SrgbDisplayP3ToRec709.MX_SrgbDisplayP3ToRec709")}},
+			{mx::Category::SrgbTextureToLinRec709,	FMaterialXMaterialFunction{TInPlaceType<FString>{}, TEXT("/InterchangeAssets/Functions/MX_SrgbToRec709.MX_SrgbToRec709")}},
 		};
 
 		if(bIsSubstrateEnabled)
@@ -397,6 +416,13 @@ namespace UE::Interchange::MaterialX
 															    TEXT("MaterialFunction'/InterchangeAssets/Functions/MX_Roughness_Anisotropy.MX_Roughness_Anisotropy'"), 
 															    TEXT("MaterialFunction'/InterchangeAssets/Functions/MX_Roughness_Dual.MX_Roughness_Dual'"),
 																TEXT("MaterialFunction'/InterchangeAssets/Functions/MX_Place2D.MX_Place2D'"),
+																TEXT("MaterialFunction'/Engine/Functions/Engine_MaterialFunctions01/Vectors/Refract.Refract'"),
+																TEXT("MaterialFunction'/InterchangeAssets/Functions/MX_ACEScgToRec709.MX_ACEScgToRec709'"),
+																TEXT("MaterialFunction'/InterchangeAssets/Functions/MX_AdobeRGBToRec709.MX_AdobeRGBToRec709'"),
+																TEXT("MaterialFunction'/InterchangeAssets/Functions/MX_LinearDisplayP3ToRec709.MX_LinearDisplayP3ToRec709'"),
+																TEXT("MaterialFunction'/InterchangeAssets/Functions/MX_LinearAdobeRGBToRec709.MX_LinearAdobeRGBToRec709'"),
+																TEXT("MaterialFunction'/InterchangeAssets/Functions/MX_SrgbToRec709.MX_SrgbToRec709'"),
+																TEXT("MaterialFunction'/InterchangeAssets/Functions/MX_SrgbDisplayP3ToRec709.MX_SrgbDisplayP3ToRec709'"),
 															  });
 
 		return bPackagesLoaded;
diff --git a/Engine/Plugins/Interchange/Runtime/Source/Import/Private/MaterialX/MaterialXUtils/MaterialXSurfaceShaderAbstract.cpp b/Engine/Plugins/Interchange/Runtime/Source/Import/Private/MaterialX/MaterialXUtils/MaterialXSurfaceShaderAbstract.cpp
--- a/Engine/Plugins/Interchange/Runtime/Source/Import/Private/MaterialX/MaterialXUtils/MaterialXSurfaceShaderAbstract.cpp
+++ b/Engine/Plugins/Interchange/Runtime/Source/Import/Private/MaterialX/MaterialXUtils/MaterialXSurfaceShaderAbstract.cpp
@@ -206,11 +206,11 @@ bool FMaterialXSurfaceShaderAbstract::ConnectNodeGraphOutputToInput(MaterialX::I
 			return false;
 		}
 
+		// We may end up with an empty output string in case of a flattened nodegraph 
+		FString OutputString = Output->hasOutputString() && !Output->getOutputString().empty() ? ANSI_TO_TCHAR(Output->getOutputString().c_str()) : TEXT("out");
 		for(mx::Edge Edge : Output->traverseGraph())
 		{
-			Output->hasOutputString() ?
-				ConnectNodeCategoryOutputToInput(Edge, ShaderNode, ParentInputName, Output->getOutputString().c_str()) :
-				ConnectNodeCategoryOutputToInput(Edge, ShaderNode, ParentInputName);
+			ConnectNodeCategoryOutputToInput(Edge, ShaderNode, ParentInputName, OutputString);
 		}
 	}
 
@@ -306,7 +306,10 @@ void FMaterialXSurfaceShaderAbstract::ConnectNodeCategoryOutputToInput(const Mat
 		TArray<UInterchangeShaderNode*> ParentShaderNodeOutputs{ShaderNode};
 		FString InputChannelName = ParentInputName;
 
-		//Replace the input's name by the one used in UE
+		// Add all the inputs from the nodedef in case the node doesn't have them all (we don't have a 1:1 match on the default values otherwise)
+		// For exemple min defaults to (0,0) whereas UE to (0, 1)
+		UpstreamNode->addInputsFromNodeDef();
+		// Replace the input's name by the one used in UE
 		SetMatchingInputsNames(UpstreamNode);
 
 		FString OutputChannelName = OutputName;
@@ -318,7 +321,8 @@ void FMaterialXSurfaceShaderAbstract::ConnectNodeCategoryOutputToInput(const Mat
 		{
 			if(mx::NodePtr DownstreamNode = DownstreamElement->asA<mx::Node>())
 			{
-				if(mx::InputPtr ConnectedInput = Edge.getConnectingElement()->asA<mx::Input>())
+				mx::InputPtr ConnectedInput = Edge.getConnectingElement()->asA<mx::Input>();
+				if(ConnectedInput)
 				{
 					InputChannelName = GetInputName(ConnectedInput);
 					if(ConnectedInput->hasOutputString())
@@ -343,7 +347,7 @@ void FMaterialXSurfaceShaderAbstract::ConnectNodeCategoryOutputToInput(const Mat
 				ParentShaderNodeOutputs.Empty();
 				for(std::size_t i = 0; i < Outputs.size(); ++i)
 				{
-					if(UInterchangeShaderNode** FoundNode = ShaderNodes.Find({ GetAttributeParentName(DownstreamNode), Outputs[i]->getName().c_str()}))
+					if(UInterchangeShaderNode** FoundNode = ShaderNodes.Find({ GetAttributeParentName(DownstreamNode, ConnectedInput), Outputs[i]->getName().c_str()}))
 					{
 						//Connect the swizzle node between the upstream and downstream node
 						if(ChannelsNode)
@@ -714,9 +718,12 @@ void FMaterialXSurfaceShaderAbstract::ConnectIfGreaterInputToOutput(const FConne
 		AddAttributeFromValueOrInterface(Input, TEXT("AEqualsB"), NodeIf);
 
 		//Let's add a new input that is a copy of in2 to connect it to the equal input
-		mx::InputPtr Input3 = Connect.UpstreamNode->addInput("in3");
-		Input3->copyContentFrom(Input);
-		SetAttributeNewName(Input3, "AEqualsB");
+		if (!Connect.UpstreamNode->getInput("in3"))
+		{
+			mx::InputPtr Input3 = Connect.UpstreamNode->addInput("in3");
+			Input3->copyContentFrom(Input);
+			SetAttributeNewName(Input3, "AEqualsB");
+		}
 	}
 }
 
@@ -748,9 +755,12 @@ void FMaterialXSurfaceShaderAbstract::ConnectIfGreaterEqInputToOutput(const FCon
 		AddAttributeFromValueOrInterface(Input, TEXT("AEqualsB"), NodeIf);
 
 		//Let's add a new input that is a copy of in2 to connect it to the equal input
-		mx::InputPtr Input3 = Connect.UpstreamNode->addInput("in3");
-		Input3->copyContentFrom(Input);
-		SetAttributeNewName(Input3, "AEqualsB");
+		if (!Connect.UpstreamNode->getInput("in3"))
+		{
+			mx::InputPtr Input3 = Connect.UpstreamNode->addInput("in3");
+			Input3->copyContentFrom(Input);
+			SetAttributeNewName(Input3, "AEqualsB");
+		}
 	}
 }
 
@@ -782,9 +792,12 @@ void FMaterialXSurfaceShaderAbstract::ConnectIfEqualInputToOutput(const FConnect
 		AddAttributeFromValueOrInterface(Input, TEXT("AGreaterThanB"), NodeIf);
 
 		//Let's add a new input that is a copy of in2 to connect it to the equal input
-		mx::InputPtr Input3 = Connect.UpstreamNode->addInput("in3");
-		Input3->copyContentFrom(Input);
-		SetAttributeNewName(Input3, "AGreaterThanB");
+		if(!Connect.UpstreamNode->getInput("in3"))
+		{
+			mx::InputPtr Input3 = Connect.UpstreamNode->addInput("in3");
+			Input3->copyContentFrom(Input);
+			SetAttributeNewName(Input3, "AGreaterThanB");
+		}
 	}
 }
 
@@ -849,18 +862,21 @@ UInterchangeShaderNode* FMaterialXSurfaceShaderAbstract::ConnectGeometryInputToO
 void FMaterialXSurfaceShaderAbstract::ConnectPositionInputToOutput(const FConnectNode& Connect)
 {
 	using namespace UE::Interchange::Materials::Standard::Nodes;
-	UInterchangeShaderNode* UnitNode = ConnectGeometryInputToOutput(Connect, TEXT("Multiply"),
-																	TransformPosition::Name.ToString(),
-																	TransformPosition::Inputs::Input.ToString(),
-																	TransformPosition::Attributes::TransformSourceType.ToString(), EMaterialPositionTransformSource::TRANSFORMPOSSOURCE_World,
-																	TransformPosition::Attributes::TransformType.ToString(), EMaterialPositionTransformSource::TRANSFORMPOSSOURCE_Local,
-																	false);
+	UInterchangeShaderNode* PositionNode = CreateShaderNode(Connect.UpstreamNode, (Connect.UpstreamNode->getName() + "_Position").c_str(), TEXT("LocalPosition"));
+	
+	if (mx::InputPtr InputSpace = Connect.UpstreamNode->getInput("space"); InputSpace && InputSpace->getValueString() == "world")
+	{
+		PositionNode->SetCustomShaderType("WorldPosition");
+	}
+
+	UInterchangeShaderNode* UnitNode = CreateShaderNode(Connect.UpstreamNode, Connect.UpstreamNode->getName().c_str(), Multiply::Name.ToString());
 
 	// In case of the position node, it seems that the unit is different, we assume for now a conversion from mm -> m, even if UE by default is cm
 	// See standard_surface_marble_solid file, especially on the fractal3d node
-	UInterchangeShaderNode* PositionNode = CreateShaderNode(Connect.UpstreamNode, Connect.UpstreamNode->getName().c_str() + FString(TEXT("_Position")), TEXT("WorldPosition"));
 	UnitNode->AddFloatAttribute(UInterchangeShaderPortsAPI::MakeInputValueKey(TEXT("B")), 0.001f);
 	UInterchangeShaderPortsAPI::ConnectDefaultOuputToInput(UnitNode, TEXT("A"), PositionNode->GetUniqueID());
+
+	UInterchangeShaderPortsAPI::ConnectDefaultOuputToInput(Connect.ParentShaderNode, Connect.InputChannelName, UnitNode->GetUniqueID());
 }
 
 void FMaterialXSurfaceShaderAbstract::ConnectNormalInputToOutput(const FConnectNode& Connect)
@@ -928,15 +944,13 @@ void FMaterialXSurfaceShaderAbstract::ConnectNoise3DInputToOutput(const FConnect
 	using namespace UE::Interchange::Materials::Standard::Nodes;
 
 	// MaterialX defines the Noise3d as Perlin Noise which is multiplied by the Amplitude then Added to Pivot
-	UInterchangeShaderNode* NoiseNode = CreateShaderNode(Connect.UpstreamNode, Connect.UpstreamNode->getName().c_str(), Noise::Name.ToString());
-	NoiseNode->AddInt32Attribute(Noise::Attributes::Function.ToString(), ENoiseFunction::NOISEFUNCTION_GradientTex);
-	NoiseNode->AddBooleanAttribute(Noise::Attributes::Turbulence.ToString(), false);
-	NoiseNode->AddFloatAttribute(Noise::Attributes::OutputMin.ToString(), 0);
+	UInterchangeShaderNode* NoiseNode = CreateShaderNode(Connect.UpstreamNode, Connect.UpstreamNode->getName().c_str(), VectorNoise::Name.ToString());
+	NoiseNode->AddInt32Attribute(VectorNoise::Attributes::Function.ToString(), EVectorNoiseFunction::VNF_VectorALU);
 
 	UInterchangeShaderNode* NodeToConnect = NoiseNode;
 
 	// Multiply Node
-	auto ConnectNodeToInput = [&](mx::InputPtr Input, UInterchangeShaderNode* NodeToConnectTo, const FString& ShaderType, int32 IndexAttrib) -> UInterchangeShaderNode*
+	auto ConnectNodeToInput = [&](mx::InputPtr Input, UInterchangeShaderNode* NodeToConnectTo, const FString& ShaderType) -> UInterchangeShaderNode*
 	{
 		if(!Input)
 		{
@@ -951,19 +965,18 @@ void FMaterialXSurfaceShaderAbstract::ConnectNoise3DInputToOutput(const FConnect
 		// Connect the amplitude node to the shader node not the noise
 		// it will be handle during the upstream-downstream connection phase
 		// The index is here for the unicity of the attribute
-		Connect.UpstreamNode->setAttribute(mx::Attributes::ParentName + std::to_string(IndexAttrib), TCHAR_TO_UTF8(*ShaderNodeName));
+		Input->setAttribute(mx::Attributes::ParentName, TCHAR_TO_UTF8(*ShaderNodeName));
 		AddAttributeFromValueOrInterface(Input, TEXT("B"), ShaderNode);
 
 		return ShaderNode;
 	};
 
-	if(UInterchangeShaderNode* MultiplyNode = ConnectNodeToInput(Connect.UpstreamNode->getInput("amplitude"), NoiseNode, TEXT("Multiply"), 0))
+	if(UInterchangeShaderNode* MultiplyNode = ConnectNodeToInput(Connect.UpstreamNode->getInput("amplitude"), NoiseNode, TEXT("Multiply")))
 	{
 		NodeToConnect = MultiplyNode;
 	}
 
-
-	if(UInterchangeShaderNode* AddNode = ConnectNodeToInput(Connect.UpstreamNode->getInput("pivot"), NodeToConnect, TEXT("Add"), 1))
+	if(UInterchangeShaderNode* AddNode = ConnectNodeToInput(Connect.UpstreamNode->getInput("pivot"), NodeToConnect, TEXT("Add")))
 	{
 		NodeToConnect = AddNode;
 	}
@@ -1226,9 +1239,12 @@ void FMaterialXSurfaceShaderAbstract::ConnectSwitchInputToOutput(const FConnectN
 		AddAttributeFromValueOrInterface(Input, Switch::Inputs::Default.ToString(), SwitchNode);
 
 		//Let's add a new input that is a copy of in1 to connect it to the Default input of the MaterialExpressionSwitch
-		mx::InputPtr InputDefault = Connect.UpstreamNode->addInput("default");
-		InputDefault->copyContentFrom(Input);
-		SetAttributeNewName(InputDefault, TCHAR_TO_UTF8(*Switch::Inputs::Default.ToString()));
+		if(!Connect.UpstreamNode->getInput("default"))
+		{
+			mx::InputPtr InputDefault = Connect.UpstreamNode->addInput("default");
+			InputDefault->copyContentFrom(Input);
+			SetAttributeNewName(InputDefault, TCHAR_TO_UTF8(*Switch::Inputs::Default.ToString()));
+		}
 	}
 
 	SwitchNode->AddInt32Attribute(Switch::Attributes::InputCount.ToString(), Index);
@@ -1285,6 +1301,28 @@ void FMaterialXSurfaceShaderAbstract::ConnectNormalMapInputToOutput(const FConne
 	}
 }
 
+void FMaterialXSurfaceShaderAbstract::ConnectRefractInputToOutput(const FConnectNode& Connect)
+{
+	UInterchangeShaderNode* RefractNode = CreateFunctionCallShaderNode(Connect.UpstreamNode, Connect.UpstreamNode->getName().c_str(), TEXT("/Engine/Functions/Engine_MaterialFunctions01/Vectors/Refract.Refract"));
+	UInterchangeShaderPortsAPI::ConnectDefaultOuputToInput(Connect.ParentShaderNode, Connect.InputChannelName, RefractNode->GetUniqueID());
+
+	for (mx::InputPtr Input : Connect.UpstreamNode->getInputs())
+	{
+		AddAttributeFromValueOrInterface(Input, GetInputName(Input), RefractNode);
+	}
+
+	//Let's add a new input to map its value to the Refractive Index Target since its default value is 1.33 and we want it be 1
+	// Because the Material Function does the operation IOR_SRC/IOR_TARGET
+	mx::InputPtr InputIorTarget = Connect.UpstreamNode->getInput("ior_target");
+	if(!InputIorTarget)
+	{
+		InputIorTarget = Connect.UpstreamNode->addInput("ior_target");
+		InputIorTarget->setValue("1", "float");
+		SetAttributeNewName(InputIorTarget, "Refractive Index Target");
+	}
+	AddAttributeFromValueOrInterface(InputIorTarget, TEXT("Refractive Index Target"), RefractNode);
+}
+
 UInterchangeShaderNode* FMaterialXSurfaceShaderAbstract::CreateMaskShaderNode(uint8 RGBA, mx::ElementPtr Element, const FString& NodeName, const FString& OutputName)
 {
 	bool bR = (0b1000 & RGBA) >> 3;
@@ -1447,24 +1485,14 @@ FLinearColor FMaterialXSurfaceShaderAbstract::GetVector(MaterialX::InputPtr Inpu
 	return LinearColor;
 }
 
-FString FMaterialXSurfaceShaderAbstract::GetAttributeParentName(MaterialX::NodePtr Node) const
+FString FMaterialXSurfaceShaderAbstract::GetAttributeParentName(MaterialX::NodePtr Node, MaterialX::InputPtr ConnectedInput) const
 {
-	FString ParentName;
-	const mx::StringVec& Attributes = Node->getAttributeNames();
-
-	// For consistency the parent attribute has an index attach to it to ensure unicity
-	// Attributes are set in order, we only need to take the first one (it will be remove after that)
-	for(const std::string& Attrib : Attributes)
+	if (ConnectedInput && ConnectedInput->hasAttribute(mx::Attributes::ParentName))
 	{
-		if(Attrib.find(mx::Attributes::ParentName) != std::string::npos)
-		{
-			ParentName = Node->getAttribute(Attrib).c_str();
-			Node->removeAttribute(Attrib);
-			break;
-		}
+		return ConnectedInput->getAttribute(mx::Attributes::ParentName).c_str();
 	}
 
-	return ParentName.IsEmpty() ? Node->getName().c_str() : ParentName;
+	return Node->getName().c_str();
 }
 
 namespace
@@ -1533,6 +1561,7 @@ void FMaterialXSurfaceShaderAbstract::RegisterConnectNodeOutputToInputDelegates(
 	MatchingConnectNodeDelegates.Add(mx::Category::TexCoord,		FOnConnectNodeOutputToInput::CreateSP(this, &FMaterialXSurfaceShaderAbstract::ConnectTexCoordInputToOutput));
 	MatchingConnectNodeDelegates.Add(mx::Category::Swizzle,			FOnConnectNodeOutputToInput::CreateSP(this, &FMaterialXSurfaceShaderAbstract::ConnectSwizzleInputToOutput));
 	MatchingConnectNodeDelegates.Add(mx::Category::Switch,			FOnConnectNodeOutputToInput::CreateSP(this, &FMaterialXSurfaceShaderAbstract::ConnectSwitchInputToOutput));
+	MatchingConnectNodeDelegates.Add(mx::Category::Refract,			FOnConnectNodeOutputToInput::CreateSP(this, &FMaterialXSurfaceShaderAbstract::ConnectRefractInputToOutput));
 }
 
 void FMaterialXSurfaceShaderAbstract::SetMatchingInputsNames(MaterialX::NodePtr Node) const
diff --git a/Engine/Plugins/Interchange/Runtime/Source/Import/Public/MaterialX/MaterialXUtils/MaterialXSurfaceShaderAbstract.h b/Engine/Plugins/Interchange/Runtime/Source/Import/Public/MaterialX/MaterialXUtils/MaterialXSurfaceShaderAbstract.h
--- a/Engine/Plugins/Interchange/Runtime/Source/Import/Public/MaterialX/MaterialXUtils/MaterialXSurfaceShaderAbstract.h
+++ b/Engine/Plugins/Interchange/Runtime/Source/Import/Public/MaterialX/MaterialXUtils/MaterialXSurfaceShaderAbstract.h
@@ -289,6 +289,9 @@ protected:
 	/** <normalmap> */
 	void ConnectNormalMapInputToOutput(const FConnectNode& Connect);
 
+	/** <refract> */
+	void ConnectRefractInputToOutput(const FConnectNode& Connect);
+
 	/** End Connect MaterialX nodes*/
 
 	/**
@@ -451,9 +454,10 @@ protected:
 	 * Retrieve the Interchange parent name of a MaterialX node. Useful when a node is a combination of several nodes connected to different inputs, such as Noise3D.
 	 *
 	 * @param Node - The node whose parent name is retrieved.
+	 * @param ConnectedInput - The connected input to the Node, may be null, may happen at the top of a graph when there's upstream/downstream connection
 	 * @return The node of the parent.
 	 */
-	FString GetAttributeParentName(MaterialX::NodePtr Node) const;
+	FString GetAttributeParentName(MaterialX::NodePtr Node, MaterialX::InputPtr ConnectedInput) const;
 
 	/**
 	 * Ensure that we put in the node container the unique name, 2 nodes in a same file may have the same name as long as their parent's name is different, we just traverse the whole hierarchy to have a unique name
