Interchange: Fix issue when a texture fail import and it is reference by a material that is imported.

#rnx

diff --git a/Engine/Plugins/Interchange/Runtime/Source/Import/Private/Texture/InterchangeTextureFactory.cpp b/Engine/Plugins/Interchange/Runtime/Source/Import/Private/Texture/InterchangeTextureFactory.cpp
--- a/Engine/Plugins/Interchange/Runtime/Source/Import/Private/Texture/InterchangeTextureFactory.cpp
+++ b/Engine/Plugins/Interchange/Runtime/Source/Import/Private/Texture/InterchangeTextureFactory.cpp
@@ -2062,7 +2062,9 @@ UInterchangeFactoryBase::FImportAssetResult UInterchangeTextureFactory::ImportAs
 	if (!bCanSetup)
 	{
 		LogErrorInvalidPayload(*this, Texture->GetClass(), Arguments.SourceData->GetFilename(), Texture->GetName());
-		ImportAssetResult.ImportedObject = Texture;
+		ImportAssetResult.ImportedObject = Arguments.ReimportObject ? Texture : nullptr;
+		ImportAssetResult.bIsFactorySkipAsset = true;
+		Arguments.AssetNode->SetEnabled(false);
 		return ImportAssetResult;
 	}
 
@@ -2086,6 +2088,56 @@ UInterchangeFactoryBase::FImportAssetResult UInterchangeTextureFactory::ImportAs
 	return ImportAssetResult;
 }
 
+//If the texture has a bad payload and its not a reimport we want to mark for garbage collect the newly created texture
+UInterchangeFactoryBase::FImportAssetResult UInterchangeTextureFactory::EndImportAsset_GameThread(const FImportAssetObjectParams& Arguments)
+{
+	TRACE_CPUPROFILER_EVENT_SCOPE(UInterchangeTextureFactory::ImportAsset_Async);
+
+	using namespace UE::Interchange::Private::InterchangeTextureFactory;
+	FImportAssetResult ImportAssetResult;
+	ImportAssetResult.bIsFactorySkipAsset = bSkipImport;
+
+	//We need a valid node
+	if (!Arguments.AssetNode)
+	{
+		return ImportAssetResult;
+	}
+
+	UObject* ExistingAsset = nullptr;
+	FSoftObjectPath ReferenceObject;
+	if (Arguments.AssetNode->GetCustomReferenceObject(ReferenceObject))
+	{
+		ExistingAsset = ReferenceObject.TryLoad();
+	}
+
+	//We need a valid imported object
+	if (!ExistingAsset)
+	{
+		return ImportAssetResult;
+	}
+
+	if (UTexture* Texture = Cast<UTexture>(ExistingAsset))
+	{
+		ImportAssetResult.ImportedObject = Texture;
+		if (!ProcessedPayload.IsValid())
+		{
+			// The texture is not supported
+			if (!Arguments.ReimportObject)
+			{
+				ImportAssetResult.ImportedObject = nullptr;
+				// Not thread safe. So those should stay on the game thread.
+				Texture->RemoveFromRoot();
+				Texture->MarkAsGarbage();
+			}
+			Arguments.AssetNode->SetEnabled(false);
+			bSkipImport = true;
+			ImportAssetResult.bIsFactorySkipAsset = bSkipImport;
+		}
+	}
+	
+	return ImportAssetResult;
+}
+
 /* This function is call in the completion task on the main thread, use it to call main thread post creation step for your assets*/
 void UInterchangeTextureFactory::SetupObject_GameThread(const FSetupObjectParams& Arguments)
 {
@@ -2160,7 +2212,7 @@ void UInterchangeTextureFactory::SetupObject_GameThread(const FSetupObjectParams
 
 	//TODO make sure this work at runtime
 #if WITH_EDITORONLY_DATA
-	if (ensure(Texture && Arguments.SourceData) && ProcessedPayload.IsValid())
+	if (Texture && Arguments.SourceData && ProcessedPayload.IsValid())
 	{
 		//We must call the Update of the asset source file in the main thread because UAssetImportData::Update execute some delegate we do not control
 		UE::Interchange::FFactoryCommon::FSetImportAssetDataParameters SetImportAssetDataParameters(Texture
diff --git a/Engine/Plugins/Interchange/Runtime/Source/Import/Public/Texture/InterchangeTextureFactory.h b/Engine/Plugins/Interchange/Runtime/Source/Import/Public/Texture/InterchangeTextureFactory.h
--- a/Engine/Plugins/Interchange/Runtime/Source/Import/Public/Texture/InterchangeTextureFactory.h
+++ b/Engine/Plugins/Interchange/Runtime/Source/Import/Public/Texture/InterchangeTextureFactory.h
@@ -58,6 +58,7 @@ public:
 	virtual EInterchangeFactoryAssetType GetFactoryAssetType() override { return EInterchangeFactoryAssetType::Textures; }
 	virtual FImportAssetResult BeginImportAsset_GameThread(const FImportAssetObjectParams& Arguments) override;
 	virtual FImportAssetResult ImportAsset_Async(const FImportAssetObjectParams& Arguments) override;
+	virtual FImportAssetResult EndImportAsset_GameThread(const FImportAssetObjectParams& Arguments) override;
 	virtual void SetupObject_GameThread(const FSetupObjectParams& Arguments) override;
 	virtual bool GetSourceFilenames(const UObject* Object, TArray<FString>& OutSourceFilenames) const override;
 	virtual bool SetSourceFilename(const UObject* Object, const FString& SourceFilename, int32 SourceIndex) const override;
