Interchange: fix animation length options not working anymore
#rnx

diff --git a/Engine/Plugins/Interchange/Runtime/Source/Parsers/Fbx/Private/FbxAPI.cpp b/Engine/Plugins/Interchange/Runtime/Source/Parsers/Fbx/Private/FbxAPI.cpp
--- a/Engine/Plugins/Interchange/Runtime/Source/Parsers/Fbx/Private/FbxAPI.cpp
+++ b/Engine/Plugins/Interchange/Runtime/Source/Parsers/Fbx/Private/FbxAPI.cpp
@@ -148,8 +148,36 @@ namespace UE
 				UInterchangeSourceNode* SourceNode = UInterchangeSourceNode::FindOrCreateUniqueInstance(&NodeContainer);
 				SourceNode->SetCustomAxisConversionInverseTransform(AxisConversionInverseTransform);
 
-				FrameRate = FbxTime::GetFrameRate(SDKScene->GetGlobalSettings().GetTimeMode());
-				FileDetails.FrameRate = FString::Printf(TEXT("%.2f"), FrameRate);
+				//Store the fbx frame rate and timeline
+				{
+					FrameRate = FbxTime::GetFrameRate(SDKScene->GetGlobalSettings().GetTimeMode());
+					FileDetails.FrameRate = FString::Printf(TEXT("%.2f"), FrameRate);
+					SourceNode->SetCustomSourceFrameRateNumerator(FrameRate);
+					constexpr double Denominator = 1.0;
+					SourceNode->SetCustomSourceFrameRateDenominator(Denominator);
+
+					//Timeline time span
+					FbxTimeSpan TimelineTimeSpan(FBXSDK_TIME_INFINITE, FBXSDK_TIME_MINUS_INFINITE);
+					SDKScene->GetGlobalSettings().GetTimelineDefaultTimeSpan(TimelineTimeSpan);
+					SourceNode->SetCustomSourceTimelineStart(TimelineTimeSpan.GetStart().GetSecondDouble());
+					SourceNode->SetCustomSourceTimelineEnd(TimelineTimeSpan.GetStop().GetSecondDouble());
+
+					//Animated time span
+					FbxTimeSpan AnimatedTimeSpan(FBXSDK_TIME_INFINITE, FBXSDK_TIME_MINUS_INFINITE);
+					int32 AnimCurveNodeCount = SDKScene->GetSrcObjectCount<FbxAnimCurveNode>();
+					for (int32 AnimCurveNodeIndex = 0; AnimCurveNodeIndex < AnimCurveNodeCount; AnimCurveNodeIndex++)
+					{
+						FbxAnimCurveNode* CurAnimCruveNode = SDKScene->GetSrcObject<FbxAnimCurveNode>(AnimCurveNodeIndex);
+						if (CurAnimCruveNode->IsAnimated(true))
+						{
+							FbxTimeSpan CurveTimeSpan(FBXSDK_TIME_INFINITE, FBXSDK_TIME_MINUS_INFINITE);
+							CurAnimCruveNode->GetAnimationInterval(CurveTimeSpan);
+							AnimatedTimeSpan.UnionAssignment(CurveTimeSpan);
+						}
+					}
+					SourceNode->SetCustomAnimatedTimeStart(AnimatedTimeSpan.GetStart().GetSecondDouble());
+					SourceNode->SetCustomAnimatedTimeEnd(AnimatedTimeSpan.GetStop().GetSecondDouble());
+				}
 
 				// Get the version number of the FBX file format.
 				int32 FileMajor, FileMinor, FileRevision;
diff --git a/Engine/Plugins/Interchange/Runtime/Source/Pipelines/Private/InterchangeGenericAnimationPipeline.cpp b/Engine/Plugins/Interchange/Runtime/Source/Pipelines/Private/InterchangeGenericAnimationPipeline.cpp
--- a/Engine/Plugins/Interchange/Runtime/Source/Pipelines/Private/InterchangeGenericAnimationPipeline.cpp
+++ b/Engine/Plugins/Interchange/Runtime/Source/Pipelines/Private/InterchangeGenericAnimationPipeline.cpp
@@ -570,53 +570,47 @@ void UInterchangeGenericAnimationPipeline::CreateAnimSequenceFactoryNode(UInterc
 	{
 		int32 Numerator, Denominator;
 		const UInterchangeSourceNode* SourceNode = UInterchangeSourceNode::GetUniqueInstance(BaseNodeContainer);
-		if (SourceNode && SourceNode->GetCustomSourceFrameRateNumerator(Numerator) && SourceNode->GetCustomSourceFrameRateDenominator(Denominator))
-		{
-			if (!bUse30HzToBakeBoneAnimation && CustomBoneAnimationSampleRate == 0)
-			{
-				if (Denominator > 0 && Numerator > 0)
-				{
-					SampleRate = static_cast<double>(Numerator) / static_cast<double>(Denominator);
-				}
-			}
 
-			if (AnimationRange == EInterchangeAnimationRange::Timeline)
+		//Get the sample rate from the options and the data
+		if (!bUse30HzToBakeBoneAnimation)
+		{
+			if (CustomBoneAnimationSampleRate > 0)
 			{
-				if (SourceNode->GetCustomSourceTimelineStart(StartTime))
-				{
-					if (SourceNode->GetCustomSourceTimelineEnd(StopTime))
-					{
-						bTimeRangeIsValid = true;
-					}
-				}
+				SampleRate = static_cast<double>(CustomBoneAnimationSampleRate);
 			}
-			else if (AnimationRange == EInterchangeAnimationRange::Animated)
+			else if (SourceNode && SourceNode->GetCustomSourceFrameRateNumerator(Numerator) && SourceNode->GetCustomSourceFrameRateDenominator(Denominator) && Denominator > 0 && Numerator > 0)
 			{
-				if (SourceNode->GetCustomAnimatedTimeStart(StartTime))
-				{
-					if (SourceNode->GetCustomAnimatedTimeEnd(StopTime))
-					{
-						bTimeRangeIsValid = true;
-					}
-				}
+				SampleRate = static_cast<double>(Numerator) / static_cast<double>(Denominator);
 			}
-			else if (AnimationRange == EInterchangeAnimationRange::SetRange)
+			else
 			{
-				StartTime = static_cast<double>(FrameImportRange.Min) / SampleRate;
-				StopTime = static_cast<double>(FrameImportRange.Max) / SampleRate;
-				bTimeRangeIsValid = true;
+				TrackNode.GetCustomAnimationSampleRate(SampleRate);
 			}
 		}
-		else
+
+		//Get the animation start and stop time (range) from the options and the data
+		//Some format don't fill the animation range data we fall back on the trackNode
+		if (AnimationRange == EInterchangeAnimationRange::Timeline)
 		{
-			bTimeRangeIsValid = TrackNode.GetCustomAnimationSampleRate(SampleRate);
-			bTimeRangeIsValid &= TrackNode.GetCustomAnimationStartTime(StartTime);
-			bTimeRangeIsValid &= TrackNode.GetCustomAnimationStopTime(StopTime);
+			bTimeRangeIsValid = (SourceNode && SourceNode->GetCustomSourceTimelineStart(StartTime) && SourceNode->GetCustomSourceTimelineEnd(StopTime));
 		}
-
-		if (!bUse30HzToBakeBoneAnimation && CustomBoneAnimationSampleRate > 0)
+		else if (AnimationRange == EInterchangeAnimationRange::Animated)
+		{
+			bTimeRangeIsValid = (SourceNode &&SourceNode->GetCustomAnimatedTimeStart(StartTime) && SourceNode->GetCustomAnimatedTimeEnd(StopTime));
+		}
+		else if (AnimationRange == EInterchangeAnimationRange::SetRange)
 		{
-			SampleRate = static_cast<double>(CustomBoneAnimationSampleRate);
+			StartTime = static_cast<double>(FrameImportRange.Min) / SampleRate;
+			StopTime = static_cast<double>(FrameImportRange.Max) / SampleRate;
+			bTimeRangeIsValid = true;
+		}
+		//If we cannot get the animation range use the node track data
+		if(!bTimeRangeIsValid)
+		{
+			StartTime = 0.;
+			StopTime = 0.;
+			bTimeRangeIsValid = TrackNode.GetCustomAnimationStartTime(StartTime);
+			bTimeRangeIsValid &= TrackNode.GetCustomAnimationStopTime(StopTime);
 		}
 
 		FFrameRate FrameRate = UE::Interchange::Animation::ConvertSampleRatetoFrameRate(SampleRate);
