Unshelved from pending changelist '38096170':

[Interchange] Fix the remove redundant key when importing curves, it doesn't add the first key.

#rnx

diff --git a/Engine/Plugins/Interchange/Runtime/Source/Import/Private/Animation/InterchangeAnimSequenceFactory.cpp b/Engine/Plugins/Interchange/Runtime/Source/Import/Private/Animation/InterchangeAnimSequenceFactory.cpp
--- a/Engine/Plugins/Interchange/Runtime/Source/Import/Private/Animation/InterchangeAnimSequenceFactory.cpp
+++ b/Engine/Plugins/Interchange/Runtime/Source/Import/Private/Animation/InterchangeAnimSequenceFactory.cpp
@@ -48,7 +48,7 @@ namespace UE::Interchange::Private
 	bool AreAllValuesZero(const TArray<ValueType>& Values, TFunctionRef<bool(const ValueType& Value)> CompareValueToZeroFunction)
 	{
 		const int32 KeyCount = Values.Num();
-		for (int32 KeyIndex = 1; KeyIndex < KeyCount; ++KeyIndex)
+		for (int32 KeyIndex = 0; KeyIndex < KeyCount; ++KeyIndex)
 		{
 			//Only add the not equal keys
 			if (!CompareValueToZeroFunction(Values[KeyIndex]))
diff --git a/Engine/Plugins/Interchange/Runtime/Source/Parsers/CommonParser/Private/InterchangeCommonAnimationPayload.cpp b/Engine/Plugins/Interchange/Runtime/Source/Parsers/CommonParser/Private/InterchangeCommonAnimationPayload.cpp
--- a/Engine/Plugins/Interchange/Runtime/Source/Parsers/CommonParser/Private/InterchangeCommonAnimationPayload.cpp
+++ b/Engine/Plugins/Interchange/Runtime/Source/Parsers/CommonParser/Private/InterchangeCommonAnimationPayload.cpp
@@ -399,6 +399,9 @@ void FInterchangeStepCurve::InternalRemoveRedundantKey(TArray<ValueType>& Values
 	TArray<ValueType> NewValues;
 	NewValues.Reserve(KeyCount);
 	ValueType LastValue = Values[0];
+	//Add the first key
+	NewKeyTimes.Add(KeyTimes[0]);
+	NewValues.Add(Values[0]);
 	for (int32 KeyIndex = 1; KeyIndex < KeyCount; ++KeyIndex)
 	{
 		//Only add the not equal keys
