[Interchange] fix animation name to be same as legacy fbx import when there is multiple animations in the same file.
Make sure the "Animated Time" option in fbx use the track node length which is the expected behavior.

diff --git a/Engine/Plugins/Interchange/Runtime/Source/Parsers/Fbx/Private/FbxAPI.cpp b/Engine/Plugins/Interchange/Runtime/Source/Parsers/Fbx/Private/FbxAPI.cpp
--- a/Engine/Plugins/Interchange/Runtime/Source/Parsers/Fbx/Private/FbxAPI.cpp
+++ b/Engine/Plugins/Interchange/Runtime/Source/Parsers/Fbx/Private/FbxAPI.cpp
@@ -161,22 +161,6 @@ namespace UE
 					SDKScene->GetGlobalSettings().GetTimelineDefaultTimeSpan(TimelineTimeSpan);
 					SourceNode->SetCustomSourceTimelineStart(TimelineTimeSpan.GetStart().GetSecondDouble());
 					SourceNode->SetCustomSourceTimelineEnd(TimelineTimeSpan.GetStop().GetSecondDouble());
-
-					//Animated time span
-					FbxTimeSpan AnimatedTimeSpan(FBXSDK_TIME_INFINITE, FBXSDK_TIME_MINUS_INFINITE);
-					int32 AnimCurveNodeCount = SDKScene->GetSrcObjectCount<FbxAnimCurveNode>();
-					for (int32 AnimCurveNodeIndex = 0; AnimCurveNodeIndex < AnimCurveNodeCount; AnimCurveNodeIndex++)
-					{
-						FbxAnimCurveNode* CurAnimCruveNode = SDKScene->GetSrcObject<FbxAnimCurveNode>(AnimCurveNodeIndex);
-						if (CurAnimCruveNode->IsAnimated(true))
-						{
-							FbxTimeSpan CurveTimeSpan(FBXSDK_TIME_INFINITE, FBXSDK_TIME_MINUS_INFINITE);
-							CurAnimCruveNode->GetAnimationInterval(CurveTimeSpan);
-							AnimatedTimeSpan.UnionAssignment(CurveTimeSpan);
-						}
-					}
-					SourceNode->SetCustomAnimatedTimeStart(AnimatedTimeSpan.GetStart().GetSecondDouble());
-					SourceNode->SetCustomAnimatedTimeEnd(AnimatedTimeSpan.GetStop().GetSecondDouble());
 				}
 
 				// Get the version number of the FBX file format.
diff --git a/Engine/Plugins/Interchange/Runtime/Source/Parsers/Fbx/Private/FbxMesh.cpp b/Engine/Plugins/Interchange/Runtime/Source/Parsers/Fbx/Private/FbxMesh.cpp
--- a/Engine/Plugins/Interchange/Runtime/Source/Parsers/Fbx/Private/FbxMesh.cpp
+++ b/Engine/Plugins/Interchange/Runtime/Source/Parsers/Fbx/Private/FbxMesh.cpp
@@ -2224,6 +2224,7 @@ UInterchangeMeshNode* FFbxMesh::CreateMeshNode(UInterchangeBaseNodeContainer& No
 	// Creating a UMaterialInterface
 	MeshNode->InitializeNode(NodeUid, DisplayLabel, EInterchangeNodeContainerType::TranslatedAsset);
 	NodeContainer.AddNode(MeshNode);
+	MeshNode->AddBooleanAttribute(TEXT("RenameLikeLegacyFbx"), true);
 	return MeshNode;
 }
 
diff --git a/Engine/Plugins/Interchange/Runtime/Source/Parsers/Fbx/Private/FbxScene.cpp b/Engine/Plugins/Interchange/Runtime/Source/Parsers/Fbx/Private/FbxScene.cpp
--- a/Engine/Plugins/Interchange/Runtime/Source/Parsers/Fbx/Private/FbxScene.cpp
+++ b/Engine/Plugins/Interchange/Runtime/Source/Parsers/Fbx/Private/FbxScene.cpp
@@ -21,6 +21,7 @@
 #include "InterchangeAnimationTrackSetNode.h"
 #include "InterchangeAnimationDefinitions.h"
 #include "InterchangeFbxSettings.h"
+#include "InterchangeHelper.h"
 
 #define LOCTEXT_NAMESPACE "InterchangeFbxScene"
 
@@ -840,15 +841,19 @@ namespace UE
 				if (UInterchangeSceneNode* UnrealNode = const_cast<UInterchangeSceneNode*>(Cast< UInterchangeSceneNode >(NodeContainer.GetNode(NodeUniqueID))))
 				{
 					bool HasSkeletonAttribute = false;
-					auto ApplySkeletonAttribute = [&SDKScene, &HasSkeletonAttribute, &SkeletonRootNodeUids, &NodeUniqueID, &SkeletalAnimationTrackNode, &UnrealNode, &AnimationIndex, &NodeContainer]()
+					auto ApplySkeletonAttribute = [&SDKScene, &HasSkeletonAttribute, &SkeletonRootNodeUids, &NodeUniqueID, &SkeletalAnimationTrackNode, &UnrealNode, &AnimationIndex, &NodeContainer, &Node]()
 					{
 						HasSkeletonAttribute = true;
 						if (SkeletonRootNodeUids.Contains(NodeUniqueID))
 						{
+							FbxAnimStack* CurrentAnimationStack = (FbxAnimStack*)SDKScene->GetSrcObject<FbxAnimStack>(AnimationIndex);
+							FString DisplayString = FFbxConvert::MakeString(CurrentAnimationStack->GetName());
+							constexpr bool bIsJointFalse = false;
+							UE::Interchange::SanitizeName(DisplayString, bIsJointFalse);
 							SkeletalAnimationTrackNode = NewObject< UInterchangeSkeletalAnimationTrackNode >(&NodeContainer);
 							FString TrackNodeUid = "\\SkeletalAnimation\\" + UnrealNode->GetUniqueID() + "_" + FString::FromInt(AnimationIndex);
-							FString DisplayString = "Anim_" + FString::FromInt(AnimationIndex) + "_" + UnrealNode->GetDisplayLabel();
 							SkeletalAnimationTrackNode->InitializeNode(TrackNodeUid, DisplayString, EInterchangeNodeContainerType::TranslatedAsset);
+							SkeletalAnimationTrackNode->AddBooleanAttribute(TEXT("RenameLikeLegacyFbx"), true);
 
 							double FrameRate = FbxTime::GetFrameRate(SDKScene->GetGlobalSettings().GetTimeMode());
 							SkeletalAnimationTrackNode->SetCustomAnimationSampleRate(FrameRate);
@@ -856,11 +861,11 @@ namespace UE
 							SkeletalAnimationTrackNode->SetCustomSkeletonNodeUid(UnrealNode->GetUniqueID());
 
 							//Calculate AnimationTime:
-							FbxAnimStack* CurrentAnimationStack = (FbxAnimStack*)SDKScene->GetSrcObject<FbxAnimStack>(AnimationIndex);
-							FbxTimeSpan TimeSpan = CurrentAnimationStack->GetLocalTimeSpan();
+							FbxTimeSpan AnimatedInterval(FBXSDK_TIME_INFINITE, FBXSDK_TIME_MINUS_INFINITE);
+							Node->GetAnimationInterval(AnimatedInterval, CurrentAnimationStack);
 
-							SkeletalAnimationTrackNode->SetCustomAnimationStartTime(TimeSpan.GetStart().GetSecondDouble());
-							SkeletalAnimationTrackNode->SetCustomAnimationStopTime(TimeSpan.GetStop().GetSecondDouble());
+							SkeletalAnimationTrackNode->SetCustomAnimationStartTime(AnimatedInterval.GetStart().GetSecondDouble());
+							SkeletalAnimationTrackNode->SetCustomAnimationStopTime(AnimatedInterval.GetStop().GetSecondDouble());
 
 							return true;
 						}
diff --git a/Engine/Plugins/Interchange/Runtime/Source/Pipelines/Private/InterchangeGenericAnimationPipeline.cpp b/Engine/Plugins/Interchange/Runtime/Source/Pipelines/Private/InterchangeGenericAnimationPipeline.cpp
--- a/Engine/Plugins/Interchange/Runtime/Source/Pipelines/Private/InterchangeGenericAnimationPipeline.cpp
+++ b/Engine/Plugins/Interchange/Runtime/Source/Pipelines/Private/InterchangeGenericAnimationPipeline.cpp
@@ -6,10 +6,12 @@
 #include "CoreMinimal.h"
 #include "Engine/SkeletalMesh.h"
 #include "Engine/StaticMesh.h"
+#include "InterchangeGenericAssetsPipeline.h"
 #include "InterchangeGenericMeshPipeline.h"
 #include "InterchangeLevelSequenceFactoryNode.h"
 #include "InterchangeAnimationTrackSetNode.h"
 #include "InterchangeAnimSequenceFactoryNode.h"
+#include "InterchangeHelper.h"
 #include "InterchangeMeshNode.h"
 #include "InterchangePipelineLog.h"
 #include "InterchangeSceneNode.h"
@@ -69,6 +71,71 @@ namespace UE::Interchange::Private
 
 		return bContainOnlyJointAnimation;
 	}
+
+	//Legacy fbx has a very particular way for naming animation sequences.
+	void RenameAnimSequenceLikeLegacyFbx(FString& DisplayString, UInterchangeSkeletalAnimationTrackNode& TrackNode, const TArray<const UInterchangeSourceData*>& SourceDatas, UInterchangePipelineBase* OuterMostPipeline, UInterchangeBaseNodeContainer* NodeContainer)
+	{
+		if(!SourceDatas.IsValidIndex(0))
+		{
+			return;
+		}
+		bool bRenameAnimSequenceLegacyFbxWay = false;
+		TrackNode.GetBooleanAttribute(TEXT("RenameLikeLegacyFbx"), bRenameAnimSequenceLegacyFbxWay);
+		if (!bRenameAnimSequenceLegacyFbxWay)
+		{
+			return;
+		}
+
+		FString SkeletonNodeUid;
+		TrackNode.GetCustomSkeletonNodeUid(SkeletonNodeUid);
+		FString BaseFilename = FPaths::GetBaseFilename(SourceDatas[0]->GetFilename());
+		FString MeshName;
+		if (UInterchangeGenericAssetsPipeline* GenericAssetPipeline = Cast<UInterchangeGenericAssetsPipeline>(OuterMostPipeline))
+		{
+			if (!GenericAssetPipeline->bUseSourceNameForAsset)
+			{
+				NodeContainer->BreakableIterateNodesOfType<UInterchangeMeshNode>([&SkeletonNodeUid, &MeshName, &NodeContainer](const FString& NodeUid, UInterchangeMeshNode* MeshNode)
+					{
+						if (MeshNode)
+						{
+							TArray<FString> SkeletonDependencies;
+							MeshNode->GetSkeletonDependencies(SkeletonDependencies);
+							if(SkeletonDependencies.IsValidIndex(0))
+							{
+								if(SkeletonDependencies.Contains(SkeletonNodeUid) || NodeContainer->GetIsAncestor(SkeletonDependencies[0], SkeletonNodeUid))
+								{
+									int32 MeshReferenceCount = MeshNode->GetSceneInstanceUidsCount();
+									if (MeshReferenceCount == 1)
+									{
+										FString SceneNodeUid;
+										MeshNode->GetSceneInstanceUid(0, SceneNodeUid);
+										if (const UInterchangeSceneNode* SceneNode = Cast<UInterchangeSceneNode>(NodeContainer->GetNode(SceneNodeUid)))
+										{
+											MeshName = TEXT("_") + SceneNode->GetDisplayLabel();
+											return true;
+										}
+									}
+									MeshName = TEXT("_") + MeshNode->GetDisplayLabel();
+									return true;
+								}
+							}
+						}
+						return false;
+					});
+			}
+		}
+
+		DisplayString = BaseFilename + MeshName + TEXT("_Anim");
+
+		int32 AnimationCount = 0;
+		NodeContainer->BreakableIterateNodesOfType<UInterchangeSkeletalAnimationTrackNode>([&AnimationCount](const FString&, UInterchangeSkeletalAnimationTrackNode*) { return (AnimationCount++ > 0); });
+		if (AnimationCount > 1)
+		{
+			DisplayString += TEXT("_") + TrackNode.GetDisplayLabel();
+		}
+		constexpr bool bIsJointFalse = false;
+		UE::Interchange::SanitizeName(DisplayString, bIsJointFalse);
+	}
 }
 
 FString UInterchangeGenericAnimationPipeline::GetPipelineCategory(UClass* AssetClass)
@@ -594,17 +661,13 @@ void UInterchangeGenericAnimationPipeline::CreateAnimSequenceFactoryNode(UInterc
 		{
 			bTimeRangeIsValid = (SourceNode && SourceNode->GetCustomSourceTimelineStart(StartTime) && SourceNode->GetCustomSourceTimelineEnd(StopTime));
 		}
-		else if (AnimationRange == EInterchangeAnimationRange::Animated)
-		{
-			bTimeRangeIsValid = (SourceNode &&SourceNode->GetCustomAnimatedTimeStart(StartTime) && SourceNode->GetCustomAnimatedTimeEnd(StopTime));
-		}
 		else if (AnimationRange == EInterchangeAnimationRange::SetRange)
 		{
 			StartTime = static_cast<double>(FrameImportRange.Min) / SampleRate;
 			StopTime = static_cast<double>(FrameImportRange.Max) / SampleRate;
 			bTimeRangeIsValid = true;
 		}
-		//If we cannot get the animation range use the node track data
+		//No custom time specified, use the track default
 		if(!bTimeRangeIsValid)
 		{
 			StartTime = 0.;
@@ -676,7 +739,11 @@ void UInterchangeGenericAnimationPipeline::CreateAnimSequenceFactoryNode(UInterc
 	const FString AnimSequenceUid = TEXT("\\AnimSequence") + TrackNode.GetUniqueID();
 
 	UInterchangeAnimSequenceFactoryNode* AnimSequenceFactoryNode = NewObject<UInterchangeAnimSequenceFactoryNode>(BaseNodeContainer, NAME_None);
-	AnimSequenceFactoryNode->InitializeAnimSequenceNode(AnimSequenceUid, TrackNode.GetDisplayLabel());
+
+	FString DisplayString = TrackNode.GetDisplayLabel();
+	UE::Interchange::Private::RenameAnimSequenceLikeLegacyFbx(DisplayString, TrackNode, SourceDatas, GetMostPipelineOuter(), BaseNodeContainer);
+
+	AnimSequenceFactoryNode->InitializeAnimSequenceNode(AnimSequenceUid, DisplayString);
 
 	if (SkeletonFactoryNode)
 	{
diff --git a/Engine/Plugins/Interchange/Runtime/Source/Pipelines/Private/InterchangeGenericSkeletalMeshPipeline.cpp b/Engine/Plugins/Interchange/Runtime/Source/Pipelines/Private/InterchangeGenericSkeletalMeshPipeline.cpp
--- a/Engine/Plugins/Interchange/Runtime/Source/Pipelines/Private/InterchangeGenericSkeletalMeshPipeline.cpp
+++ b/Engine/Plugins/Interchange/Runtime/Source/Pipelines/Private/InterchangeGenericSkeletalMeshPipeline.cpp
@@ -69,6 +69,31 @@ namespace UE::Interchange::SkeletalMeshGenericPipeline
 			}
 		}
 	}
+
+	void RenameSkeletalMeshLikeLegacyFbx(FString& DisplayLabel, const UInterchangeMeshNode* MeshNode, const TArray<const UInterchangeSourceData*>& SourceDatas, UInterchangePipelineBase* OuterMostPipeline)
+	{
+		if(!MeshNode || !SourceDatas.IsValidIndex(0))
+		{
+			return;
+		}
+		bool bRenameLikeLegacyFbx = false;
+		MeshNode->GetBooleanAttribute(TEXT("RenameLikeLegacyFbx"), bRenameLikeLegacyFbx);
+		if (!bRenameLikeLegacyFbx)
+		{
+			return;
+		}
+
+		FString BaseFilename = FPaths::GetBaseFilename(SourceDatas[0]->GetFilename());
+		FString MeshName;
+		if (UInterchangeGenericAssetsPipeline* GenericAssetPipeline = Cast<UInterchangeGenericAssetsPipeline>(OuterMostPipeline))
+		{
+			if (!GenericAssetPipeline->bUseSourceNameForAsset && !DisplayLabel.IsEmpty())
+			{
+				MeshName = TEXT("_") + DisplayLabel;
+			}
+		}
+		DisplayLabel = BaseFilename + MeshName;
+	}
 }
 
 void UInterchangeGenericMeshPipeline::ExecutePreImportPipelineSkeletalMesh()
@@ -394,14 +419,14 @@ UInterchangeSkeletalMeshFactoryNode* UInterchangeGenericMeshPipeline::CreateSkel
 	FString SkeletalMeshUid_MeshNamePart = FirstMeshNodeUid;
 	if(FirstSceneNode)
 	{
-		//If we are instancing one scene node, we want to use it to name the mesh
-		if (SceneNodeCount == 1)
-		{
-			DisplayLabel = FirstSceneNode->GetDisplayLabel();
-		}
+		//use the scene node to name the skeletal mesh
+		DisplayLabel = FirstSceneNode->GetDisplayLabel();
 		//Use the first scene node uid this skeletalmesh reference, add backslash since this uid is not asset typed (\\Mesh\\) like FirstMeshNodeUid
 		SkeletalMeshUid_MeshNamePart = TEXT("\\") + FirstSceneNode->GetUniqueID();
 	}
+
+	UE::Interchange::SkeletalMeshGenericPipeline::RenameSkeletalMeshLikeLegacyFbx(DisplayLabel, FirstMeshNode, SourceDatas, GetMostPipelineOuter());
+
 	const FString SkeletalMeshUid = UInterchangeFactoryBaseNode::BuildFactoryNodeUid(SkeletalMeshUid_MeshNamePart + SkeletonUid);
 	UInterchangeSkeletalMeshFactoryNode* SkeletalMeshFactoryNode = NewObject<UInterchangeSkeletalMeshFactoryNode>(BaseNodeContainer, NAME_None);
 	if (!ensure(SkeletalMeshFactoryNode))
