Interchange MaterialX:
- Add cellnoise2, worleynoise2d, noise2d, since we can't have a perfect match match with MaterialX spec since UE is using 3D Random generators, an arbitrary seed of 0 is used for the third component, we'll see in the future if it generates too much artifact.
- Fix the output of cellnoise3d, worleynoise3d, noise3d to match the one from aterialX spec
- Avoid adding all the inputs of a nodedefs for all nodes, since we don't have a perfect match on the inputs (for example some expressions define their inputs as Advanced Input pins), discriminate only certains nodes
- Refactor the Manager, to avoid having hard coded string literals, and to be sure that all UE matching inputs are always added in the data structure at the same time (less error-prone)

diff --git a/Engine/Plugins/Interchange/Runtime/Source/Common/Public/MaterialX/InterchangeMaterialXDefinitions.h b/Engine/Plugins/Interchange/Runtime/Source/Common/Public/MaterialX/InterchangeMaterialXDefinitions.h
--- a/Engine/Plugins/Interchange/Runtime/Source/Common/Public/MaterialX/InterchangeMaterialXDefinitions.h
+++ b/Engine/Plugins/Interchange/Runtime/Source/Common/Public/MaterialX/InterchangeMaterialXDefinitions.h
@@ -134,6 +134,148 @@ namespace UE
 			}
 		}
 	}
+
+	namespace Expressions
+	{
+		namespace Names
+		{
+			static constexpr const TCHAR* Abs = TEXT("Abs");
+			static constexpr const TCHAR* Add = TEXT("Add");
+			static constexpr const TCHAR* AppendVector = TEXT("AppendVector");
+			static constexpr const TCHAR* Arccosine = TEXT("Arccosine");
+			static constexpr const TCHAR* Arcsine = TEXT("Arcsine");
+			static constexpr const TCHAR* Arctangent2 = TEXT("Arctangent2");
+			static constexpr const TCHAR* BlackBody = TEXT("BlackBody");
+			static constexpr const TCHAR* Ceil = TEXT("Ceil");
+			static constexpr const TCHAR* Clamp = TEXT("Clamp");
+			static constexpr const TCHAR* Cosine = TEXT("Cosine");
+			static constexpr const TCHAR* Crossproduct = TEXT("Crossproduct");
+			static constexpr const TCHAR* Desaturation = TEXT("Desaturation");
+			static constexpr const TCHAR* Divide = TEXT("Divide");
+			static constexpr const TCHAR* Distance = TEXT("Distance");
+			static constexpr const TCHAR* Dotproduct = TEXT("Dotproduct");
+			static constexpr const TCHAR* Exponential = TEXT("Exponential");
+			static constexpr const TCHAR* Floor = TEXT("Floor");
+			static constexpr const TCHAR* Fmod = TEXT("Fmod");
+			static constexpr const TCHAR* HsvToRgb = TEXT("HsvToRgb");
+			static constexpr const TCHAR* If = TEXT("If");
+			static constexpr const TCHAR* Logarithm = TEXT("Logarithm");
+			static constexpr const TCHAR* Length = TEXT("Length");
+			static constexpr const TCHAR* Lerp = TEXT("Lerp");
+			static constexpr const TCHAR* LocalPosition = TEXT("LocalPosition");
+			static constexpr const TCHAR* MaterialXAppend3Vector = TEXT("MaterialXAppend3Vector");
+			static constexpr const TCHAR* MaterialXAppend4Vector = TEXT("MaterialXAppend4Vector");
+			static constexpr const TCHAR* MaterialXBurn = TEXT("MaterialXBurn");
+			static constexpr const TCHAR* MaterialXDifference = TEXT("MaterialXDifference");
+			static constexpr const TCHAR* MaterialXDisjointover = TEXT("MaterialXDisjointover");
+			static constexpr const TCHAR* MaterialXDodge = TEXT("MaterialXDodge");
+			static constexpr const TCHAR* MaterialXFractal3D = TEXT("MaterialXFractal3D");
+			static constexpr const TCHAR* MaterialXIn = TEXT("MaterialXIn");
+			static constexpr const TCHAR* MaterialXLuminance = TEXT("MaterialXLuminance");
+			static constexpr const TCHAR* MaterialXMask = TEXT("MaterialXMask");
+			static constexpr const TCHAR* MaterialXMatte = TEXT("MaterialXMatte");
+			static constexpr const TCHAR* MaterialXMinus = TEXT("MaterialXMinus");
+			static constexpr const TCHAR* MaterialXOut = TEXT("MaterialXOut");
+			static constexpr const TCHAR* MaterialXOver = TEXT("MaterialXOver");
+			static constexpr const TCHAR* MaterialXOverlay = TEXT("MaterialXOverlay");
+			static constexpr const TCHAR* MaterialXPlus = TEXT("MaterialXPlus");
+			static constexpr const TCHAR* MaterialXPremult = TEXT("MaterialXPremult");
+			static constexpr const TCHAR* MaterialXRamp4 = TEXT("MaterialXRamp4");
+			static constexpr const TCHAR* MaterialXRampLeftRight = TEXT("MaterialXRampLeftRight");
+			static constexpr const TCHAR* MaterialXRampTopBottom = TEXT("MaterialXRampTopBottom");
+			static constexpr const TCHAR* MaterialXRemap = TEXT("MaterialXRemap");
+			static constexpr const TCHAR* MaterialXScreen = TEXT("MaterialXScreen");
+			static constexpr const TCHAR* MaterialXSplitLeftRight = TEXT("MaterialXSplitLeftRight");
+			static constexpr const TCHAR* MaterialXSplitTopBottom = TEXT("MaterialXSplitTopBottom");
+			static constexpr const TCHAR* MaterialXUnpremult = TEXT("MaterialXUnpremult");
+			static constexpr const TCHAR* Max = TEXT("Max");
+			static constexpr const TCHAR* Min = TEXT("Min");
+			static constexpr const TCHAR* Multiply = TEXT("Multiply");
+			static constexpr const TCHAR* Normalize = TEXT("Normalize");
+			static constexpr const TCHAR* Power = TEXT("Power");
+			static constexpr const TCHAR* RgbToHsv = TEXT("RgbToHsv");
+			static constexpr const TCHAR* Round = TEXT("Round");
+			static constexpr const TCHAR* Sign = TEXT("Sign");
+			static constexpr const TCHAR* Sine = TEXT("Sine");
+			static constexpr const TCHAR* SmoothStep = TEXT("SmoothStep");
+			static constexpr const TCHAR* SquareRoot = TEXT("SquareRoot");
+			static constexpr const TCHAR* SubstrateAdd = TEXT("SubstrateAdd");
+			static constexpr const TCHAR* SubstrateHorizontalMixing = TEXT("SubstrateHorizontalMixing");
+			static constexpr const TCHAR* SubstrateVerticalLayering = TEXT("SubstrateVerticalLayering");
+			static constexpr const TCHAR* SubstrateWeight = TEXT("SubstrateWeight");
+			static constexpr const TCHAR* Subtract = TEXT("Subtract");
+			static constexpr const TCHAR* Tangent = TEXT("Tangent");
+			static constexpr const TCHAR* Time = TEXT("Time");
+			static constexpr const TCHAR* Transform = TEXT("Transform");
+			static constexpr const TCHAR* TransformPosition = TEXT("TransformPosition");
+			static constexpr const TCHAR* VertexColor = TEXT("VertexColor");
+			static constexpr const TCHAR* VertexNormalWS = TEXT("VertexNormalWS");
+			static constexpr const TCHAR* VertexTangentWS = TEXT("VertexTangentWS");
+			static constexpr const TCHAR* WorldPosition = TEXT("WorldPosition");
+		}
+
+		namespace Inputs
+		{
+			static constexpr const TCHAR* A = TEXT("A");
+			static constexpr const TCHAR* AEqualsB = TEXT("AEqualsB");
+			static constexpr const TCHAR* AGreaterThanB = TEXT("AGreaterThanB");
+			static constexpr const TCHAR* Alpha = TEXT("Alpha");
+			static constexpr const TCHAR* ALessThanB = TEXT("ALessThanB");
+			static constexpr const TCHAR* Amplitude = TEXT("Amplitude");
+			static constexpr const TCHAR* B = TEXT("B");
+			static constexpr const TCHAR* Background = TEXT("Background");
+			static constexpr const TCHAR* Base = TEXT("Base");
+			static constexpr const TCHAR* Bottom = TEXT("Bottom");
+			static constexpr const TCHAR* C = TEXT("C");
+			static constexpr const TCHAR* Center = TEXT("Center");
+			static constexpr const TCHAR* Coordinate = TEXT("Coordinate");
+			static constexpr const TCHAR* Coordinates = TEXT("Coordinates");
+			static constexpr const TCHAR* D = TEXT("D");
+			static constexpr const TCHAR* Default = TEXT("Default");
+			static constexpr const TCHAR* Diminish = TEXT("Diminish");
+			static constexpr const TCHAR* Exponent = TEXT("Exponent");
+			static constexpr const TCHAR* Factor = TEXT("Factor");
+			static constexpr const TCHAR* Flatness = TEXT("Flatness");
+			static constexpr const TCHAR* Foreground = TEXT("Foreground");
+			static constexpr const TCHAR* Fraction = TEXT("Fraction");
+			static constexpr const TCHAR* in1 = TEXT("in1");
+			static constexpr const TCHAR* in2 = TEXT("in2");
+			static constexpr const TCHAR* in3 = TEXT("in3");
+			static constexpr const TCHAR* in4 = TEXT("in4");
+			static constexpr const TCHAR* in5 = TEXT("in5");
+			static constexpr const TCHAR* Input = TEXT("Input");
+			static constexpr const TCHAR* InputLow = TEXT("InputLow");
+			static constexpr const TCHAR* InputHigh = TEXT("InputHigh");
+			static constexpr const TCHAR* Lacunarity = TEXT("Lacunarity");
+			static constexpr const TCHAR* LuminanceFactors = TEXT("LuminanceFactors");
+			static constexpr const TCHAR* Max = TEXT("Max");
+			static constexpr const TCHAR* Min = TEXT("Min");
+			static constexpr const TCHAR* Mix = TEXT("Mix");
+			static constexpr const TCHAR* Normal = TEXT("Normal");
+			static constexpr const TCHAR* NormalizedRotationAxis = TEXT("NormalizedRotationAxis");
+			static constexpr const TCHAR* Octaves = TEXT("Octaves");
+			static constexpr const TCHAR* Offset = TEXT("Offset");
+			static constexpr const TCHAR* Pivot = TEXT("Pivot");
+			static constexpr const TCHAR* Position = TEXT("Position");
+			static constexpr const TCHAR* RayDirection = TEXT("Ray Direction");
+			static constexpr const TCHAR* RefractiveIndexOrigin = TEXT("Refractive Index Origin");
+			static constexpr const TCHAR* RefractiveIndexTarget = TEXT("Refractive Index Target");
+			static constexpr const TCHAR* RotationAngle = TEXT("RotationAngle");
+			static constexpr const TCHAR* Scale = TEXT("Scale");
+			static constexpr const TCHAR* SurfaceNormal = TEXT("Surface Normal");
+			static constexpr const TCHAR* SwitchValue = TEXT("SwitchValue");
+			static constexpr const TCHAR* TargetLow = TEXT("TargetLow");
+			static constexpr const TCHAR* TargetHigh = TEXT("TargetHigh");
+			static constexpr const TCHAR* Temp = TEXT("Temp");
+			static constexpr const TCHAR* Time = TEXT("Time");
+			static constexpr const TCHAR* Top = TEXT("Top");
+			static constexpr const TCHAR* Value = TEXT("Value");
+			static constexpr const TCHAR* VectorInput = TEXT("VectorInput");
+			static constexpr const TCHAR* Weight = TEXT("Weight");
+			static constexpr const TCHAR* X = TEXT("X");
+			static constexpr const TCHAR* Y = TEXT("Y");
+		}
+	}
 }
 
 MATERIALX_NAMESPACE_BEGIN
@@ -551,12 +693,15 @@ MATERIALX_NAMESPACE_BEGIN
 		static constexpr const char* Swizzle = "swizzle";
 		// Procedural nodes 
 		static constexpr const char* Constant = "constant";
-		// Procedural2D nodes 
+		// Procedural2D nodes
+		static constexpr const char* CellNoise2D = "cellnoise2d";
+		static constexpr const char* Noise2D = "noise2d";
 		static constexpr const char* Ramp4 = "ramp4";
 		static constexpr const char* RampLR = "ramplr";
 		static constexpr const char* RampTB = "ramptb";
 		static constexpr const char* SplitLR = "splitlr";
 		static constexpr const char* SplitTB = "splittb";
+		static constexpr const char* WorleyNoise2D = "worleynoise2d";
 		// Procedural3D nodes 
 		static constexpr const char* CellNoise3D = "cellnoise3d";
 		static constexpr const char* Fractal3D = "fractal3d";
diff --git a/Engine/Plugins/Interchange/Runtime/Source/Import/Private/MaterialX/MaterialXUtils/MaterialXBase.cpp b/Engine/Plugins/Interchange/Runtime/Source/Import/Private/MaterialX/MaterialXUtils/MaterialXBase.cpp
--- a/Engine/Plugins/Interchange/Runtime/Source/Import/Private/MaterialX/MaterialXUtils/MaterialXBase.cpp
+++ b/Engine/Plugins/Interchange/Runtime/Source/Import/Private/MaterialX/MaterialXUtils/MaterialXBase.cpp
@@ -22,8 +22,10 @@ void FMaterialXBase::UpdateDocumentRecursively(MaterialX::GraphElementPtr Graph)
 		mx::Category::AdobeRgbToLinRec709,
 		mx::Category::Distance,
 		mx::Category::Extract,
+		mx::Category::Fractal3D,
 		mx::Category::LinAdobeRgbToLinRec709,
 		mx::Category::LinDisplayP3ToLinRec709,
+		mx::Category::Noise2D,
 		mx::Category::Noise3D,
 		mx::Category::OpenPBRSurface,
 		mx::Category::Overlay,
diff --git a/Engine/Plugins/Interchange/Runtime/Source/Import/Private/MaterialX/MaterialXUtils/MaterialXManager.cpp b/Engine/Plugins/Interchange/Runtime/Source/Import/Private/MaterialX/MaterialXUtils/MaterialXManager.cpp
--- a/Engine/Plugins/Interchange/Runtime/Source/Import/Private/MaterialX/MaterialXUtils/MaterialXManager.cpp
+++ b/Engine/Plugins/Interchange/Runtime/Source/Import/Private/MaterialX/MaterialXUtils/MaterialXManager.cpp
@@ -30,204 +30,163 @@ const TCHAR FMaterialXManager::TexturePayloadSeparator = TEXT('{');
 #define MX_MATERIALFUNCTION(AssetName) TEXT("/InterchangeAssets/") + FString{bIsSubstrateEnabled ? TEXT("Substrate/") : TEXT("Functions/")} + TEXT(AssetName) TEXT(".") TEXT(AssetName)
 
 FMaterialXManager::FMaterialXManager()
-	: MatchingInputNames {
-		{{TEXT(""),                            TEXT("amplitude")},   TEXT("Amplitude")},
-		{{TEXT(""),                            TEXT("bg")},          TEXT("B")},
-		{{TEXT(""),                            TEXT("center")},      TEXT("Center")},
-		{{TEXT(""),                            TEXT("diminish")},    TEXT("Diminish")},
-		{{TEXT(""),                            TEXT("fg")},          TEXT("A")},
-		{{TEXT(""),                            TEXT("high")},        TEXT("Max")},
-		{{TEXT(""),                            TEXT("in")},          TEXT("Input")},
-		{{TEXT(""),                            TEXT("in1")},         TEXT("A")},
-		{{TEXT(""),                            TEXT("in2")},         TEXT("B")},
-		{{TEXT(""),                            TEXT("in3")},         TEXT("C")},
-		{{TEXT(""),                            TEXT("in4")},         TEXT("D")},
-		{{TEXT(""),                            TEXT("inlow")},       TEXT("InputLow")},
-		{{TEXT(""),                            TEXT("inhigh")},      TEXT("InputHigh")},
-		{{TEXT(""),                            TEXT("lacunarity")},  TEXT("Lacunarity")},
-		{{TEXT(""),                            TEXT("low")},         TEXT("Min")},
-		{{TEXT(""),                            TEXT("lumacoeffs")},  TEXT("LuminanceFactors")}, // for the moment not yet handled by Interchange, because of the attribute being an advanced pin
-		{{TEXT(""),                            TEXT("mix")},         TEXT("Alpha")},
-		{{TEXT(""),                            TEXT("offset")},      TEXT("Offset")},
-		{{TEXT(""),                            TEXT("position")},    TEXT("Position")},
-		{{TEXT(""),                            TEXT("texcoord")},    TEXT("Coordinates")},
-		{{TEXT(""),                            TEXT("octaves")},     TEXT("Octaves")},
-		{{TEXT(""),                            TEXT("temperature")}, TEXT("Temp")},
-		{{TEXT(""),                            TEXT("outlow")},      TEXT("TargetLow")},
-		{{TEXT(""),                            TEXT("outhigh")},     TEXT("TargetHigh")},
-		{{TEXT(""),                            TEXT("valuel")},      TEXT("A")},
-		{{TEXT(""),                            TEXT("valuer")},      TEXT("B")},
-		{{TEXT(""),                            TEXT("valuet")},      TEXT("A")},
-		{{TEXT(""),                            TEXT("valueb")},      TEXT("B")},
-		{{TEXT(""),                            TEXT("valuetl")},     TEXT("A")},
-		{{TEXT(""),                            TEXT("valuetr")},     TEXT("B")},
-		{{TEXT(""),                            TEXT("valuebl")},     TEXT("C")},
-		{{TEXT(""),                            TEXT("valuebr")},     TEXT("D")},
-		{{TEXT(""),                            TEXT("value1")},      TEXT("A")},
-		{{TEXT(""),                            TEXT("value2")},      TEXT("B")},
-		{{MaterialX::Category::Atan2,          TEXT("in1")},         TEXT("Y")},
-		{{MaterialX::Category::Atan2,          TEXT("in2")},         TEXT("X")},
-		{{MaterialX::Category::HeightToNormal, TEXT("scale")},       UE::Interchange::Materials::Standard::Nodes::NormalFromHeightMap::Inputs::Intensity.ToString()},
-		{{MaterialX::Category::IfGreater,      TEXT("in1")},         TEXT("AGreaterThanB")},
-		{{MaterialX::Category::IfGreater,      TEXT("in2")},         TEXT("ALessThanB")},     //another input is added for the case 'equal', see ConnectIfGreater
-		{{MaterialX::Category::IfGreaterEq,    TEXT("in1")},         TEXT("AGreaterThanB")},  //another input is added for the case 'equal', see ConnectIfGreaterEq
-		{{MaterialX::Category::IfGreaterEq,    TEXT("in2")},         TEXT("ALessThanB")},
-		{{MaterialX::Category::IfEqual,        TEXT("in1")},         TEXT("AEqualsB")},
-		{{MaterialX::Category::IfEqual,        TEXT("in2")},         TEXT("ALessThanB")},     // another input is added for the case 'greater', see ConnectIfEqual
-		{{MaterialX::Category::Inside,         TEXT("in")},          TEXT("A")},			  // Inside is treated as a Multiply node
-		{{MaterialX::Category::Inside,         TEXT("mask")},        TEXT("B")},			  // Inside is treated as a Multiply node
-		{{MaterialX::Category::Invert,         TEXT("amount")},      TEXT("A")},
-		{{MaterialX::Category::Invert,         TEXT("in")},          TEXT("B")},
-		{{MaterialX::Category::Mix,            TEXT("fg")},          TEXT("B")},
-		{{MaterialX::Category::Mix,            TEXT("bg")},          TEXT("A")},
-		{{MaterialX::Category::Mix,            TEXT("mix")},		 TEXT("Factor")},
-		{{MaterialX::Category::Noise3D,        TEXT("amplitude")},   TEXT("B")},              // The amplitude of the noise is connected to a multiply node
-		{{MaterialX::Category::Noise3D,        TEXT("pivot")},       TEXT("B")},              // The pivot of the noise is connected to a add node
-		{{MaterialX::Category::Normalize,      TEXT("in")},          TEXT("VectorInput")},
-		{{MaterialX::Category::NormalMap,      TEXT("in")},          TEXT("Normal")},
-		{{MaterialX::Category::NormalMap,      TEXT("scale")},       TEXT("Flatness")},
-		{{MaterialX::Category::Outside,        TEXT("in")},          TEXT("A")},			  // Outside is treated as Multiply node
-		{{MaterialX::Category::Outside,        TEXT("mask")},        TEXT("B")},			  // Outside is treated as Multiply node
-		{{MaterialX::Category::Power,          TEXT("in1")},         TEXT("Base")},
-		{{MaterialX::Category::Power,          TEXT("in2")},         TEXT("Exponent")},
-		{{MaterialX::Category::Refract,        TEXT("in")},          TEXT("Ray Direction")},
-		{{MaterialX::Category::Refract,        TEXT("normal")},      TEXT("Surface Normal")},
-		{{MaterialX::Category::Refract,        TEXT("ior")},         TEXT("Refractive Index Origin")},
-		{{MaterialX::Category::Rotate2D,       TEXT("in")},          TEXT("Coordinate")},
-		{{MaterialX::Category::Rotate2D,       TEXT("amount")},      TEXT("Time")},
-		{{MaterialX::Category::Rotate3D,       TEXT("amount")},      TEXT("RotationAngle")},
-		{{MaterialX::Category::Rotate3D,       TEXT("axis")},		 TEXT("NormalizedRotationAxis")},
-		{{MaterialX::Category::Rotate3D,       TEXT("in")},          TEXT("Position")},
-		{{MaterialX::Category::Saturate,       TEXT("amount")},      TEXT("Fraction")},
-		{{MaterialX::Category::Smoothstep,     TEXT("in")},          TEXT("Value")},
-		{{MaterialX::Category::Switch,         TEXT("in1")},         TEXT("in1")},
-		{{MaterialX::Category::Switch,         TEXT("in2")},         TEXT("in2")},
-		{{MaterialX::Category::Switch,         TEXT("in3")},         TEXT("in3")},
-		{{MaterialX::Category::Switch,         TEXT("in4")},         TEXT("in4")},
-		{{MaterialX::Category::Switch,         TEXT("which")},       TEXT("SwitchValue")},
+	: MatchingInputNames{
+	   {{TEXT(""),                            TEXT("amplitude")},   ExpressionInput(UE::Expressions::Inputs::Amplitude)},
+	   {{TEXT(""),                            TEXT("bg")},          ExpressionInput(UE::Expressions::Inputs::B)},
+	   {{TEXT(""),                            TEXT("center")},      ExpressionInput(UE::Expressions::Inputs::Center)},
+	   {{TEXT(""),                            TEXT("diminish")},    ExpressionInput(UE::Expressions::Inputs::Diminish)},
+	   {{TEXT(""),                            TEXT("fg")},          ExpressionInput(UE::Expressions::Inputs::A)},
+	   {{TEXT(""),                            TEXT("high")},        ExpressionInput(UE::Expressions::Inputs::Max)},
+	   {{TEXT(""),                            TEXT("in")},          ExpressionInput(UE::Expressions::Inputs::Input)},
+	   {{TEXT(""),                            TEXT("in1")},         ExpressionInput(UE::Expressions::Inputs::A)},
+	   {{TEXT(""),                            TEXT("in2")},         ExpressionInput(UE::Expressions::Inputs::B)},
+	   {{TEXT(""),                            TEXT("in3")},         ExpressionInput(UE::Expressions::Inputs::C)},
+	   {{TEXT(""),                            TEXT("in4")},         ExpressionInput(UE::Expressions::Inputs::D)},
+	   {{TEXT(""),                            TEXT("inlow")},       ExpressionInput(UE::Expressions::Inputs::InputLow)},
+	   {{TEXT(""),                            TEXT("inhigh")},      ExpressionInput(UE::Expressions::Inputs::InputHigh)},
+	   {{TEXT(""),                            TEXT("lacunarity")},  ExpressionInput(UE::Expressions::Inputs::Lacunarity)},
+	   {{TEXT(""),                            TEXT("low")},         ExpressionInput(UE::Expressions::Inputs::Min)},
+	   {{TEXT(""),                            TEXT("lumacoeffs")},  ExpressionInput(UE::Expressions::Inputs::LuminanceFactors)}, // for the moment not yet handled by Interchange, because of the attribute being an advanced pin
+	   {{TEXT(""),                            TEXT("mix")},         ExpressionInput(UE::Expressions::Inputs::Alpha)},
+	   {{TEXT(""),                            TEXT("offset")},      ExpressionInput(UE::Expressions::Inputs::Offset)},
+	   {{TEXT(""),                            TEXT("position")},    ExpressionInput(UE::Expressions::Inputs::Position)},
+	   {{TEXT(""),                            TEXT("texcoord")},    ExpressionInput(UE::Expressions::Inputs::Coordinates)},
+	   {{TEXT(""),                            TEXT("octaves")},     ExpressionInput(UE::Expressions::Inputs::Octaves)},
+	   {{TEXT(""),                            TEXT("temperature")}, ExpressionInput(UE::Expressions::Inputs::Temp)},
+	   {{TEXT(""),                            TEXT("outlow")},      ExpressionInput(UE::Expressions::Inputs::TargetLow)},
+	   {{TEXT(""),                            TEXT("outhigh")},     ExpressionInput(UE::Expressions::Inputs::TargetHigh)},
+	   {{TEXT(""),                            TEXT("valuel")},      ExpressionInput(UE::Expressions::Inputs::A)},
+	   {{TEXT(""),                            TEXT("valuer")},      ExpressionInput(UE::Expressions::Inputs::B)},
+	   {{TEXT(""),                            TEXT("valuet")},      ExpressionInput(UE::Expressions::Inputs::A)},
+	   {{TEXT(""),                            TEXT("valueb")},      ExpressionInput(UE::Expressions::Inputs::B)},
+	   {{TEXT(""),                            TEXT("valuetl")},     ExpressionInput(UE::Expressions::Inputs::A)},
+	   {{TEXT(""),                            TEXT("valuetr")},     ExpressionInput(UE::Expressions::Inputs::B)},
+	   {{TEXT(""),                            TEXT("valuebl")},     ExpressionInput(UE::Expressions::Inputs::C)},
+	   {{TEXT(""),                            TEXT("valuebr")},     ExpressionInput(UE::Expressions::Inputs::D)},
+	   {{TEXT(""),                            TEXT("value1")},      ExpressionInput(UE::Expressions::Inputs::A)},
+	   {{TEXT(""),                            TEXT("value2")},      ExpressionInput(UE::Expressions::Inputs::B)},
+	   {{mx::Category::Atan2,                 TEXT("in1")},         ExpressionInput(UE::Expressions::Inputs::Y)},
+	   {{mx::Category::Atan2,                 TEXT("in2")},         ExpressionInput(UE::Expressions::Inputs::X)},
+	   {{mx::Category::HeightToNormal,        TEXT("scale")},       ExpressionInput(*UE::Interchange::Materials::Standard::Nodes::NormalFromHeightMap::Inputs::Intensity.ToString())},
+	   {{mx::Category::IfGreater,             TEXT("in1")},         ExpressionInput(UE::Expressions::Inputs::AGreaterThanB)},
+	   {{mx::Category::IfGreater,             TEXT("in2")},         ExpressionInput(UE::Expressions::Inputs::ALessThanB)},     //another input is added for the case 'equal', see ConnectIfGreater
+	   {{mx::Category::IfGreaterEq,           TEXT("in1")},         ExpressionInput(UE::Expressions::Inputs::AGreaterThanB)},  //another input is added for the case 'equal', see ConnectIfGreaterEq
+	   {{mx::Category::IfGreaterEq,           TEXT("in2")},         ExpressionInput(UE::Expressions::Inputs::ALessThanB)},
+	   {{mx::Category::IfEqual,               TEXT("in1")},         ExpressionInput(UE::Expressions::Inputs::AEqualsB)},
+	   {{mx::Category::IfEqual,               TEXT("in2")},         ExpressionInput(UE::Expressions::Inputs::ALessThanB)},     // another input is added for the case 'greater', see ConnectIfEqual
+	   {{mx::Category::Inside,                TEXT("in")},          ExpressionInput(UE::Expressions::Inputs::A)},			  // Inside is treated as a Multiply node
+	   {{mx::Category::Inside,                TEXT("mask")},        ExpressionInput(UE::Expressions::Inputs::B)},			  // Inside is treated as a Multiply node
+	   {{mx::Category::Invert,                TEXT("amount")},      ExpressionInput(UE::Expressions::Inputs::A)},
+	   {{mx::Category::Invert,                TEXT("in")},          ExpressionInput(UE::Expressions::Inputs::B)},
+	   {{mx::Category::Mix,                   TEXT("fg")},          ExpressionInput(UE::Expressions::Inputs::B)},
+	   {{mx::Category::Mix,                   TEXT("bg")},          ExpressionInput(UE::Expressions::Inputs::A)},
+	   {{mx::Category::Mix,                   TEXT("mix")},		    ExpressionInput(UE::Expressions::Inputs::Factor)},
+	   {{mx::Category::Noise2D,               TEXT("amplitude")},   ExpressionInput(UE::Expressions::Inputs::B)},              // The amplitude of the noise is connected to a multiply node
+	   {{mx::Category::Noise2D,               TEXT("pivot")},       ExpressionInput(UE::Expressions::Inputs::B)},              // The pivot of the noise is connected to a add node
+	   {{mx::Category::Noise3D,               TEXT("amplitude")},   ExpressionInput(UE::Expressions::Inputs::B)},              // The amplitude of the noise is connected to a multiply node
+	   {{mx::Category::Noise3D,               TEXT("pivot")},       ExpressionInput(UE::Expressions::Inputs::B)},              // The pivot of the noise is connected to a add node
+	   {{mx::Category::Normalize,             TEXT("in")},          ExpressionInput(UE::Expressions::Inputs::VectorInput)},
+	   {{mx::Category::NormalMap,             TEXT("in")},          ExpressionInput(UE::Expressions::Inputs::Normal)},
+	   {{mx::Category::NormalMap,             TEXT("scale")},       ExpressionInput(UE::Expressions::Inputs::Flatness)},
+	   {{mx::Category::Outside,               TEXT("in")},          ExpressionInput(UE::Expressions::Inputs::A)},			  // Outside is treated as Multiply node
+	   {{mx::Category::Outside,               TEXT("mask")},        ExpressionInput(UE::Expressions::Inputs::B)},			  // Outside is treated as Multiply node
+	   {{mx::Category::Power,                 TEXT("in1")},         ExpressionInput(UE::Expressions::Inputs::Base)},
+	   {{mx::Category::Power,                 TEXT("in2")},         ExpressionInput(UE::Expressions::Inputs::Exponent)},
+	   {{mx::Category::Refract,               TEXT("in")},          ExpressionInput(UE::Expressions::Inputs::RayDirection)},
+	   {{mx::Category::Refract,               TEXT("normal")},      ExpressionInput(UE::Expressions::Inputs::SurfaceNormal)},
+	   {{mx::Category::Refract,               TEXT("ior")},         ExpressionInput(UE::Expressions::Inputs::RefractiveIndexOrigin)},
+	   {{mx::Category::Rotate2D,              TEXT("in")},          ExpressionInput(UE::Expressions::Inputs::Coordinate)},
+	   {{mx::Category::Rotate2D,              TEXT("amount")},      ExpressionInput(UE::Expressions::Inputs::Time)},
+	   {{mx::Category::Rotate3D,              TEXT("amount")},      ExpressionInput(UE::Expressions::Inputs::RotationAngle)},
+	   {{mx::Category::Rotate3D,              TEXT("axis")},		ExpressionInput(UE::Expressions::Inputs::NormalizedRotationAxis)},
+	   {{mx::Category::Rotate3D,              TEXT("in")},          ExpressionInput(UE::Expressions::Inputs::Position)},
+	   {{mx::Category::Saturate,              TEXT("amount")},      ExpressionInput(UE::Expressions::Inputs::Fraction)},
+	   {{mx::Category::Smoothstep,            TEXT("in")},          ExpressionInput(UE::Expressions::Inputs::Value)},
+	   {{mx::Category::Switch,                TEXT("in1")},         ExpressionInput(UE::Expressions::Inputs::in1)},
+	   {{mx::Category::Switch,                TEXT("in2")},         ExpressionInput(UE::Expressions::Inputs::in2)},
+	   {{mx::Category::Switch,                TEXT("in3")},         ExpressionInput(UE::Expressions::Inputs::in3)},
+	   {{mx::Category::Switch,                TEXT("in4")},         ExpressionInput(UE::Expressions::Inputs::in4)},
+	   {{mx::Category::Switch,                TEXT("in5")},         ExpressionInput(UE::Expressions::Inputs::in5)},
+	   {{mx::Category::Switch,                TEXT("which")},       ExpressionInput(UE::Expressions::Inputs::SwitchValue)},
 	}
 	, MatchingMaterialExpressions {
 		// Math nodes
-		{MaterialX::Category::Absval,       TEXT("Abs")},
-		{MaterialX::Category::Add,          TEXT("Add")},
-		{MaterialX::Category::Acos,         TEXT("Arccosine")},
-		{MaterialX::Category::Asin,         TEXT("Arcsine")},
-		{MaterialX::Category::Atan2,        TEXT("Arctangent2")},
-		{MaterialX::Category::Ceil,         TEXT("Ceil")},
-		{MaterialX::Category::Clamp,        TEXT("Clamp")},
-		{MaterialX::Category::Cos,          TEXT("Cosine")},
-		{MaterialX::Category::CrossProduct, TEXT("Crossproduct")},
-		{MaterialX::Category::Divide,       TEXT("Divide")},
-		{MaterialX::Category::Distance,     TEXT("Distance")},
-		{MaterialX::Category::DotProduct,   TEXT("Dotproduct")},
-		{MaterialX::Category::Exp,          TEXT("Exponential")},
-		{MaterialX::Category::Floor,        TEXT("Floor")},
-		{MaterialX::Category::Invert,       TEXT("Subtract")},
-		{MaterialX::Category::Ln,           TEXT("Logarithm")},
-		{MaterialX::Category::Magnitude,    TEXT("Length")},
-		{MaterialX::Category::Max,          TEXT("Max")},
-		{MaterialX::Category::Min,          TEXT("Min")},
-		{MaterialX::Category::Modulo,       TEXT("Fmod")},
-		{MaterialX::Category::Multiply,     TEXT("Multiply")},
-		{MaterialX::Category::Normalize,    TEXT("Normalize")},
-		{MaterialX::Category::Power,        TEXT("Power")},
-		{MaterialX::Category::RampLR,       TEXT("MaterialXRampLeftRight")},
-		{MaterialX::Category::RampTB,       TEXT("MaterialXRampTopBottom")},
-		{MaterialX::Category::Round,        TEXT("Round")},
-		{MaterialX::Category::Sign,         TEXT("Sign")},
-		{MaterialX::Category::Sin,          TEXT("Sine")},
-		{MaterialX::Category::SplitLR,      TEXT("MaterialXSplitLeftRight")},
-		{MaterialX::Category::SplitTB,      TEXT("MaterialXSplitTopBottom")},
-		{MaterialX::Category::Sqrt,         TEXT("SquareRoot")},
-		{MaterialX::Category::Sub,          TEXT("Subtract")},
-		{MaterialX::Category::Tan,          TEXT("Tangent")},
+		{mx::Category::Absval,       UE::Expressions::Names::Abs},
+		{mx::Category::Add,          UE::Expressions::Names::Add},
+		{mx::Category::Acos,         UE::Expressions::Names::Arccosine},
+		{mx::Category::Asin,         UE::Expressions::Names::Arcsine},
+		{mx::Category::Atan2,        UE::Expressions::Names::Arctangent2},
+		{mx::Category::Ceil,         UE::Expressions::Names::Ceil},
+		{mx::Category::Clamp,        UE::Expressions::Names::Clamp},
+		{mx::Category::Cos,          UE::Expressions::Names::Cosine},
+		{mx::Category::CrossProduct, UE::Expressions::Names::Crossproduct},
+		{mx::Category::Divide,       UE::Expressions::Names::Divide},
+		{mx::Category::Distance,     UE::Expressions::Names::Distance},
+		{mx::Category::DotProduct,   UE::Expressions::Names::Dotproduct},
+		{mx::Category::Exp,          UE::Expressions::Names::Exponential},
+		{mx::Category::Floor,        UE::Expressions::Names::Floor},
+		{mx::Category::Invert,       UE::Expressions::Names::Subtract},
+		{mx::Category::Ln,           UE::Expressions::Names::Logarithm},
+		{mx::Category::Magnitude,    UE::Expressions::Names::Length},
+		{mx::Category::Max,          UE::Expressions::Names::Max},
+		{mx::Category::Min,          UE::Expressions::Names::Min},
+		{mx::Category::Modulo,       UE::Expressions::Names::Fmod},
+		{mx::Category::Multiply,     UE::Expressions::Names::Multiply},
+		{mx::Category::Normalize,    UE::Expressions::Names::Normalize},
+		{mx::Category::Power,        UE::Expressions::Names::Power},
+		{mx::Category::RampLR,       UE::Expressions::Names::MaterialXRampLeftRight},
+		{mx::Category::RampTB,       UE::Expressions::Names::MaterialXRampTopBottom},
+		{mx::Category::Round,        UE::Expressions::Names::Round},
+		{mx::Category::Sign,         UE::Expressions::Names::Sign},
+		{mx::Category::Sin,          UE::Expressions::Names::Sine},
+		{mx::Category::SplitLR,      UE::Expressions::Names::MaterialXSplitLeftRight},
+		{mx::Category::SplitTB,      UE::Expressions::Names::MaterialXSplitTopBottom},
+		{mx::Category::Sqrt,         UE::Expressions::Names::SquareRoot},
+		{mx::Category::Sub,          UE::Expressions::Names::Subtract},
+		{mx::Category::Tan,          UE::Expressions::Names::Tangent},
 		// Compositing nodes
-		{MaterialX::Category::Burn,         TEXT("MaterialXBurn")},
-		{MaterialX::Category::Difference,   TEXT("MaterialXDifference")},
-		{MaterialX::Category::Disjointover, TEXT("MaterialXDisjointover")},
-		{MaterialX::Category::Dodge,        TEXT("MaterialXDodge")},
-		{MaterialX::Category::In,           TEXT("MaterialXIn")},
-		{MaterialX::Category::Inside,       TEXT("Multiply")},
-		{MaterialX::Category::Mask,         TEXT("MaterialXMask")},
-		{MaterialX::Category::Matte,        TEXT("MaterialXMatte")},
-		{MaterialX::Category::Minus,        TEXT("MaterialXMinus")},
-		{MaterialX::Category::Mix,          TEXT("Lerp")},
-		{MaterialX::Category::Out,          TEXT("MaterialXOut")},
-		{MaterialX::Category::Over,         TEXT("MaterialXOver")},
-		{MaterialX::Category::Overlay,      TEXT("MaterialXOverlay")},
-		{MaterialX::Category::Plus,         TEXT("MaterialXPlus")},
-		{MaterialX::Category::Premult,      TEXT("MaterialXPremult")},
-		{MaterialX::Category::Screen,       TEXT("MaterialXScreen")},
-		{MaterialX::Category::Unpremult,    TEXT("MaterialXUnpremult")},
+		{mx::Category::Burn,         UE::Expressions::Names::MaterialXBurn},
+		{mx::Category::Difference,   UE::Expressions::Names::MaterialXDifference},
+		{mx::Category::Disjointover, UE::Expressions::Names::MaterialXDisjointover},
+		{mx::Category::Dodge,        UE::Expressions::Names::MaterialXDodge},
+		{mx::Category::In,           UE::Expressions::Names::MaterialXIn},
+		{mx::Category::Inside,       UE::Expressions::Names::Multiply},
+		{mx::Category::Mask,         UE::Expressions::Names::MaterialXMask},
+		{mx::Category::Matte,        UE::Expressions::Names::MaterialXMatte},
+		{mx::Category::Minus,        UE::Expressions::Names::MaterialXMinus},
+		{mx::Category::Mix,          UE::Expressions::Names::Lerp},
+		{mx::Category::Out,          UE::Expressions::Names::MaterialXOut},
+		{mx::Category::Over,         UE::Expressions::Names::MaterialXOver},
+		{mx::Category::Overlay,      UE::Expressions::Names::MaterialXOverlay},
+		{mx::Category::Plus,         UE::Expressions::Names::MaterialXPlus},
+		{mx::Category::Premult,      UE::Expressions::Names::MaterialXPremult},
+		{mx::Category::Screen,       UE::Expressions::Names::MaterialXScreen},
+		{mx::Category::Unpremult,    UE::Expressions::Names::MaterialXUnpremult},
 		// Channel nodes
-		{MaterialX::Category::Combine2,     TEXT("AppendVector")},
-		{MaterialX::Category::Combine3,     TEXT("MaterialXAppend3Vector")},
-		{MaterialX::Category::Combine4,     TEXT("MaterialXAppend4Vector")},
+		{mx::Category::Combine2,     UE::Expressions::Names::AppendVector},
+		{mx::Category::Combine3,     UE::Expressions::Names::MaterialXAppend3Vector},
+		{mx::Category::Combine4,     UE::Expressions::Names::MaterialXAppend4Vector},
 		// PBR
-		{MaterialX::Category::Blackbody,    TEXT("BlackBody")},
+		{mx::Category::Blackbody,    UE::Expressions::Names::BlackBody},
 		// Procedural2D nodes
-		{MaterialX::Category::Ramp4,        TEXT("MaterialXRamp4")},
-		{MaterialX::Category::RampLR,       TEXT("MaterialXRampLeftRight")},
-		{MaterialX::Category::RampTB,       TEXT("MaterialXRampTopBottom")},
-		{MaterialX::Category::SplitLR,      TEXT("MaterialXSplitLeftRight")},
-		{MaterialX::Category::SplitTB,      TEXT("MaterialXSplitTopBottom")},
+		{mx::Category::Ramp4,        UE::Expressions::Names::MaterialXRamp4},
+		{mx::Category::RampLR,       UE::Expressions::Names::MaterialXRampLeftRight},
+		{mx::Category::RampTB,       UE::Expressions::Names::MaterialXRampTopBottom},
+		{mx::Category::SplitLR,      UE::Expressions::Names::MaterialXSplitLeftRight},
+		{mx::Category::SplitTB,      UE::Expressions::Names::MaterialXSplitTopBottom},
 		// Procedural3D nodes
-		{MaterialX::Category::Fractal3D,    TEXT("MaterialXFractal3D")},
+		{mx::Category::Fractal3D,    UE::Expressions::Names::MaterialXFractal3D},
 		// Geometric nodes 
-		{MaterialX::Category::GeomColor,	TEXT("VertexColor")},
+		{mx::Category::GeomColor,    UE::Expressions::Names::VertexColor},
 		// Adjustment nodes,
-		{MaterialX::Category::HsvToRgb,		TEXT("HsvToRgb")},
-		{MaterialX::Category::Luminance,    TEXT("MaterialXLuminance")},
-		{MaterialX::Category::Remap,        TEXT("MaterialXRemap")},
-		{MaterialX::Category::RgbToHsv,		TEXT("RgbToHsv")},
-		{MaterialX::Category::Saturate,		TEXT("Desaturation")},
-		{MaterialX::Category::Smoothstep,   TEXT("SmoothStep")}
+		{mx::Category::HsvToRgb,     UE::Expressions::Names::HsvToRgb},
+		{mx::Category::Luminance,    UE::Expressions::Names::MaterialXLuminance},
+		{mx::Category::Remap,        UE::Expressions::Names::MaterialXRemap},
+		{mx::Category::RgbToHsv,     UE::Expressions::Names::RgbToHsv},
+		{mx::Category::Saturate,     UE::Expressions::Names::Desaturation},
+		{mx::Category::Smoothstep,   UE::Expressions::Names::SmoothStep}
 	}
-	, MaterialExpressionInputs {
-		TEXT("A"),
-		TEXT("Alpha"),
-		TEXT("Amplitude"),
-		TEXT("B"),
-		TEXT("Base"),
-		TEXT("C"),
-		TEXT("Center"),
-		TEXT("Coordinate"),
-		TEXT("Coordinates"),
-		TEXT("D"),
-		TEXT("Default"),
-		TEXT("Exponent"),
-		TEXT("Factor"),
-		TEXT("Fraction"),
-		TEXT("Input"),
-		TEXT("InputLow"),
-		TEXT("InputHigh"),
-		TEXT("Lacunarity"),
-		TEXT("LuminanceFactors"),
-		TEXT("Max"),
-		TEXT("Min"),
-		TEXT("Octaves"),
-		TEXT("Offset"),
-		TEXT("Pivot"),
-		TEXT("Position"),
-		TEXT("RotationAngle"),
-		TEXT("Scale"),
-		TEXT("SwitchValue"),
-		TEXT("TargetLow"),
-		TEXT("TargetHigh"),
-		TEXT("Temp"),
-		TEXT("Time"),
-		TEXT("Value"),
-		TEXT("VectorInput"),
-		TEXT("X"),
-		TEXT("Y"),
-
-		// case for Refract material function
-		TEXT("Ray Direction"),
-		TEXT("Surface Normal"),
-		TEXT("Refractive Index Origin"),
-		TEXT("Refractive Index Target"),
-	}
-	, MaterialXContainerDelegates{
+	, MaterialXContainerDelegates {
 		{mx::Category::OpenPBRSurface, FMaterialXManager::FOnGetMaterialXInstance::CreateStatic(&FMaterialXOpenPBRSurfaceShader::MakeInstance)},
 		{mx::Category::Surface, FMaterialXManager::FOnGetMaterialXInstance::CreateStatic(&FMaterialXSurfaceShader::MakeInstance)},
 		{mx::Category::SurfaceUnlit, FMaterialXManager::FOnGetMaterialXInstance::CreateStatic(&FMaterialXSurfaceUnlitShader::MakeInstance)},
@@ -239,69 +198,62 @@ FMaterialXManager::FMaterialXManager()
 		{mx::Category::SpotLight, FMaterialXManager::FOnGetMaterialXInstance::CreateStatic(&FMaterialXSpotLightShader::MakeInstance)},
 		{mx::SURFACE_MATERIAL_NODE_STRING.c_str(), FMaterialXManager::FOnGetMaterialXInstance::CreateStatic(&FMaterialXSurfaceMaterial::MakeInstance)}
 	}
+	, MatchingMaterialFunctions {
+		// BSDF Nodes
+		{mx::Category::BurleyDiffuseBSDF,		FMaterialXMaterialFunction{TInPlaceType<EInterchangeMaterialXBSDF>{}, EInterchangeMaterialXBSDF::BurleyDiffuse}},
+		{mx::Category::ConductorBSDF,			FMaterialXMaterialFunction{TInPlaceType<EInterchangeMaterialXBSDF>{}, EInterchangeMaterialXBSDF::Conductor}},
+		{mx::Category::DielectricBSDF,			FMaterialXMaterialFunction{TInPlaceType<EInterchangeMaterialXBSDF>{}, EInterchangeMaterialXBSDF::Dielectric}},
+		{mx::Category::GeneralizedSchlickBSDF,	FMaterialXMaterialFunction{TInPlaceType<EInterchangeMaterialXBSDF>{}, EInterchangeMaterialXBSDF::GeneralizedSchlick}},
+		{mx::Category::OrenNayarDiffuseBSDF ,	FMaterialXMaterialFunction{TInPlaceType<EInterchangeMaterialXBSDF>{}, EInterchangeMaterialXBSDF::OrenNayarDiffuse}},
+		{mx::Category::SheenBSDF,				FMaterialXMaterialFunction{TInPlaceType<EInterchangeMaterialXBSDF>{}, EInterchangeMaterialXBSDF::Sheen}},
+		{mx::Category::SubsurfaceBSDF ,			FMaterialXMaterialFunction{TInPlaceType<EInterchangeMaterialXBSDF>{}, EInterchangeMaterialXBSDF::Subsurface}},
+		{mx::Category::ThinFilmBSDF,			FMaterialXMaterialFunction{TInPlaceType<EInterchangeMaterialXBSDF>{}, EInterchangeMaterialXBSDF::ThinFilm}},
+		{mx::Category::TranslucentBSDF,			FMaterialXMaterialFunction{TInPlaceType<EInterchangeMaterialXBSDF>{}, EInterchangeMaterialXBSDF::Translucent}},
+		// EDF Nodes
+		{mx::Category::ConicalEDF,				FMaterialXMaterialFunction{TInPlaceType<EInterchangeMaterialXEDF>{}, EInterchangeMaterialXEDF::Conical}},
+		{mx::Category::MeasuredEDF,				FMaterialXMaterialFunction{TInPlaceType<EInterchangeMaterialXEDF>{}, EInterchangeMaterialXEDF::Measured}},
+		{mx::Category::UniformEDF,				FMaterialXMaterialFunction{TInPlaceType<EInterchangeMaterialXEDF>{}, EInterchangeMaterialXEDF::Uniform}},
+		// VDF Nodes
+		{mx::Category::AbsorptionVDF,			FMaterialXMaterialFunction{TInPlaceType<EInterchangeMaterialXVDF>{}, EInterchangeMaterialXVDF::Absorption}},
+		{mx::Category::AnisotropicVDF,			FMaterialXMaterialFunction{TInPlaceType<EInterchangeMaterialXVDF>{}, EInterchangeMaterialXVDF::Anisotropic}},
+		// Utility nodes
+		{mx::Category::ArtisticIOR,				FMaterialXMaterialFunction{TInPlaceType<FString>{}, TEXT("/InterchangeAssets/Functions/MX_Artistic_IOR.MX_Artistic_IOR")}},
+		{mx::Category::RoughnessAnisotropy,		FMaterialXMaterialFunction{TInPlaceType<FString>{}, TEXT("/InterchangeAssets/Functions/MX_Roughness_Anisotropy.MX_Roughness_Anisotropy")}},
+		{mx::Category::RoughnessDual,			FMaterialXMaterialFunction{TInPlaceType<FString>{}, TEXT("/InterchangeAssets/Functions/MX_Roughness_Dual.MX_Roughness_Dual")}},
+		// Math
+		{mx::Category::Place2D,					FMaterialXMaterialFunction{TInPlaceType<FString>{}, TEXT("/InterchangeAssets/Functions/MX_Place2D.MX_Place2D")}},
+		{mx::Category::Refract,					FMaterialXMaterialFunction{TInPlaceType<FString>{}, TEXT("/Engine/Functions/Engine_MaterialFunctions01/Vectors/Refract.Refract")}},
+		//ColorTransform
+		{mx::Category::ACEScgToLinRec709,		FMaterialXMaterialFunction{TInPlaceType<FString>{}, TEXT("/InterchangeAssets/Functions/MX_ACEScgToRec709.MX_ACEScgToRec709")}},
+		{mx::Category::AdobeRgbToLinRec709,		FMaterialXMaterialFunction{TInPlaceType<FString>{}, TEXT("/InterchangeAssets/Functions/MX_AdobeRGBToRec709.MX_AdobeRGBToRec709")}},
+		{mx::Category::LinAdobeRgbToLinRec709,	FMaterialXMaterialFunction{TInPlaceType<FString>{}, TEXT("/InterchangeAssets/Functions/MX_LinearAdobeRGBToRec709.MX_LinearAdobeRGBToRec709")}},
+		{mx::Category::LinDisplayP3ToLinRec709,	FMaterialXMaterialFunction{TInPlaceType<FString>{}, TEXT("/InterchangeAssets/Functions/MX_LinearDisplayP3ToRec709.MX_LinearDisplayP3ToRec709")}},
+		{mx::Category::SrgbDisplayP3ToLinRec709,FMaterialXMaterialFunction{TInPlaceType<FString>{}, TEXT("/InterchangeAssets/Functions/MX_SrgbDisplayP3ToRec709.MX_SrgbDisplayP3ToRec709")}},
+		{mx::Category::SrgbTextureToLinRec709,	FMaterialXMaterialFunction{TInPlaceType<FString>{}, TEXT("/InterchangeAssets/Functions/MX_SrgbToRec709.MX_SrgbToRec709")}},
+	}
 	, bIsSubstrateEnabled{ IInterchangeImportModule::IsAvailable() ? IInterchangeImportModule::Get().IsSubstrateEnabled() : false }
 {
-		MatchingMaterialFunctions = {
-			// BSDF Nodes
-			{mx::Category::BurleyDiffuseBSDF,		FMaterialXMaterialFunction{TInPlaceType<EInterchangeMaterialXBSDF>{}, EInterchangeMaterialXBSDF::BurleyDiffuse}},
-			{mx::Category::ConductorBSDF,			FMaterialXMaterialFunction{TInPlaceType<EInterchangeMaterialXBSDF>{}, EInterchangeMaterialXBSDF::Conductor}},
-			{mx::Category::DielectricBSDF,			FMaterialXMaterialFunction{TInPlaceType<EInterchangeMaterialXBSDF>{}, EInterchangeMaterialXBSDF::Dielectric}},
-			{mx::Category::GeneralizedSchlickBSDF,	FMaterialXMaterialFunction{TInPlaceType<EInterchangeMaterialXBSDF>{}, EInterchangeMaterialXBSDF::GeneralizedSchlick}},
-			{mx::Category::OrenNayarDiffuseBSDF ,	FMaterialXMaterialFunction{TInPlaceType<EInterchangeMaterialXBSDF>{}, EInterchangeMaterialXBSDF::OrenNayarDiffuse}},
-			{mx::Category::SheenBSDF,				FMaterialXMaterialFunction{TInPlaceType<EInterchangeMaterialXBSDF>{}, EInterchangeMaterialXBSDF::Sheen}},
-			{mx::Category::SubsurfaceBSDF ,			FMaterialXMaterialFunction{TInPlaceType<EInterchangeMaterialXBSDF>{}, EInterchangeMaterialXBSDF::Subsurface}},
-			{mx::Category::ThinFilmBSDF,			FMaterialXMaterialFunction{TInPlaceType<EInterchangeMaterialXBSDF>{}, EInterchangeMaterialXBSDF::ThinFilm}},
-			{mx::Category::TranslucentBSDF,			FMaterialXMaterialFunction{TInPlaceType<EInterchangeMaterialXBSDF>{}, EInterchangeMaterialXBSDF::Translucent}},
-			// EDF Nodes
-			{mx::Category::ConicalEDF,				FMaterialXMaterialFunction{TInPlaceType<EInterchangeMaterialXEDF>{}, EInterchangeMaterialXEDF::Conical}},
-			{mx::Category::MeasuredEDF,				FMaterialXMaterialFunction{TInPlaceType<EInterchangeMaterialXEDF>{}, EInterchangeMaterialXEDF::Measured}},
-			{mx::Category::UniformEDF,				FMaterialXMaterialFunction{TInPlaceType<EInterchangeMaterialXEDF>{}, EInterchangeMaterialXEDF::Uniform}},
-			// VDF Nodes
-			{mx::Category::AbsorptionVDF,			FMaterialXMaterialFunction{TInPlaceType<EInterchangeMaterialXVDF>{}, EInterchangeMaterialXVDF::Absorption}},
-			{mx::Category::AnisotropicVDF,			FMaterialXMaterialFunction{TInPlaceType<EInterchangeMaterialXVDF>{}, EInterchangeMaterialXVDF::Anisotropic}},
-			// Utility nodes
-			{mx::Category::ArtisticIOR,				FMaterialXMaterialFunction{TInPlaceType<FString>{}, TEXT("/InterchangeAssets/Functions/MX_Artistic_IOR.MX_Artistic_IOR")}},
-			{mx::Category::RoughnessAnisotropy,		FMaterialXMaterialFunction{TInPlaceType<FString>{}, TEXT("/InterchangeAssets/Functions/MX_Roughness_Anisotropy.MX_Roughness_Anisotropy")}},
-			{mx::Category::RoughnessDual,			FMaterialXMaterialFunction{TInPlaceType<FString>{}, TEXT("/InterchangeAssets/Functions/MX_Roughness_Dual.MX_Roughness_Dual")}},
-			// Math
-			{mx::Category::Place2D,					FMaterialXMaterialFunction{TInPlaceType<FString>{}, TEXT("/InterchangeAssets/Functions/MX_Place2D.MX_Place2D")}},
-			{mx::Category::Refract,					FMaterialXMaterialFunction{TInPlaceType<FString>{}, TEXT("/Engine/Functions/Engine_MaterialFunctions01/Vectors/Refract.Refract")}},
-			//ColorTransform
-			{mx::Category::ACEScgToLinRec709,		FMaterialXMaterialFunction{TInPlaceType<FString>{}, TEXT("/InterchangeAssets/Functions/MX_ACEScgToRec709.MX_ACEScgToRec709")}},
-			{mx::Category::AdobeRgbToLinRec709,		FMaterialXMaterialFunction{TInPlaceType<FString>{}, TEXT("/InterchangeAssets/Functions/MX_AdobeRGBToRec709.MX_AdobeRGBToRec709")}},
-			{mx::Category::LinAdobeRgbToLinRec709,	FMaterialXMaterialFunction{TInPlaceType<FString>{}, TEXT("/InterchangeAssets/Functions/MX_LinearAdobeRGBToRec709.MX_LinearAdobeRGBToRec709")}},
-			{mx::Category::LinDisplayP3ToLinRec709,	FMaterialXMaterialFunction{TInPlaceType<FString>{}, TEXT("/InterchangeAssets/Functions/MX_LinearDisplayP3ToRec709.MX_LinearDisplayP3ToRec709")}},
-			{mx::Category::SrgbDisplayP3ToLinRec709,FMaterialXMaterialFunction{TInPlaceType<FString>{}, TEXT("/InterchangeAssets/Functions/MX_SrgbDisplayP3ToRec709.MX_SrgbDisplayP3ToRec709")}},
-			{mx::Category::SrgbTextureToLinRec709,	FMaterialXMaterialFunction{TInPlaceType<FString>{}, TEXT("/InterchangeAssets/Functions/MX_SrgbToRec709.MX_SrgbToRec709")}},
-		};
-
-		if(bIsSubstrateEnabled)
-		{
-			// Vertical Layering
-			MatchingInputNames.Add({ TEXT(""), TEXT("top") }, TEXT("Top"));
-			MatchingInputNames.Add({ TEXT(""), TEXT("base") }, TEXT("Bottom"));
-			MaterialExpressionInputs.Add(TEXT("Top"));
-			MaterialExpressionInputs.Add(TEXT("Bottom"));
-			MatchingMaterialExpressions.Add({ MaterialX::Category::Layer }, TEXT("SubstrateVerticalLayering"));
+	if (bIsSubstrateEnabled)
+	{
+		// Vertical Layering
+		MatchingInputNames.Add({ TEXT(""), TEXT("top") }, ExpressionInput(UE::Expressions::Inputs::Top));
+		MatchingInputNames.Add({ TEXT(""), TEXT("base") }, ExpressionInput(UE::Expressions::Inputs::Bottom));
+		MatchingMaterialExpressions.Add({ mx::Category::Layer }, UE::Expressions::Names::SubstrateVerticalLayering);
 
-			// Horizontal Layering, MaterialX names it mix so we have to be careful not to take the lerp expression
-			MatchingInputNames.Add({ FKeyCategoryNodegroup{MaterialX::Category::Mix, MaterialX::NodeGroup::PBR}, TEXT("bg") }, TEXT("Background"));
-			MatchingInputNames.Add({ FKeyCategoryNodegroup{MaterialX::Category::Mix, MaterialX::NodeGroup::PBR}, TEXT("fg") }, TEXT("Foreground"));
-			MatchingInputNames.Add({ FKeyCategoryNodegroup{MaterialX::Category::Mix, MaterialX::NodeGroup::PBR}, TEXT("mix") }, TEXT("Mix"));
-			MaterialExpressionInputs.Add(TEXT("Background"));
-			MaterialExpressionInputs.Add(TEXT("Foreground"));
-			MaterialExpressionInputs.Add(TEXT("Mix"));
-			MatchingMaterialExpressions.Add({ MaterialX::Category::Mix, MaterialX::NodeGroup::PBR }, TEXT("SubstrateHorizontalMixing"));
+		// Horizontal Layering, MaterialX names it mix so we have to be careful not to take the lerp expression
+		MatchingInputNames.Add({ FKeyCategoryNodegroup{mx::Category::Mix, mx::NodeGroup::PBR}, TEXT("bg") }, ExpressionInput(UE::Expressions::Inputs::Background));
+		MatchingInputNames.Add({ FKeyCategoryNodegroup{mx::Category::Mix, mx::NodeGroup::PBR}, TEXT("fg") }, ExpressionInput(UE::Expressions::Inputs::Foreground));
+		MatchingInputNames.Add({ FKeyCategoryNodegroup{mx::Category::Mix, mx::NodeGroup::PBR}, TEXT("mix") }, ExpressionInput(UE::Expressions::Inputs::Mix));
+		MatchingMaterialExpressions.Add({ mx::Category::Mix, mx::NodeGroup::PBR }, UE::Expressions::Names::SubstrateHorizontalMixing);
 
-			// Add
-			MatchingMaterialExpressions.Add({ MaterialX::Category::Add, MaterialX::NodeGroup::PBR }, TEXT("SubstrateAdd"));
+		// Add
+		MatchingMaterialExpressions.Add({ mx::Category::Add, mx::NodeGroup::PBR }, UE::Expressions::Names::SubstrateAdd);
 
-			// Multiply
-			MatchingInputNames.Add({ FKeyCategoryNodegroup{MaterialX::Category::Multiply, MaterialX::NodeGroup::PBR}, TEXT("in2") }, TEXT("Weight"));
-			MaterialExpressionInputs.Add(TEXT("Weight"));
-			MatchingMaterialExpressions.Add({ MaterialX::Category::Multiply, MaterialX::NodeGroup::PBR }, TEXT("SubstrateWeight"));
-		}
+		// Multiply
+		MatchingInputNames.Add({ FKeyCategoryNodegroup{mx::Category::Multiply, mx::NodeGroup::PBR}, TEXT("in2") }, ExpressionInput(UE::Expressions::Inputs::Weight));
+		MatchingMaterialExpressions.Add({ mx::Category::Multiply, mx::NodeGroup::PBR }, UE::Expressions::Names::SubstrateWeight);
 	}
+}
 
 FMaterialXManager& FMaterialXManager::GetInstance()
 {
@@ -380,6 +332,13 @@ bool FMaterialXManager::IsSubstrateEnabled() const
 {
 	return bIsSubstrateEnabled;
 }
+
+const TCHAR* FMaterialXManager::ExpressionInput(const TCHAR* Input)
+{
+	MaterialExpressionInputs.FindOrAdd(Input);
+	return Input;
+}
+
 #endif
 
 namespace UE::Interchange::MaterialX
diff --git a/Engine/Plugins/Interchange/Runtime/Source/Import/Private/MaterialX/MaterialXUtils/MaterialXSurfaceShaderAbstract.cpp b/Engine/Plugins/Interchange/Runtime/Source/Import/Private/MaterialX/MaterialXUtils/MaterialXSurfaceShaderAbstract.cpp
--- a/Engine/Plugins/Interchange/Runtime/Source/Import/Private/MaterialX/MaterialXUtils/MaterialXSurfaceShaderAbstract.cpp
+++ b/Engine/Plugins/Interchange/Runtime/Source/Import/Private/MaterialX/MaterialXUtils/MaterialXSurfaceShaderAbstract.cpp
@@ -308,7 +308,13 @@ void FMaterialXSurfaceShaderAbstract::ConnectNodeCategoryOutputToInput(const Mat
 
 		// Add all the inputs from the nodedef in case the node doesn't have them all (we don't have a 1:1 match on the default values otherwise)
 		// For exemple min defaults to (0,0) whereas UE to (0, 1)
-		UpstreamNode->addInputsFromNodeDef();
+		// We don't want to add every node, since we don't a have perfect match for each input
+		static const std::set<std::string> NodeDefsCategories{"min", "max", "modulo"};
+		if(NodeDefsCategories.find(UpstreamNode->getCategory()) != NodeDefsCategories.cend())
+		{
+			UpstreamNode->addInputsFromNodeDef();
+		}
+
 		// Replace the input's name by the one used in UE
 		SetMatchingInputsNames(UpstreamNode);
 
@@ -445,15 +451,15 @@ void FMaterialXSurfaceShaderAbstract::ConnectDotInputToOutput(const FConnectNode
 
 void FMaterialXSurfaceShaderAbstract::ConnectTransformPositionInputToOutput(const FConnectNode& Connect)
 {
-	UInterchangeShaderNode* TransformNode = CreateShaderNode(Connect.UpstreamNode, Connect.UpstreamNode->getName().c_str(), TEXT("TransformPosition"));
-	AddAttributeFromValueOrInterface(Connect.UpstreamNode->getInput("in"), TEXT("Input"), TransformNode);
+	UInterchangeShaderNode* TransformNode = CreateShaderNode(Connect.UpstreamNode, Connect.UpstreamNode->getName().c_str(), UE::Expressions::Names::TransformPosition);
+	AddAttributeFromValueOrInterface(Connect.UpstreamNode->getInput("in"), UE::Expressions::Inputs::Input, TransformNode);
 	UInterchangeShaderPortsAPI::ConnectDefaultOuputToInput(Connect.ParentShaderNode, Connect.InputChannelName, TransformNode->GetUniqueID());
 }
 
 void FMaterialXSurfaceShaderAbstract::ConnectTransformVectorInputToOutput(const FConnectNode& Connect)
 {
-	UInterchangeShaderNode* TransformNode = CreateShaderNode(Connect.UpstreamNode, Connect.UpstreamNode->getName().c_str(), TEXT("Transform"));
-	AddAttributeFromValueOrInterface(Connect.UpstreamNode->getInput("in"), TEXT("Input"), TransformNode);
+	UInterchangeShaderNode* TransformNode = CreateShaderNode(Connect.UpstreamNode, Connect.UpstreamNode->getName().c_str(), UE::Expressions::Names::Transform);
+	AddAttributeFromValueOrInterface(Connect.UpstreamNode->getInput("in"), UE::Expressions::Inputs::Input, TransformNode);
 	UInterchangeShaderPortsAPI::ConnectDefaultOuputToInput(Connect.ParentShaderNode, Connect.InputChannelName, TransformNode->GetUniqueID());
 }
 
@@ -473,14 +479,14 @@ void FMaterialXSurfaceShaderAbstract::ConnectRotate2DInputToOutput(const FConnec
 	if(mx::InputPtr Input = Connect.UpstreamNode->getInput("amount"))
 	{
 		FString InputName = GetInputName(Input);
-		UInterchangeShaderNode* DegreesToRadiansNode = CreateShaderNode(Connect.UpstreamNode, (Connect.UpstreamNode->getName() + "multiply").c_str(), TEXT("Multiply"));
+		UInterchangeShaderNode* DegreesToRadiansNode = CreateShaderNode(Connect.UpstreamNode, (Connect.UpstreamNode->getName() + "multiply").c_str(), Multiply::Name.ToString());
 		constexpr float DegreesToRadians = UE_PI / 180.f;
 		DegreesToRadiansNode->AddFloatAttribute(UInterchangeShaderPortsAPI::MakeInputValueKey(TEXT("B")), DegreesToRadians);
 
 		//If it's a value we should always attach it on the A input of Divide node
-		AddAttributeFromValueOrInterface(Input, TEXT("A"), DegreesToRadiansNode);
+		AddAttributeFromValueOrInterface(Input, UE::Expressions::Inputs::A, DegreesToRadiansNode);
 		UInterchangeShaderPortsAPI::ConnectDefaultOuputToInput(Rotate2DNode, InputName, DegreesToRadiansNode->GetUniqueID());
-		SetAttributeNewName(Input, "A");
+		SetAttributeNewName(Input, TCHAR_TO_ANSI(UE::Expressions::Inputs::A));
 	}
 
 	Rotate2DNode->AddFloatAttribute(Rotator::Attributes::CenterX.ToString(), 0.f);
@@ -495,13 +501,13 @@ void FMaterialXSurfaceShaderAbstract::ConnectRotate3DInputToOutput(const FConnec
 	UInterchangeShaderNode* Rotate3DNode = CreateShaderNode(Connect.UpstreamNode, Connect.UpstreamNode->getName().c_str(), TEXT("RotateAboutAxis"));
 	Rotate3DNode->AddLinearColorAttribute(UInterchangeShaderPortsAPI::MakeInputValueKey(TEXT("PivotPoint")), FLinearColor(0.5, 0.5, 0));
 
-	AddAttributeFromValueOrInterface(Connect.UpstreamNode->getInput("in"), TEXT("Position"), Rotate3DNode);
-	AddAttributeFromValueOrInterface(Connect.UpstreamNode->getInput("axis"), TEXT("NormalizedRotationAxis"), Rotate3DNode);
+	AddAttributeFromValueOrInterface(Connect.UpstreamNode->getInput("in"), UE::Expressions::Inputs::Position, Rotate3DNode);
+	AddAttributeFromValueOrInterface(Connect.UpstreamNode->getInput("axis"), UE::Expressions::Inputs::NormalizedRotationAxis, Rotate3DNode);
 
 	// we create a Divide node to convert MaterialX angle in degrees to UE's angle which is a value between [0,1]
 	if(mx::InputPtr Input = Connect.UpstreamNode->getInput("amount"))
 	{
-		if(!AddAttributeFromValueOrInterface(Input, TEXT("RotationAngle"), Rotate3DNode))
+		if(!AddAttributeFromValueOrInterface(Input, UE::Expressions::Inputs::RotationAngle, Rotate3DNode))
 		{
 			//copy "in" input into "in1" and create "amount" input as "in2" with the value 360 because UE's angle is a value between [0,1]
 			mx::NodePtr NewDivideNode = CreateNode(Connect.UpstreamNode->getParent(),
@@ -618,7 +624,7 @@ void FMaterialXSurfaceShaderAbstract::ConnectConvertInputToOutput(const FConnect
 			MaskShaderNode = CreateMaskShaderNode(0b1000, Connect.UpstreamNode, Connect.UpstreamNode->getName().c_str());
 		}
 
-		AddAttributeFromValueOrInterface(Input, TEXT("Input"), MaskShaderNode);
+		AddAttributeFromValueOrInterface(Input, UE::Expressions::Inputs::Input, MaskShaderNode);
 
 		UInterchangeShaderPortsAPI::ConnectDefaultOuputToInput(Connect.ParentShaderNode, Connect.InputChannelName, MaskShaderNode->GetUniqueID());
 	}
@@ -692,7 +698,7 @@ void FMaterialXSurfaceShaderAbstract::ConnectConvertInputToOutput(const FConnect
 
 void FMaterialXSurfaceShaderAbstract::ConnectIfGreaterInputToOutput(const FConnectNode& Connect)
 {
-	UInterchangeShaderNode* NodeIf = CreateShaderNode(Connect.UpstreamNode, Connect.UpstreamNode->getName().c_str(), TEXT("If"));
+	UInterchangeShaderNode* NodeIf = CreateShaderNode(Connect.UpstreamNode, Connect.UpstreamNode->getName().c_str(), UE::Expressions::Names::If);
 	UInterchangeShaderPortsAPI::ConnectDefaultOuputToInput(Connect.ParentShaderNode, Connect.InputChannelName, NodeIf->GetUniqueID());
 
 	if(mx::InputPtr Input = Connect.UpstreamNode->getInput("value1"))
@@ -715,21 +721,21 @@ void FMaterialXSurfaceShaderAbstract::ConnectIfGreaterInputToOutput(const FConne
 	if(Input)
 	{
 		AddAttributeFromValueOrInterface(Input, GetInputName(Input), NodeIf);
-		AddAttributeFromValueOrInterface(Input, TEXT("AEqualsB"), NodeIf);
+		AddAttributeFromValueOrInterface(Input, UE::Expressions::Inputs::AEqualsB, NodeIf);
 
 		//Let's add a new input that is a copy of in2 to connect it to the equal input
 		if (!Connect.UpstreamNode->getInput("in3"))
 		{
 			mx::InputPtr Input3 = Connect.UpstreamNode->addInput("in3");
 			Input3->copyContentFrom(Input);
-			SetAttributeNewName(Input3, "AEqualsB");
+			SetAttributeNewName(Input3, TCHAR_TO_ANSI(UE::Expressions::Inputs::AEqualsB));
 		}
 	}
 }
 
 void FMaterialXSurfaceShaderAbstract::ConnectIfGreaterEqInputToOutput(const FConnectNode& Connect)
 {
-	UInterchangeShaderNode* NodeIf = CreateShaderNode(Connect.UpstreamNode, Connect.UpstreamNode->getName().c_str(), TEXT("If"));
+	UInterchangeShaderNode* NodeIf = CreateShaderNode(Connect.UpstreamNode, Connect.UpstreamNode->getName().c_str(), UE::Expressions::Names::If);
 	UInterchangeShaderPortsAPI::ConnectDefaultOuputToInput(Connect.ParentShaderNode, Connect.InputChannelName, NodeIf->GetUniqueID());
 
 	if(mx::InputPtr Input = Connect.UpstreamNode->getInput("value1"))
@@ -752,21 +758,21 @@ void FMaterialXSurfaceShaderAbstract::ConnectIfGreaterEqInputToOutput(const FCon
 	if(Input)
 	{
 		AddAttributeFromValueOrInterface(Input, GetInputName(Input), NodeIf);
-		AddAttributeFromValueOrInterface(Input, TEXT("AEqualsB"), NodeIf);
+		AddAttributeFromValueOrInterface(Input, UE::Expressions::Inputs::AEqualsB, NodeIf);
 
 		//Let's add a new input that is a copy of in2 to connect it to the equal input
 		if (!Connect.UpstreamNode->getInput("in3"))
 		{
 			mx::InputPtr Input3 = Connect.UpstreamNode->addInput("in3");
 			Input3->copyContentFrom(Input);
-			SetAttributeNewName(Input3, "AEqualsB");
+			SetAttributeNewName(Input3, TCHAR_TO_ANSI(UE::Expressions::Inputs::AEqualsB));
 		}
 	}
 }
 
 void FMaterialXSurfaceShaderAbstract::ConnectIfEqualInputToOutput(const FConnectNode& Connect)
 {
-	UInterchangeShaderNode* NodeIf = CreateShaderNode(Connect.UpstreamNode, Connect.UpstreamNode->getName().c_str(), TEXT("If"));
+	UInterchangeShaderNode* NodeIf = CreateShaderNode(Connect.UpstreamNode, Connect.UpstreamNode->getName().c_str(), UE::Expressions::Names::If);
 	UInterchangeShaderPortsAPI::ConnectDefaultOuputToInput(Connect.ParentShaderNode, Connect.InputChannelName, NodeIf->GetUniqueID());
 
 	if(mx::InputPtr Input = Connect.UpstreamNode->getInput("value1"))
@@ -789,28 +795,29 @@ void FMaterialXSurfaceShaderAbstract::ConnectIfEqualInputToOutput(const FConnect
 	if(Input)
 	{
 		AddAttributeFromValueOrInterface(Input, GetInputName(Input), NodeIf);
-		AddAttributeFromValueOrInterface(Input, TEXT("AGreaterThanB"), NodeIf);
+		AddAttributeFromValueOrInterface(Input, UE::Expressions::Inputs::AGreaterThanB, NodeIf);
 
 		//Let's add a new input that is a copy of in2 to connect it to the equal input
 		if(!Connect.UpstreamNode->getInput("in3"))
 		{
 			mx::InputPtr Input3 = Connect.UpstreamNode->addInput("in3");
 			Input3->copyContentFrom(Input);
-			SetAttributeNewName(Input3, "AGreaterThanB");
+			SetAttributeNewName(Input3, TCHAR_TO_ANSI(UE::Expressions::Inputs::AGreaterThanB));
 		}
 	}
 }
 
 void FMaterialXSurfaceShaderAbstract::ConnectOutsideInputToOutput(const FConnectNode& Connect)
 {
+	using namespace UE::Interchange::Materials::Standard::Nodes;
 	//in * (1 - mask)
-	UInterchangeShaderNode* NodeMultiply = CreateShaderNode(Connect.UpstreamNode, Connect.UpstreamNode->getName().c_str(), TEXT("Multiply"));
-	AddAttributeFromValueOrInterface(Connect.UpstreamNode->getInput("in"), TEXT("A"), NodeMultiply);
+	UInterchangeShaderNode* NodeMultiply = CreateShaderNode(Connect.UpstreamNode, Connect.UpstreamNode->getName().c_str(), Multiply::Name.ToString());
+	AddAttributeFromValueOrInterface(Connect.UpstreamNode->getInput("in"), Multiply::Inputs::A.ToString(), NodeMultiply);
 	UInterchangeShaderPortsAPI::ConnectDefaultOuputToInput(Connect.ParentShaderNode, Connect.InputChannelName, NodeMultiply->GetUniqueID());
 
-	UInterchangeShaderNode* NodeOneMinus = CreateShaderNode(Connect.UpstreamNode, Connect.UpstreamNode->getName().c_str() + FString(TEXT("_OneMinus")), TEXT("OneMinus"));
-	AddAttributeFromValueOrInterface(Connect.UpstreamNode->getInput("mask"), TEXT("Input"), NodeOneMinus);
-	UInterchangeShaderPortsAPI::ConnectDefaultOuputToInput(NodeMultiply, TEXT("B"), NodeOneMinus->GetUniqueID());
+	UInterchangeShaderNode* NodeOneMinus = CreateShaderNode(Connect.UpstreamNode, Connect.UpstreamNode->getName().c_str() + FString(TEXT("_OneMinus")), OneMinus::Name.ToString());
+	AddAttributeFromValueOrInterface(Connect.UpstreamNode->getInput("mask"), OneMinus::Inputs::Input.ToString(), NodeOneMinus);
+	UInterchangeShaderPortsAPI::ConnectDefaultOuputToInput(NodeMultiply, Multiply::Inputs::B.ToString(), NodeOneMinus->GetUniqueID());
 }
 
 UInterchangeShaderNode* FMaterialXSurfaceShaderAbstract::ConnectGeometryInputToOutput(const FConnectNode& Connect, const FString& ShaderType, const FString& TransformShaderType, const FString& TransformInput, const FString& TransformSourceType, int32 TransformSource, const FString& TransformType, int32 TransformSDestination, bool bIsVector)
@@ -859,14 +866,44 @@ UInterchangeShaderNode* FMaterialXSurfaceShaderAbstract::ConnectGeometryInputToO
 	return GeometryNode;
 }
 
+UInterchangeShaderNode* FMaterialXSurfaceShaderAbstract::ConnectNoise2DInputToOutput(const FConnectNode& Connect, const FString& ShaderType, EVectorNoiseFunction NoiseFunction, uint8 Mask)
+{
+	using namespace UE::Interchange::Materials::Standard::Nodes;
+
+	// Let's use a vector noise for this one, the only one that is close to MaterialX implementation
+	UInterchangeShaderNode* NoiseNode = CreateShaderNode(Connect.UpstreamNode, (Connect.UpstreamNode->getName() + '_').c_str() + ShaderType, VectorNoise::Name.ToString());
+	NoiseNode->AddInt32Attribute(VectorNoise::Attributes::Function.ToString(), NoiseFunction);
+	//Quality only has a meaning with Voronoi (Worley), to be as close to MaterialX spec we want a 3x3x3 neighborhood
+	NoiseNode->AddInt32Attribute(VectorNoise::Attributes::Quality.ToString(), 3);
+
+	// These noise algorithms are meant to be used with 3D inputs, in order to use it with 2D input, let's put an arbitrary seed for 'z'
+	// not sure 0 is a good candidate, we'll see when we'll have real use cases.
+	UInterchangeShaderNode* AppendNode = CreateShaderNode(Connect.UpstreamNode, Connect.UpstreamNode->getName().c_str(), UE::Expressions::Names::AppendVector);
+	AppendNode->AddFloatAttribute(UInterchangeShaderPortsAPI::MakeInputValueKey(UE::Expressions::Inputs::B), 0.f);
+
+	// The connecting node has to be plugged in the Append expression not the noise, since we are constructing a vector2 to a vector3
+	if (mx::InputPtr Input = Connect.UpstreamNode->getInput("texcoord"))
+	{
+		AddAttributeFromValueOrInterface(Input, UE::Expressions::Inputs::A, AppendNode);
+		SetAttributeNewName(Input, TCHAR_TO_ANSI(UE::Expressions::Inputs::A));
+	}
+
+	UInterchangeShaderNode* MaskNode = CreateMaskShaderNode(Mask, Connect.UpstreamNode, (Connect.UpstreamNode->getName() + "_ComponentMask").c_str());
+
+	UInterchangeShaderPortsAPI::ConnectDefaultOuputToInput(NoiseNode, VectorNoise::Inputs::Position.ToString(), AppendNode->GetUniqueID());
+	UInterchangeShaderPortsAPI::ConnectDefaultOuputToInput(MaskNode, Mask::Inputs::Input.ToString(), NoiseNode->GetUniqueID());
+
+	return MaskNode;
+}
+
 void FMaterialXSurfaceShaderAbstract::ConnectPositionInputToOutput(const FConnectNode& Connect)
 {
 	using namespace UE::Interchange::Materials::Standard::Nodes;
-	UInterchangeShaderNode* PositionNode = CreateShaderNode(Connect.UpstreamNode, (Connect.UpstreamNode->getName() + "_Position").c_str(), TEXT("LocalPosition"));
+	UInterchangeShaderNode* PositionNode = CreateShaderNode(Connect.UpstreamNode, (Connect.UpstreamNode->getName() + "_Position").c_str(), UE::Expressions::Names::LocalPosition);
 	
 	if (mx::InputPtr InputSpace = Connect.UpstreamNode->getInput("space"); InputSpace && InputSpace->getValueString() == "world")
 	{
-		PositionNode->SetCustomShaderType("WorldPosition");
+		PositionNode->SetCustomShaderType(UE::Expressions::Names::WorldPosition);
 	}
 
 	UInterchangeShaderNode* UnitNode = CreateShaderNode(Connect.UpstreamNode, Connect.UpstreamNode->getName().c_str(), Multiply::Name.ToString());
@@ -909,8 +946,8 @@ void FMaterialXSurfaceShaderAbstract::ConnectBitangentInputToOutput(const FConne
 																		 TransformVector::Attributes::TransformSourceType.ToString(), EMaterialVectorCoordTransformSource::TRANSFORMSOURCE_World,
 																		 TransformVector::Attributes::TransformType.ToString(), EMaterialVectorCoordTransform::TRANSFORM_Local);
 
-	UInterchangeShaderNode* NormalNode = CreateShaderNode(Connect.UpstreamNode, Connect.UpstreamNode->getName().c_str() + FString(TEXT("_Normal")), TEXT("VertexNormalWS"));
-	UInterchangeShaderNode* TangentNode = CreateShaderNode(Connect.UpstreamNode, Connect.UpstreamNode->getName().c_str() + FString(TEXT("_Tangent")), TEXT("VertexTangentWS"));
+	UInterchangeShaderNode* NormalNode = CreateShaderNode(Connect.UpstreamNode, Connect.UpstreamNode->getName().c_str() + FString(TEXT("_Normal")), UE::Expressions::Names::VertexNormalWS);
+	UInterchangeShaderNode* TangentNode = CreateShaderNode(Connect.UpstreamNode, Connect.UpstreamNode->getName().c_str() + FString(TEXT("_Tangent")), UE::Expressions::Names::VertexTangentWS);
 
 	UInterchangeShaderPortsAPI::ConnectDefaultOuputToInput(BitangentNode, TEXT("A"), NormalNode->GetUniqueID());
 	UInterchangeShaderPortsAPI::ConnectDefaultOuputToInput(BitangentNode, TEXT("B"), TangentNode->GetUniqueID());
@@ -919,7 +956,7 @@ void FMaterialXSurfaceShaderAbstract::ConnectBitangentInputToOutput(const FConne
 void FMaterialXSurfaceShaderAbstract::ConnectTimeInputToOutput(const FConnectNode& Connect)
 {
 	using namespace UE::Interchange::Materials::Standard::Nodes;
-	UInterchangeShaderNode* TimeNode = CreateShaderNode(Connect.UpstreamNode, Connect.UpstreamNode->getName().c_str(), TEXT("Time"));
+	UInterchangeShaderNode* TimeNode = CreateShaderNode(Connect.UpstreamNode, Connect.UpstreamNode->getName().c_str(), UE::Expressions::Names::Time);
 	TimeNode->AddBooleanAttribute(Time::Attributes::OverridePeriod.ToString(), true);
 
 	float FPS;
@@ -939,6 +976,73 @@ void FMaterialXSurfaceShaderAbstract::ConnectTimeInputToOutput(const FConnectNod
 	UInterchangeShaderPortsAPI::ConnectDefaultOuputToInput(Connect.ParentShaderNode, Connect.InputChannelName, TimeNode->GetUniqueID());
 }
 
+void FMaterialXSurfaceShaderAbstract::ConnectNoise2DInputToOutput(const FConnectNode& Connect)
+{
+	using namespace UE::Interchange::Materials::Standard::Nodes;
+	
+	const std::string& Type = Connect.UpstreamNode->getType();
+	uint8 Mask = 
+		Type == "float" ? 0b1000 :
+		Type == "vector2" ? 0b1100 :
+		0b1110;
+
+	UInterchangeShaderNode* PerlinNoiseNode = ConnectNoise2DInputToOutput(Connect, VectorNoise::Name.ToString(), EVectorNoiseFunction::VNF_VectorALU, Mask);
+
+	UInterchangeShaderNode* NodeToConnect = PerlinNoiseNode;
+
+	// Amplitude * Noise + Pivot
+	auto ConnectNodeToInput = [&](mx::InputPtr Input, UInterchangeShaderNode* NodeToConnectTo, const FString& ShaderType) -> UInterchangeShaderNode*
+		{
+			if (!Input)
+			{
+				return nullptr;
+			}
+
+			const FString ShaderNodeName = Connect.UpstreamNode->getName().c_str() + FString{ TEXT("_") } + ShaderType;
+			UInterchangeShaderNode* ShaderNode = CreateShaderNode(Connect.UpstreamNode, ShaderNodeName, ShaderType);
+
+			UInterchangeShaderPortsAPI::ConnectDefaultOuputToInput(ShaderNode, UE::Expressions::Inputs::A, NodeToConnectTo->GetUniqueID());
+
+			// Connect the amplitude/pivot node to the shader node not the noise
+			// it will be handle during the upstream-downstream connection phase
+			Input->setAttribute(mx::Attributes::ParentName, TCHAR_TO_UTF8(*ShaderNodeName));
+			AddAttributeFromValueOrInterface(Input, UE::Expressions::Inputs::B, ShaderNode);
+
+			return ShaderNode;
+		};
+
+	if (UInterchangeShaderNode* MultiplyNode = ConnectNodeToInput(Connect.UpstreamNode->getInput("amplitude"), PerlinNoiseNode, Multiply::Name.ToString()))
+	{
+		NodeToConnect = MultiplyNode;
+	}
+
+	if (UInterchangeShaderNode* AddNode = ConnectNodeToInput(Connect.UpstreamNode->getInput("pivot"), NodeToConnect, Add::Name.ToString()))
+	{
+		NodeToConnect = AddNode;
+	}
+
+	UInterchangeShaderPortsAPI::ConnectDefaultOuputToInput(Connect.ParentShaderNode, Connect.InputChannelName, NodeToConnect->GetUniqueID());
+}
+
+void FMaterialXSurfaceShaderAbstract::ConnectCellNoise2DInputToOutput(const FConnectNode& Connect)
+{
+	using namespace UE::Interchange::Materials::Standard::Nodes;
+
+	// UE's Cellnoise algorithm returns a float3 but MaterialX returns only a float
+	UInterchangeShaderNode* CellNoiseNode = ConnectNoise2DInputToOutput(Connect, VectorNoise::Name.ToString(), EVectorNoiseFunction::VNF_CellnoiseALU, 0b1000);
+	UInterchangeShaderPortsAPI::ConnectDefaultOuputToInput(Connect.ParentShaderNode, Connect.InputChannelName, CellNoiseNode->GetUniqueID());
+}
+
+void FMaterialXSurfaceShaderAbstract::ConnectWorleyNoise2DInputToOutput(const FConnectNode& Connect)
+{
+	using namespace UE::Interchange::Materials::Standard::Nodes;
+
+	// UE's Voronoi algorithm returns either XYZ (the position of a Cell in Voronoi) or W which is the distance (what we want in that case)
+	// Sadly we have no way to return the distance to the closest seed in another dimensions, since the component is too deep in the algorithm (see VoronoiNoise3D_ALU in Random.ush)
+	UInterchangeShaderNode* WorleyNoiseNode = ConnectNoise2DInputToOutput(Connect, VectorNoise::Name.ToString(), EVectorNoiseFunction::VNF_VoronoiALU, 0b0001);
+	UInterchangeShaderPortsAPI::ConnectDefaultOuputToInput(Connect.ParentShaderNode, Connect.InputChannelName, WorleyNoiseNode->GetUniqueID());
+}
+
 void FMaterialXSurfaceShaderAbstract::ConnectNoise3DInputToOutput(const FConnectNode& Connect)
 {
 	using namespace UE::Interchange::Materials::Standard::Nodes;
@@ -964,24 +1068,32 @@ void FMaterialXSurfaceShaderAbstract::ConnectNoise3DInputToOutput(const FConnect
 
 		// Connect the amplitude node to the shader node not the noise
 		// it will be handle during the upstream-downstream connection phase
-		// The index is here for the unicity of the attribute
 		Input->setAttribute(mx::Attributes::ParentName, TCHAR_TO_UTF8(*ShaderNodeName));
-		AddAttributeFromValueOrInterface(Input, TEXT("B"), ShaderNode);
+		AddAttributeFromValueOrInterface(Input, UE::Expressions::Inputs::B, ShaderNode);
 
 		return ShaderNode;
 	};
 
-	if(UInterchangeShaderNode* MultiplyNode = ConnectNodeToInput(Connect.UpstreamNode->getInput("amplitude"), NoiseNode, TEXT("Multiply")))
+	if(UInterchangeShaderNode* MultiplyNode = ConnectNodeToInput(Connect.UpstreamNode->getInput("amplitude"), NoiseNode, Multiply::Name.ToString()))
 	{
 		NodeToConnect = MultiplyNode;
 	}
 
-	if(UInterchangeShaderNode* AddNode = ConnectNodeToInput(Connect.UpstreamNode->getInput("pivot"), NodeToConnect, TEXT("Add")))
+	if(UInterchangeShaderNode* AddNode = ConnectNodeToInput(Connect.UpstreamNode->getInput("pivot"), NodeToConnect, Add::Name.ToString()))
 	{
 		NodeToConnect = AddNode;
 	}
 
-	UInterchangeShaderPortsAPI::ConnectDefaultOuputToInput(Connect.ParentShaderNode, Connect.InputChannelName, NodeToConnect->GetUniqueID());
+	const std::string& Type = Connect.UpstreamNode->getType();
+	uint8 Mask =
+		Type == "float" ? 0b1000 :
+		Type == "vector2" ? 0b1100 :
+		0b1110;
+
+	UInterchangeShaderNode* MaskNode = CreateMaskShaderNode(Mask, Connect.UpstreamNode, (Connect.UpstreamNode->getName() + "_ComponentMask").c_str());
+
+	UInterchangeShaderPortsAPI::ConnectDefaultOuputToInput(MaskNode, Mask::Inputs::Input.ToString(), NodeToConnect->GetUniqueID());
+	UInterchangeShaderPortsAPI::ConnectDefaultOuputToInput(Connect.ParentShaderNode, Connect.InputChannelName, MaskNode->GetUniqueID());
 }
 
 void FMaterialXSurfaceShaderAbstract::ConnectCellNoise3DInputToOutput(const FConnectNode& Connect)
@@ -992,7 +1104,11 @@ void FMaterialXSurfaceShaderAbstract::ConnectCellNoise3DInputToOutput(const FCon
 	UInterchangeShaderNode* NoiseNode = CreateShaderNode(Connect.UpstreamNode, Connect.UpstreamNode->getName().c_str(), VectorNoise::Name.ToString());
 	NoiseNode->AddInt32Attribute(VectorNoise::Attributes::Function.ToString(), EVectorNoiseFunction::VNF_CellnoiseALU);
 
-	UInterchangeShaderPortsAPI::ConnectDefaultOuputToInput(Connect.ParentShaderNode, Connect.InputChannelName, NoiseNode->GetUniqueID());
+	//cellnoise3d only supports float output
+	UInterchangeShaderNode* MaskNode = CreateMaskShaderNode(0b1000, Connect.UpstreamNode, (Connect.UpstreamNode->getName() + "_ComponentMask").c_str());
+
+	UInterchangeShaderPortsAPI::ConnectDefaultOuputToInput(MaskNode, Mask::Inputs::Input.ToString(), NoiseNode->GetUniqueID());
+	UInterchangeShaderPortsAPI::ConnectDefaultOuputToInput(Connect.ParentShaderNode, Connect.InputChannelName, MaskNode->GetUniqueID());
 }
 
 void FMaterialXSurfaceShaderAbstract::ConnectWorleyNoise3DInputToOutput(const FConnectNode& Connect)
@@ -1000,10 +1116,16 @@ void FMaterialXSurfaceShaderAbstract::ConnectWorleyNoise3DInputToOutput(const FC
 	using namespace UE::Interchange::Materials::Standard::Nodes;
 
 	//Also called Voronoi, the implementation is a bit different in UE, especially we don't have access to the jitter
-	UInterchangeShaderNode* NoiseNode = CreateShaderNode(Connect.UpstreamNode, Connect.UpstreamNode->getName().c_str(), Noise::Name.ToString());
-	NoiseNode->AddInt32Attribute(Noise::Attributes::Function.ToString(), ENoiseFunction::NOISEFUNCTION_VoronoiALU);
+	UInterchangeShaderNode* NoiseNode = CreateShaderNode(Connect.UpstreamNode, Connect.UpstreamNode->getName().c_str(), VectorNoise::Name.ToString());
+	NoiseNode->AddInt32Attribute(Noise::Attributes::Function.ToString(), EVectorNoiseFunction::VNF_VoronoiALU);
+	// 3x3x3 neighborhood
+	NoiseNode->AddInt32Attribute(Noise::Attributes::Quality.ToString(), 3);
 
-	UInterchangeShaderPortsAPI::ConnectDefaultOuputToInput(Connect.ParentShaderNode, Connect.InputChannelName, NoiseNode->GetUniqueID());
+	// Voronoi only supports float distance to the seed
+	UInterchangeShaderNode* MaskNode = CreateMaskShaderNode(0b0001, Connect.UpstreamNode, (Connect.UpstreamNode->getName() + "_ComponentMask").c_str());
+
+	UInterchangeShaderPortsAPI::ConnectDefaultOuputToInput(MaskNode, Mask::Inputs::Input.ToString(), NoiseNode->GetUniqueID());
+	UInterchangeShaderPortsAPI::ConnectDefaultOuputToInput(Connect.ParentShaderNode, Connect.InputChannelName, MaskNode->GetUniqueID());
 }
 
 void FMaterialXSurfaceShaderAbstract::ConnectHeightToNormalInputToOutput(const FConnectNode& Connect)
@@ -1158,9 +1280,9 @@ void FMaterialXSurfaceShaderAbstract::ConnectTexCoordInputToOutput(const FConnec
 	UInterchangeShaderNode* OneMinusG = CreateShaderNode(Connect.UpstreamNode, (Connect.UpstreamNode->getName() + "OneMinusG").c_str(), OneMinus::Name.ToString());
 	UInterchangeShaderPortsAPI::ConnectDefaultOuputToInput(OneMinusG, OneMinus::Inputs::Input.ToString(), TexCoordG->GetUniqueID());
 
-	UInterchangeShaderNode* TexCoordAppend = CreateShaderNode(Connect.UpstreamNode, (Connect.UpstreamNode->getName() + "Append").c_str(), TEXT("AppendVector"));
-	UInterchangeShaderPortsAPI::ConnectDefaultOuputToInput(TexCoordAppend, TEXT("A"), TexCoordR->GetUniqueID());
-	UInterchangeShaderPortsAPI::ConnectDefaultOuputToInput(TexCoordAppend, TEXT("B"), OneMinusG->GetUniqueID());
+	UInterchangeShaderNode* TexCoordAppend = CreateShaderNode(Connect.UpstreamNode, (Connect.UpstreamNode->getName() + "Append").c_str(), UE::Expressions::Names::AppendVector);
+	UInterchangeShaderPortsAPI::ConnectDefaultOuputToInput(TexCoordAppend, UE::Expressions::Inputs::A, TexCoordR->GetUniqueID());
+	UInterchangeShaderPortsAPI::ConnectDefaultOuputToInput(TexCoordAppend, UE::Expressions::Inputs::B, OneMinusG->GetUniqueID());
 
 	UInterchangeShaderPortsAPI::ConnectDefaultOuputToInput(Connect.ParentShaderNode, Connect.InputChannelName, TexCoordAppend->GetUniqueID());
 }
@@ -1272,6 +1394,10 @@ void FMaterialXSurfaceShaderAbstract::ConnectNormalMapInputToOutput(const FConne
 {
 	using namespace UE::Interchange::Materials::Standard::Nodes;
 
+	// FlattenNormalMap Material Function doesn't have normal/tangent inputs, we just remove them to avoid unnecessary connections
+	Connect.UpstreamNode->removeInput("normal");
+	Connect.UpstreamNode->removeInput("tangent");
+
 	// Only create a FunctionCall if there's a scale and the value is not 1, otherwise just like dot
 	mx::InputPtr Input = Connect.UpstreamNode->getInput("scale");
 	bool bIsNotEqualOne = true;
@@ -1318,9 +1444,13 @@ void FMaterialXSurfaceShaderAbstract::ConnectRefractInputToOutput(const FConnect
 	{
 		InputIorTarget = Connect.UpstreamNode->addInput("ior_target");
 		InputIorTarget->setValue("1", "float");
-		SetAttributeNewName(InputIorTarget, "Refractive Index Target");
+		SetAttributeNewName(InputIorTarget, TCHAR_TO_ANSI(UE::Expressions::Inputs::RefractiveIndexTarget));
 	}
-	AddAttributeFromValueOrInterface(InputIorTarget, TEXT("Refractive Index Target"), RefractNode);
+	AddAttributeFromValueOrInterface(InputIorTarget, UE::Expressions::Inputs::RefractiveIndexTarget, RefractNode);
+}
+
+void FMaterialXSurfaceShaderAbstract::ConnectBumpInputToOutput(const FConnectNode& Connect)
+{
 }
 
 UInterchangeShaderNode* FMaterialXSurfaceShaderAbstract::CreateMaskShaderNode(uint8 RGBA, mx::ElementPtr Element, const FString& NodeName, const FString& OutputName)
@@ -1550,8 +1680,11 @@ void FMaterialXSurfaceShaderAbstract::RegisterConnectNodeOutputToInputDelegates(
 	MatchingConnectNodeDelegates.Add(mx::Category::Tangent,			FOnConnectNodeOutputToInput::CreateSP(this, &FMaterialXSurfaceShaderAbstract::ConnectTangentInputToOutput));
 	MatchingConnectNodeDelegates.Add(mx::Category::Bitangent,		FOnConnectNodeOutputToInput::CreateSP(this, &FMaterialXSurfaceShaderAbstract::ConnectBitangentInputToOutput));
 	MatchingConnectNodeDelegates.Add(mx::Category::Time,			FOnConnectNodeOutputToInput::CreateSP(this, &FMaterialXSurfaceShaderAbstract::ConnectTimeInputToOutput));
+	MatchingConnectNodeDelegates.Add(mx::Category::Noise2D,			FOnConnectNodeOutputToInput::CreateSP(this, &FMaterialXSurfaceShaderAbstract::ConnectNoise2DInputToOutput));
 	MatchingConnectNodeDelegates.Add(mx::Category::Noise3D,			FOnConnectNodeOutputToInput::CreateSP(this, &FMaterialXSurfaceShaderAbstract::ConnectNoise3DInputToOutput));
+	MatchingConnectNodeDelegates.Add(mx::Category::CellNoise2D,     FOnConnectNodeOutputToInput::CreateSP(this, &FMaterialXSurfaceShaderAbstract::ConnectCellNoise2DInputToOutput));
 	MatchingConnectNodeDelegates.Add(mx::Category::CellNoise3D,		FOnConnectNodeOutputToInput::CreateSP(this, &FMaterialXSurfaceShaderAbstract::ConnectCellNoise3DInputToOutput));
+	MatchingConnectNodeDelegates.Add(mx::Category::WorleyNoise2D,	FOnConnectNodeOutputToInput::CreateSP(this, &FMaterialXSurfaceShaderAbstract::ConnectWorleyNoise2DInputToOutput));
 	MatchingConnectNodeDelegates.Add(mx::Category::WorleyNoise3D,	FOnConnectNodeOutputToInput::CreateSP(this, &FMaterialXSurfaceShaderAbstract::ConnectWorleyNoise3DInputToOutput));
 	MatchingConnectNodeDelegates.Add(mx::Category::Blur,			FOnConnectNodeOutputToInput::CreateSP(this, &FMaterialXSurfaceShaderAbstract::ConnectBlurInputToOutput));
 	MatchingConnectNodeDelegates.Add(mx::Category::HeightToNormal,	FOnConnectNodeOutputToInput::CreateSP(this, &FMaterialXSurfaceShaderAbstract::ConnectHeightToNormalInputToOutput));
diff --git a/Engine/Plugins/Interchange/Runtime/Source/Import/Public/MaterialX/MaterialXUtils/MaterialXManager.h b/Engine/Plugins/Interchange/Runtime/Source/Import/Public/MaterialX/MaterialXUtils/MaterialXManager.h
--- a/Engine/Plugins/Interchange/Runtime/Source/Import/Public/MaterialX/MaterialXUtils/MaterialXManager.h
+++ b/Engine/Plugins/Interchange/Runtime/Source/Import/Public/MaterialX/MaterialXUtils/MaterialXManager.h
@@ -84,21 +84,24 @@ private:
 
 	FMaterialXManager();
 
+	/** Helper function that returns the input as a string by also adding it in the MaterialExpressionInputs container*/
+	const TCHAR* ExpressionInput(const TCHAR* Input);
+
+	/** The different inputs of material expression that we may encounter, the MaterialX Document is modified consequently regarding those*/
+	TSet<FString> MaterialExpressionInputs;
+
 	/** Given a MaterialX node (category (optionally a nodegroup) - input), return the UE/Interchange input name.*/
 	TMap<TPair<FKeyCategoryNodegroup, FString>, FString> MatchingInputNames;
 
 	/** Given a MaterialX node category, optionally with a node group, return the UE material expression class name*/
 	TMap<FKeyCategoryNodegroup, FString> MatchingMaterialExpressions;
 
-	/** Given a MaterialX node category, return the UE material function, used for BSDF nodes*/
-	TMap<FString, FMaterialXMaterialFunction> MatchingMaterialFunctions;
-
-	/** The different inputs of material expression that we may encounter, the MaterialX Document is modified consequently regarding those*/
-	TSet<FString> MaterialExpressionInputs;
-
 	/** Container of a MaterialX document, to translate the different nodes based on their category*/
 	TMap<FString, FOnGetMaterialXInstance> MaterialXContainerDelegates;
 
+	/** Given a MaterialX node category, return the UE material function, used for BSDF nodes*/
+	TMap<FString, FMaterialXMaterialFunction> MatchingMaterialFunctions;
+
 	bool bIsSubstrateEnabled;
 };
 #endif
diff --git a/Engine/Plugins/Interchange/Runtime/Source/Import/Public/MaterialX/MaterialXUtils/MaterialXSurfaceShaderAbstract.h b/Engine/Plugins/Interchange/Runtime/Source/Import/Public/MaterialX/MaterialXUtils/MaterialXSurfaceShaderAbstract.h
--- a/Engine/Plugins/Interchange/Runtime/Source/Import/Public/MaterialX/MaterialXUtils/MaterialXSurfaceShaderAbstract.h
+++ b/Engine/Plugins/Interchange/Runtime/Source/Import/Public/MaterialX/MaterialXUtils/MaterialXSurfaceShaderAbstract.h
@@ -10,6 +10,8 @@
 #include "MaterialXBase.h"
 #include "MaterialXManager.h"
 
+enum EVectorNoiseFunction : int;
+
 class FMaterialXSurfaceShaderAbstract : public FMaterialXBase
 {
 protected:
@@ -259,6 +261,15 @@ protected:
 	/** <time> */
 	void ConnectTimeInputToOutput(const FConnectNode& Connect);
 
+	/** <noise2d> */
+	void ConnectNoise2DInputToOutput(const FConnectNode& Connect);
+
+	/** <cellnoise2d> */
+	void ConnectCellNoise2DInputToOutput(const FConnectNode& Connect);
+
+	/** <worleynoise2d> */
+	void ConnectWorleyNoise2DInputToOutput(const FConnectNode& Connect);
+
 	/** <noise3d> */
 	void ConnectNoise3DInputToOutput(const FConnectNode& Connect);
 
@@ -292,6 +303,9 @@ protected:
 	/** <refract> */
 	void ConnectRefractInputToOutput(const FConnectNode& Connect);
 
+	/** <bump> */
+	void ConnectBumpInputToOutput(const FConnectNode& Connect);
+
 	/** End Connect MaterialX nodes*/
 
 	/**
@@ -497,7 +511,7 @@ protected:
 private:
 	
 	UInterchangeShaderNode* ConnectGeometryInputToOutput(const FConnectNode& Connect, const FString& ShaderType, const FString& TransformShaderType, const FString& TransformInput, const FString& TransformSourceType, int32 TransformSource, const FString& TransformType, int32 TransformSDestination, bool bIsVector = true);
-
+	UInterchangeShaderNode* ConnectNoise2DInputToOutput(const FConnectNode& Connect, const FString& ShaderType, EVectorNoiseFunction NoiseFunction, uint8 Mask);
 protected:
 
 	/** 
