[Interchange - glTF] Adding SourceTimelineStart/End support.

#rnx

diff --git a/Engine/Plugins/Interchange/Runtime/Source/Import/Private/Gltf/InterchangeGltfTranslator.cpp b/Engine/Plugins/Interchange/Runtime/Source/Import/Private/Gltf/InterchangeGltfTranslator.cpp
--- a/Engine/Plugins/Interchange/Runtime/Source/Import/Private/Gltf/InterchangeGltfTranslator.cpp
+++ b/Engine/Plugins/Interchange/Runtime/Source/Import/Private/Gltf/InterchangeGltfTranslator.cpp
@@ -509,6 +509,45 @@ bool UInterchangeGLTFTranslator::Translate( UInterchangeBaseNodeContainer& NodeC
 		SourceNode->SetExtraInformation(Extra.Name, Extra.Value);
 	}
 
+	//Set Animation Time related values:
+	if (GltfAsset.Animations.Num())
+	{
+		// glTF animations are tied to seconds, not framerate
+		// glTF animations do no support Time Line.
+		// glTF animations always start at 0 seconds
+		SourceNode->SetCustomSourceTimelineStart(0);
+		
+		//To acquire the AnimatedTimeEnd:
+		//	Find the last entry of each animation.samplers and find the biggest entry (aka the time (in seconds) which is the biggest).
+		float AnimatedTimeEnd = 0;
+		for (const GLTF::FAnimation& Animation : GltfAsset.Animations)
+		{
+			for (const GLTF::FAnimation::FSampler& Sampler : Animation.Samplers)
+			{
+				TArray<float> FrameTimeBuffer;
+
+				Sampler.Input.GetFloatArray(FrameTimeBuffer);
+
+				if (AnimatedTimeEnd < FrameTimeBuffer.Last())
+				{
+					AnimatedTimeEnd = FrameTimeBuffer.Last();
+				}
+			}
+		}
+
+		//Making sure that all animations can fit into the FrameNumbers:
+		//Currently glTF Animations are calcualted with 30.0 frames (as glTF does not provide Frequency) :
+		constexpr float FrameRate = 30.0;
+		constexpr float SingleFrameDuration = 1.0 / FrameRate;
+
+		float CurrentFrameNumberCandidate = AnimatedTimeEnd / SingleFrameDuration;
+		int32 CurrentFrameNumber = FMath::CeilToInt(CurrentFrameNumberCandidate);
+
+		AnimatedTimeEnd = CurrentFrameNumber * SingleFrameDuration;
+
+		SourceNode->SetCustomSourceTimelineEnd(AnimatedTimeEnd);
+	}
+
 	{
 		//For sockets, crete the Axis System conversion inverse:
 
