[Interchange] Fix Blender Armature bone to behave like legacy fbx

diff --git a/Engine/Plugins/Interchange/Runtime/Source/Import/Private/Mesh/InterchangeSkeletalMeshFactory.cpp b/Engine/Plugins/Interchange/Runtime/Source/Import/Private/Mesh/InterchangeSkeletalMeshFactory.cpp
--- a/Engine/Plugins/Interchange/Runtime/Source/Import/Private/Mesh/InterchangeSkeletalMeshFactory.cpp
+++ b/Engine/Plugins/Interchange/Runtime/Source/Import/Private/Mesh/InterchangeSkeletalMeshFactory.cpp
@@ -40,6 +40,7 @@
 #include "Misc/ScopedSlowTask.h"
 #include "Nodes/InterchangeBaseNode.h"
 #include "Nodes/InterchangeBaseNodeContainer.h"
+#include "Nodes/InterchangeSourceNode.h"
 #include "PhysicsEngine/PhysicsAsset.h"
 #include "Rendering/SkeletalMeshLODImporterData.h"
 #include "Rendering/SkeletalMeshLODModel.h"
@@ -1059,6 +1060,18 @@ namespace UE
 				FTransform BakeToRootJointTransfromModifier = RootJointNodeGlobalTransform.Inverse() * RootJointNodeLocalTransform; //It is used for !bBakeMeshes, and so the GlobalTransform will be inversed out when multiplied into the CustomBindPoseGlobalTransform
 				FTransform BakeFromRootJointTransfromModifier = RootJointNodeLocalTransform.Inverse() * RootJointNodeGlobalTransform * GlobalOffsetTransform.Inverse(); //GlobalOffseTransform will be added by the BindPoseGlobalTransform when bBakeMeshes && !bRootAncestorOfMeshDependency is used
 
+				{
+					// Temporary 5.5.2: fbx legacy has a special way to bake the skeletal mesh that do not fit the interchange standard.
+					// This fix the issue with blender armature bone skip if the armature has transform not identity.
+					const UInterchangeSourceNode* SourceNode = UInterchangeSourceNode::GetUniqueInstance(Arguments.NodeContainer);
+					bool bUseSceneNodeGlobalTransform = false;
+					SourceNode->GetBooleanAttribute(TEXT("SkeletalMesh-Bake-UseSceneNodeGlobalTransform"), bUseSceneNodeGlobalTransform);
+					if (bUseSceneNodeGlobalTransform)
+					{
+						BakeFromRootJointTransfromModifier.SetIdentity();
+					}
+				}
+
 				//Scope to query the mesh node
 				{
 					TArray<FString> MeshUids;
diff --git a/Engine/Plugins/Interchange/Runtime/Source/Parsers/Fbx/Private/FbxAPI.cpp b/Engine/Plugins/Interchange/Runtime/Source/Parsers/Fbx/Private/FbxAPI.cpp
--- a/Engine/Plugins/Interchange/Runtime/Source/Parsers/Fbx/Private/FbxAPI.cpp
+++ b/Engine/Plugins/Interchange/Runtime/Source/Parsers/Fbx/Private/FbxAPI.cpp
@@ -163,6 +163,11 @@ namespace UE
 					SourceNode->SetCustomSourceTimelineEnd(TimelineTimeSpan.GetStop().GetSecondDouble());
 				}
 
+				// Temporary 5.5.2: fbx legacy has a special way to bake the skeletal mesh that do not fit the interchange standard
+				// The interchange skeletal mesh factory will read this to use the proper bake transform so it match legacy behavior.
+				// This fix the issue with blender armature bone skip
+				SourceNode->AddBooleanAttribute(TEXT("SkeletalMesh-Bake-UseSceneNodeGlobalTransform"), true);
+
 				// Get the version number of the FBX file format.
 				int32 FileMajor, FileMinor, FileRevision;
 				SDKImporter->GetFileVersion(FileMajor, FileMinor, FileRevision);
