[Interchange]: Make sure fbx bone animation and morph target animation on the same fbx animtrack index are on the same anim sequence.

diff --git a/Engine/Plugins/Interchange/Runtime/Source/Parsers/Fbx/Private/FbxScene.cpp b/Engine/Plugins/Interchange/Runtime/Source/Parsers/Fbx/Private/FbxScene.cpp
--- a/Engine/Plugins/Interchange/Runtime/Source/Parsers/Fbx/Private/FbxScene.cpp
+++ b/Engine/Plugins/Interchange/Runtime/Source/Parsers/Fbx/Private/FbxScene.cpp
@@ -47,6 +47,13 @@ namespace UE
 					}
 				}
 			} // ns RecursiveHelper
+
+			FString CreateTrackNodeUid(const FString& JointUid, const int32 AnimationIndex)
+			{
+				FString TrackNodeUid = TEXT("\\SkeletalAnimation\\") + JointUid + TEXT("\\") + FString::FromInt(AnimationIndex);
+				return TrackNodeUid;
+			}
+
 			void FFbxScene::CreateMeshNodeReference(UInterchangeSceneNode* UnrealSceneNode, FbxNodeAttribute* NodeAttribute, UInterchangeBaseNodeContainer& NodeContainer, const FTransform& GeometricTransform, const FTransform& PivotNodeTransform)
 			{
 				const UInterchangeMeshNode* MeshNode = nullptr;
@@ -850,9 +857,23 @@ namespace UE
 							FString DisplayString = FFbxConvert::MakeString(CurrentAnimationStack->GetName());
 							constexpr bool bIsJointFalse = false;
 							UE::Interchange::SanitizeName(DisplayString, bIsJointFalse);
-							SkeletalAnimationTrackNode = NewObject< UInterchangeSkeletalAnimationTrackNode >(&NodeContainer);
-							FString TrackNodeUid = "\\SkeletalAnimation\\" + UnrealNode->GetUniqueID() + "_" + FString::FromInt(AnimationIndex);
-							SkeletalAnimationTrackNode->InitializeNode(TrackNodeUid, DisplayString, EInterchangeNodeContainerType::TranslatedAsset);
+							FString TrackNodeUid = CreateTrackNodeUid(UnrealNode->GetUniqueID(), AnimationIndex);
+							NodeContainer.BreakableIterateNodesOfType<UInterchangeSkeletalAnimationTrackNode>([&TrackNodeUid, &SkeletalAnimationTrackNode](const FString& NodeUid, UInterchangeSkeletalAnimationTrackNode* Node)
+							{
+								if (NodeUid == TrackNodeUid)
+								{
+									SkeletalAnimationTrackNode = Node;
+									return true;
+								}
+								return false;
+							});
+
+							if(!SkeletalAnimationTrackNode)
+							{
+								SkeletalAnimationTrackNode = NewObject< UInterchangeSkeletalAnimationTrackNode >(&NodeContainer);
+								TrackNodeUid = "\\SkeletalAnimation\\" + UnrealNode->GetUniqueID() + "_" + FString::FromInt(AnimationIndex);
+								SkeletalAnimationTrackNode->InitializeNode(TrackNodeUid, DisplayString, EInterchangeNodeContainerType::TranslatedAsset);
+							}
 							SkeletalAnimationTrackNode->AddBooleanAttribute(TEXT("RenameLikeLegacyFbx"), true);
 
 							double FrameRate = FbxTime::GetFrameRate(SDKScene->GetGlobalSettings().GetTimeMode());
@@ -1192,18 +1213,34 @@ namespace UE
 						FbxAnimStack* CurrentAnimationStack = (FbxAnimStack*)SDKScene->GetSrcObject<FbxAnimStack>(AnimationIndex);
 						FbxTimeSpan TimeSpan = CurrentAnimationStack->GetLocalTimeSpan();
 
-						UInterchangeSkeletalAnimationTrackNode* SkeletalAnimationTrackNode = NewObject< UInterchangeSkeletalAnimationTrackNode >(&NodeContainer);
-						FString TrackNodeUid = "\\MorphTargetAnimation\\" + SkeletonNode->GetUniqueID() + TEXT("\\") + FString::FromInt(AnimationIndex);
-						FString DisplayString = SkeletonNode->GetDisplayLabel() + TEXT("_MorphAnim_") + FString::FromInt(AnimationIndex);
-						SkeletalAnimationTrackNode->InitializeNode(TrackNodeUid, DisplayString, EInterchangeNodeContainerType::TranslatedAsset);
+						UInterchangeSkeletalAnimationTrackNode* SkeletalAnimationTrackNode = nullptr;
+						FString TrackNodeUid = CreateTrackNodeUid(SkeletonNode->GetUniqueID(), AnimationIndex);
+						NodeContainer.BreakableIterateNodesOfType<UInterchangeSkeletalAnimationTrackNode>([&TrackNodeUid, &SkeletalAnimationTrackNode](const FString& NodeUid, UInterchangeSkeletalAnimationTrackNode* Node)
+						{
+							if (NodeUid == TrackNodeUid)
+							{
+								SkeletalAnimationTrackNode = Node;
+								return true;
+							}
+							return false;
+						});
+
+						if(!SkeletalAnimationTrackNode)
+						{
+							SkeletalAnimationTrackNode = NewObject< UInterchangeSkeletalAnimationTrackNode >(&NodeContainer);
+							FString DisplayString = FFbxConvert::MakeString(CurrentAnimationStack->GetName());
+							SkeletalAnimationTrackNode->InitializeNode(TrackNodeUid, DisplayString, EInterchangeNodeContainerType::TranslatedAsset);
 
-						double FrameRate = FbxTime::GetFrameRate(SDKScene->GetGlobalSettings().GetTimeMode());
-						SkeletalAnimationTrackNode->SetCustomAnimationSampleRate(FrameRate);
+							double FrameRate = FbxTime::GetFrameRate(SDKScene->GetGlobalSettings().GetTimeMode());
+							SkeletalAnimationTrackNode->SetCustomAnimationSampleRate(FrameRate);
 
-						SkeletalAnimationTrackNode->SetCustomSkeletonNodeUid(SkeletonNode->GetUniqueID());
+							SkeletalAnimationTrackNode->SetCustomSkeletonNodeUid(SkeletonNode->GetUniqueID());
 
-						SkeletalAnimationTrackNode->SetCustomAnimationStartTime(TimeSpan.GetStart().GetSecondDouble());
-						SkeletalAnimationTrackNode->SetCustomAnimationStopTime(TimeSpan.GetStop().GetSecondDouble());
+							SkeletalAnimationTrackNode->SetCustomAnimationStartTime(TimeSpan.GetStart().GetSecondDouble());
+							SkeletalAnimationTrackNode->SetCustomAnimationStopTime(TimeSpan.GetStop().GetSecondDouble());
+
+							SkeletalAnimationTrackNode->AddBooleanAttribute(TEXT("RenameLikeLegacyFbx"), true);
+						}
 
 						ProcessCustomAttributes(Parser, CurrentAnimationStack, SkeletalAnimationTrackNode);
 
