[Interchange - glTF] Adding support for up to 8 TexCoords. Fixing SkeletalMesh multi UV imports.

#rnx

diff --git a/Engine/Plugins/Interchange/Runtime/Source/Import/Private/Gltf/InterchangeGltfMesh.cpp b/Engine/Plugins/Interchange/Runtime/Source/Import/Private/Gltf/InterchangeGltfMesh.cpp
--- a/Engine/Plugins/Interchange/Runtime/Source/Import/Private/Gltf/InterchangeGltfMesh.cpp
+++ b/Engine/Plugins/Interchange/Runtime/Source/Import/Private/Gltf/InterchangeGltfMesh.cpp
@@ -130,7 +130,7 @@ namespace UE::Interchange::Gltf::Private
 			const GLTF::FSkinInfo& Skin = GltfAsset.Skins[SkinIndex];
 
 			FSkeletalMeshAttributes SkeletalMeshAttributes(SkinnedMeshDescription);
-			SkeletalMeshAttributes.Register();
+			SkeletalMeshAttributes.Register(true);
 
 			using namespace UE::AnimationCore;
 			TMap<FVertexID, TArray<FBoneWeight>> RawBoneWeights;
diff --git a/Engine/Plugins/Interchange/Runtime/Source/Parsers/GLTFCore/Private/GLTF/GLTFAccessor.cpp b/Engine/Plugins/Interchange/Runtime/Source/Parsers/GLTFCore/Private/GLTF/GLTFAccessor.cpp
--- a/Engine/Plugins/Interchange/Runtime/Source/Parsers/GLTFCore/Private/GLTF/GLTFAccessor.cpp
+++ b/Engine/Plugins/Interchange/Runtime/Source/Parsers/GLTFCore/Private/GLTF/GLTFAccessor.cpp
@@ -20,6 +20,18 @@ namespace GLTF
 			return TEXT("TEXCOORD_0");
 		case GLTF::TEXCOORD_1:
 			return TEXT("TEXCOORD_1");
+		case GLTF::TEXCOORD_2:
+			return TEXT("TEXCOORD_2");
+		case GLTF::TEXCOORD_3:
+			return TEXT("TEXCOORD_3");
+		case GLTF::TEXCOORD_4:
+			return TEXT("TEXCOORD_4");
+		case GLTF::TEXCOORD_5:
+			return TEXT("TEXCOORD_5");
+		case GLTF::TEXCOORD_6:
+			return TEXT("TEXCOORD_6");
+		case GLTF::TEXCOORD_7:
+			return TEXT("TEXCOORD_7");
 		case GLTF::COLOR_0:
 			return TEXT("COLOR_0");
 		case GLTF::JOINTS_0:
@@ -582,6 +594,12 @@ namespace GLTF
 				return (!bMorphTargetProperty && Type == FAccessor::EType::Vec4) || (bMorphTargetProperty && Type == FAccessor::EType::Vec3);
 			case GLTF::EMeshAttributeType::TEXCOORD_0:
 			case GLTF::EMeshAttributeType::TEXCOORD_1:
+			case GLTF::EMeshAttributeType::TEXCOORD_2:
+			case GLTF::EMeshAttributeType::TEXCOORD_3:
+			case GLTF::EMeshAttributeType::TEXCOORD_4:
+			case GLTF::EMeshAttributeType::TEXCOORD_5:
+			case GLTF::EMeshAttributeType::TEXCOORD_6:
+			case GLTF::EMeshAttributeType::TEXCOORD_7:
 				return Type == FAccessor::EType::Vec2;
 			case GLTF::EMeshAttributeType::COLOR_0:
 				return Type == FAccessor::EType::Vec3 || Type == FAccessor::EType::Vec4;
@@ -609,6 +627,12 @@ namespace GLTF
 					return !bNormalized && ComponentType == FAccessor::EComponentType::F32;
 				case GLTF::EMeshAttributeType::TEXCOORD_0:
 				case GLTF::EMeshAttributeType::TEXCOORD_1:
+				case GLTF::EMeshAttributeType::TEXCOORD_2:
+				case GLTF::EMeshAttributeType::TEXCOORD_3:
+				case GLTF::EMeshAttributeType::TEXCOORD_4:
+				case GLTF::EMeshAttributeType::TEXCOORD_5:
+				case GLTF::EMeshAttributeType::TEXCOORD_6:
+				case GLTF::EMeshAttributeType::TEXCOORD_7:
 				case GLTF::EMeshAttributeType::COLOR_0:
 					return (!bNormalized && ComponentType == FAccessor::EComponentType::F32) 
 						|| (bNormalized && (ComponentType == FAccessor::EComponentType::U8 || ComponentType == FAccessor::EComponentType::U16));
@@ -633,6 +657,12 @@ namespace GLTF
 					return !bNormalized && ComponentType == FAccessor::EComponentType::F32;
 				case GLTF::EMeshAttributeType::TEXCOORD_0:
 				case GLTF::EMeshAttributeType::TEXCOORD_1:
+				case GLTF::EMeshAttributeType::TEXCOORD_2:
+				case GLTF::EMeshAttributeType::TEXCOORD_3:
+				case GLTF::EMeshAttributeType::TEXCOORD_4:
+				case GLTF::EMeshAttributeType::TEXCOORD_5:
+				case GLTF::EMeshAttributeType::TEXCOORD_6:
+				case GLTF::EMeshAttributeType::TEXCOORD_7:
 				case GLTF::EMeshAttributeType::COLOR_0:
 					return (!bNormalized && ComponentType == FAccessor::EComponentType::F32)
 						|| (bNormalized && (ComponentType == FAccessor::EComponentType::S8 || ComponentType == FAccessor::EComponentType::S16 || ComponentType == FAccessor::EComponentType::U8 || ComponentType == FAccessor::EComponentType::U16));
@@ -659,6 +689,12 @@ namespace GLTF
 					return bNormalized && (ComponentType == FAccessor::EComponentType::S8 || ComponentType == FAccessor::EComponentType::S16);
 				case GLTF::EMeshAttributeType::TEXCOORD_0:
 				case GLTF::EMeshAttributeType::TEXCOORD_1:
+				case GLTF::EMeshAttributeType::TEXCOORD_2:
+				case GLTF::EMeshAttributeType::TEXCOORD_3:
+				case GLTF::EMeshAttributeType::TEXCOORD_4:
+				case GLTF::EMeshAttributeType::TEXCOORD_5:
+				case GLTF::EMeshAttributeType::TEXCOORD_6:
+				case GLTF::EMeshAttributeType::TEXCOORD_7:
 					return (ComponentType == FAccessor::EComponentType::S8 || ComponentType == FAccessor::EComponentType::S16)
 						|| (!bNormalized && (ComponentType == FAccessor::EComponentType::U8 || ComponentType == FAccessor::EComponentType::U16));
 				default:
@@ -678,6 +714,12 @@ namespace GLTF
 					return bNormalized && (ComponentType == FAccessor::EComponentType::S8 || ComponentType == FAccessor::EComponentType::S16);
 				case GLTF::EMeshAttributeType::TEXCOORD_0:
 				case GLTF::EMeshAttributeType::TEXCOORD_1:
+				case GLTF::EMeshAttributeType::TEXCOORD_2:
+				case GLTF::EMeshAttributeType::TEXCOORD_3:
+				case GLTF::EMeshAttributeType::TEXCOORD_4:
+				case GLTF::EMeshAttributeType::TEXCOORD_5:
+				case GLTF::EMeshAttributeType::TEXCOORD_6:
+				case GLTF::EMeshAttributeType::TEXCOORD_7:
 					return !bNormalized && (ComponentType == FAccessor::EComponentType::S8 || ComponentType == FAccessor::EComponentType::S16);
 				default:
 					ensure(false);
diff --git a/Engine/Plugins/Interchange/Runtime/Source/Parsers/GLTFCore/Private/GLTF/GLTFMesh.cpp b/Engine/Plugins/Interchange/Runtime/Source/Parsers/GLTFCore/Private/GLTF/GLTFMesh.cpp
--- a/Engine/Plugins/Interchange/Runtime/Source/Parsers/GLTFCore/Private/GLTF/GLTFMesh.cpp
+++ b/Engine/Plugins/Interchange/Runtime/Source/Parsers/GLTFCore/Private/GLTF/GLTFMesh.cpp
@@ -6,37 +6,34 @@
 
 namespace GLTF
 {
-	//For MorphTargets (morph targets do not support Joints and Weights)
-	FAttributeAccessors::FAttributeAccessors(const FAccessor& InPositionDisplacements, const FAccessor& InNormalDisplacements,
-												const FAccessor& InTangentDisplacements,
-												const FAccessor& InTexCoord0Displacements, const FAccessor& InTexCoord1Displacements,
-												const FAccessor& InColor0Deltas)
-		: bMorphTarget(true)
-	{
-		AttributeAccessors.Emplace(EMeshAttributeType::POSITION, InPositionDisplacements);
-		AttributeAccessors.Emplace(EMeshAttributeType::NORMAL, InNormalDisplacements);
-		AttributeAccessors.Emplace(EMeshAttributeType::TANGENT, InTangentDisplacements);
-		AttributeAccessors.Emplace(EMeshAttributeType::TEXCOORD_0, InTexCoord0Displacements);
-		AttributeAccessors.Emplace(EMeshAttributeType::TEXCOORD_1, InTexCoord1Displacements);
-		AttributeAccessors.Emplace(EMeshAttributeType::COLOR_0, InColor0Deltas);
-	}
-
-	//For Primitives
-	FAttributeAccessors::FAttributeAccessors(const FAccessor& InPosition, const FAccessor& InNormal,
-												const FAccessor& InTangent,
-												const FAccessor& InTexCoord0, const FAccessor& InTexCoord1,
-												const FAccessor& InColor0,
-												const FAccessor& InJoints0, const FAccessor& InWeights0)
-		: bMorphTarget(false)
-	{
-		AttributeAccessors.Emplace(EMeshAttributeType::POSITION, InPosition);
-		AttributeAccessors.Emplace(EMeshAttributeType::NORMAL, InNormal);
-		AttributeAccessors.Emplace(EMeshAttributeType::TANGENT, InTangent);
-		AttributeAccessors.Emplace(EMeshAttributeType::TEXCOORD_0, InTexCoord0);
-		AttributeAccessors.Emplace(EMeshAttributeType::TEXCOORD_1, InTexCoord1);
-		AttributeAccessors.Emplace(EMeshAttributeType::COLOR_0, InColor0);
-		AttributeAccessors.Emplace(EMeshAttributeType::JOINTS_0, InJoints0);
-		AttributeAccessors.Emplace(EMeshAttributeType::WEIGHTS_0, InWeights0);
+	FAttributeAccessors::FAttributeAccessors(const TArray<TPair<EMeshAttributeType, FAccessor&>>& InAttributeAccessors, bool bInMorphTarget)
+	{
+		
+		for (uint8 AttributeTypeCounter = 0; AttributeTypeCounter < EMeshAttributeType::COUNT; AttributeTypeCounter++)
+		{
+			EMeshAttributeType AttributeType = EMeshAttributeType(AttributeTypeCounter);
+
+			if (InAttributeAccessors.Num() <= AttributeTypeCounter)
+			{
+				break;
+			}
+
+			if (InAttributeAccessors[AttributeType].Key == AttributeType)
+			{
+				AttributeAccessors.Emplace(AttributeType, InAttributeAccessors[AttributeType].Value);
+			}
+			else
+			{
+				//Expectation mismatch, update InAttributeAccessors array based on EMeshAttributeType.
+				ensure(false);
+
+				static FAccessor EmptyAccessor;
+
+				AttributeAccessors.Emplace(AttributeType, EmptyAccessor);
+			}
+		}
+
+		bMorphTarget = bInMorphTarget;
 	}
 
 	bool FAttributeAccessors::HasAttributeAccessor(const EMeshAttributeType& Type) const
@@ -46,9 +43,6 @@ namespace GLTF
 			return AttributeAccessors[Type].Value.IsValid();
 		}
 
-		//Expectation missmatch update AttributeAccessors array based on EMeshAttributeType.
-		ensure(false);
-
 		return false;
 	}
 	
@@ -59,9 +53,6 @@ namespace GLTF
 			return AttributeAccessors[Type].Value;
 		}
 
-		//Expectation missmatch update AttributeAccessors array based on EMeshAttributeType.
-		ensure(false);
-
 		static FAccessor EmptyAccessor;
 		return EmptyAccessor;
 	}
@@ -112,9 +103,8 @@ namespace GLTF
 		}
 	}
 
-	FMorphTarget::FMorphTarget(const FAccessor& InPositionDisplacements, const FAccessor& InNormalDisplacements,
-		const FAccessor& InTangentDisplacements, const FAccessor& InTexCoord0Displacements, const FAccessor& InTexCoord1Displacements, const FAccessor& InColor0Deltas)
-		: FAttributeAccessors(InPositionDisplacements, InNormalDisplacements, InTangentDisplacements, InTexCoord0Displacements, InTexCoord1Displacements, InColor0Deltas)
+	FMorphTarget::FMorphTarget(const TArray<TPair<EMeshAttributeType, FAccessor&>>& InAttributeAccessors)
+		: FAttributeAccessors(InAttributeAccessors, true)
 	{
 	}
 
@@ -178,39 +168,31 @@ namespace GLTF
 
 	bool FMorphTarget::HasTexCoordDisplacements(uint32 Index) const
 	{
-		switch (Index)
+		if (Index < SupportedTexCoordNum)
 		{
-			case 0:
-				return HasAttributeAccessor(EMeshAttributeType::TEXCOORD_0);
-			case 1:
-				return HasAttributeAccessor(EMeshAttributeType::TEXCOORD_1);
-			default:
-				return false;
+			const uint32 TexcoordAttributeIndex = EMeshAttributeType::TEXCOORD_0 + Index;
+			return HasAttributeAccessor(EMeshAttributeType(TexcoordAttributeIndex));
 		}
+
+		return false;
 	}
 	void FMorphTarget::GetTexCoordDisplacements(uint32 Index, TArray<FVector2f>& Buffer) const
 	{
-		switch (Index)
+		if (Index < SupportedTexCoordNum)
 		{
-			case 0:
-				return GetAttributeAccessor(EMeshAttributeType::TEXCOORD_0).GetVec2Array(Buffer);
-			case 1:
-				return GetAttributeAccessor(EMeshAttributeType::TEXCOORD_1).GetVec2Array(Buffer);
-			default:
-				break;
+			const uint32 TexcoordAttributeIndex = EMeshAttributeType::TEXCOORD_0 + Index;
+			return GetAttributeAccessor(EMeshAttributeType(TexcoordAttributeIndex)).GetVec2Array(Buffer);
 		}
 	}
 	int32 FMorphTarget::GetNumberOfTexCoordDisplacements(uint32 Index) const
 	{
-		switch (Index)
+		if (Index < SupportedTexCoordNum)
 		{
-			case 0:
-				return GetAttributeAccessor(EMeshAttributeType::TEXCOORD_0).Count;
-			case 1:
-				return GetAttributeAccessor(EMeshAttributeType::TEXCOORD_1).Count;
-			default:
-				return 0;
+			const uint32 TexcoordAttributeIndex = EMeshAttributeType::TEXCOORD_0 + Index;
+			return GetAttributeAccessor(EMeshAttributeType(TexcoordAttributeIndex)).Count;
 		}
+
+		return 0;
 	}
 
 	void FMorphTarget::GetColorDeltas(TArray<FVector4f>& Buffer) const
@@ -256,10 +238,8 @@ namespace GLTF
 		}
 	}
 
-	FPrimitive::FPrimitive(EMode InMode, int32 InMaterial, const FAccessor& InIndices, const FAccessor& InPosition, const FAccessor& InNormal,
-	                       const FAccessor& InTangent, const FAccessor& InTexCoord0, const FAccessor& InTexCoord1, const FAccessor& InColor0,
-	                       const FAccessor& InJoints0, const FAccessor& InWeights0)
-	    : FAttributeAccessors(InPosition, InNormal, InTangent, InTexCoord0, InTexCoord1, InColor0, InJoints0, InWeights0)
+	FPrimitive::FPrimitive(EMode InMode, int32 InMaterial, const FAccessor& InIndices, const TArray<TPair<EMeshAttributeType, FAccessor&>>& InAttributeAccessors)
+	    : FAttributeAccessors(InAttributeAccessors, false)
 		, Mode(InMode)
 	    , MaterialIndex(InMaterial)
 	    , Indices(InIndices)
@@ -374,15 +354,14 @@ namespace GLTF
 
 	bool FPrimitive::HasTexCoords(uint32 Index) const
 	{
-		switch (Index)
+		if (Index < SupportedTexCoordNum)
 		{
-		case 0:
-			return HasAttributeAccessor(EMeshAttributeType::TEXCOORD_0);
-		case 1:
-			return HasAttributeAccessor(EMeshAttributeType::TEXCOORD_1);
-		default:
-			return false;
+			const uint32 TexcoordAttributeIndex = EMeshAttributeType::TEXCOORD_0 + Index;
+			return HasAttributeAccessor(EMeshAttributeType(TexcoordAttributeIndex));
 		}
+
+		return false;
+		
 	}
 
 	void FPrimitive::GetPositions(TArray<FVector3f>& Buffer) const
@@ -408,15 +387,10 @@ namespace GLTF
 
 	void FPrimitive::GetTexCoords(uint32 Index, TArray<FVector2f>& Buffer) const
 	{
-		switch (Index)
+		if (Index < SupportedTexCoordNum)
 		{
-		case 0:
-			return GetAttributeAccessor(EMeshAttributeType::TEXCOORD_0).GetVec2Array(Buffer);
-		case 1:
-			return GetAttributeAccessor(EMeshAttributeType::TEXCOORD_1).GetVec2Array(Buffer);
-		default:
-			ensure(false);
-			break;
+			const uint32 TexcoordAttributeIndex = EMeshAttributeType::TEXCOORD_0 + Index;
+			return GetAttributeAccessor(EMeshAttributeType(TexcoordAttributeIndex)).GetVec2Array(Buffer);
 		}
 	}
 
diff --git a/Engine/Plugins/Interchange/Runtime/Source/Parsers/GLTFCore/Private/GLTF/GLTFReader.cpp b/Engine/Plugins/Interchange/Runtime/Source/Parsers/GLTFCore/Private/GLTF/GLTFReader.cpp
--- a/Engine/Plugins/Interchange/Runtime/Source/Parsers/GLTFCore/Private/GLTF/GLTFReader.cpp
+++ b/Engine/Plugins/Interchange/Runtime/Source/Parsers/GLTFCore/Private/GLTF/GLTFReader.cpp
@@ -102,19 +102,6 @@ namespace GLTF
 			}
 		}
 
-		const FAccessor& AccessorAtIndex(const TArray<FAccessor>& Accessors, int32 Index)
-		{
-			if (Accessors.IsValidIndex(Index))
-			{
-				return Accessors[Index];
-			}
-			else
-			{
-				static const FAccessor EmptyAccessor;
-				return EmptyAccessor;
-			}
-		}
-
 		template<typename T>
 		void SetTransformFromMatrix(FTransform& Transform, const UE::Math::TMatrix<T>& InMatrix)
 		{
@@ -387,20 +374,24 @@ namespace GLTF
 	void FFileReader::SetupMorphTarget(const FJsonObject& Object, GLTF::FPrimitive& Primitive, const bool bMeshQuantized) const
 	{
 		const TArray<FAccessor>& A = Asset->Accessors;
-		FAccessor& Position = AccessorAtIndex(Asset->Accessors, GetIndex(Object, TEXT("POSITION")));
-		FAccessor& Normal = AccessorAtIndex(Asset->Accessors, GetIndex(Object, TEXT("NORMAL")));
-		FAccessor& Tangent = AccessorAtIndex(Asset->Accessors, GetIndex(Object, TEXT("TANGENT")));
-		FAccessor& TexCoord0 = AccessorAtIndex(Asset->Accessors, GetIndex(Object, TEXT("TEXCOORD_0")));
-		FAccessor& TexCoord1 = AccessorAtIndex(Asset->Accessors, GetIndex(Object, TEXT("TEXCOORD_1")));
-		const FAccessor& Color0 = AccessorAtIndex(A, GetIndex(Object, TEXT("COLOR_0")));
-
-		Position.bQuantized = bMeshQuantized;
-		Normal.bQuantized = bMeshQuantized;
-		Tangent.bQuantized = bMeshQuantized;
-		TexCoord0.bQuantized = bMeshQuantized;
-		TexCoord1.bQuantized = bMeshQuantized;
-
-		Primitive.MorphTargets.Emplace(Position, Normal, Tangent, TexCoord0, TexCoord1, Color0);
+
+		//Morph targets do not support joints/weights, aka they go only until Color_0
+		// all Accessors have to be updated to bQuantized except for Color_0
+
+		TArray<TPair<EMeshAttributeType, FAccessor&>> AttributeAccessors;
+		AttributeAccessors.Reserve(EMeshAttributeType::COLOR_0);
+		for (uint8 AttributeTypeCounter = 0; AttributeTypeCounter < EMeshAttributeType::COLOR_0; AttributeTypeCounter++)
+		{
+			EMeshAttributeType AttributeType = EMeshAttributeType(AttributeTypeCounter);
+			FAccessor& Accessor = AccessorAtIndex(Asset->Accessors, GetIndex(Object, *ToString(AttributeType)));
+			Accessor.bQuantized = bMeshQuantized;
+			AttributeAccessors.Emplace(AttributeType, Accessor);
+		}
+
+		FAccessor& AccessorColor0 = AccessorAtIndex(Asset->Accessors, GetIndex(Object, *ToString(EMeshAttributeType::COLOR_0)));
+		AttributeAccessors.Emplace(EMeshAttributeType::COLOR_0, AccessorColor0);
+
+		Primitive.MorphTargets.Emplace(AttributeAccessors);
 	}
 
 	void FFileReader::SetupPrimitive(const FJsonObject& Object, FMesh& Mesh, const bool bMeshQuantized, const uint32& PrimitiveIndex) const
@@ -420,29 +411,30 @@ namespace GLTF
 				FText::FromString(Mesh.Name)));
 		}
 
-		const int32                   MaterialIndex = GetIndex(Object, TEXT("material"));
-		const TArray<FAccessor>& A             = Asset->Accessors;
+		const int32 MaterialIndex = GetIndex(Object, TEXT("material"));
+		const FAccessor& Indices = AccessorAtIndex(Asset->Accessors, GetIndex(Object, TEXT("indices")));
+		
+		const FJsonObject& Attributes = *Object.GetObjectField(TEXT("attributes"));
+
+		TArray<TPair<EMeshAttributeType, FAccessor&>> AttributeAccessors;
+		AttributeAccessors.Reserve(EMeshAttributeType::COUNT);
+		for (uint8 AttributeTypeCounter = 0; AttributeTypeCounter < EMeshAttributeType::COLOR_0; AttributeTypeCounter++)
+		{
+			EMeshAttributeType AttributeType = EMeshAttributeType(AttributeTypeCounter);
+			FAccessor& Accessor = AccessorAtIndex(Asset->Accessors, GetIndex(Attributes, *ToString(AttributeType)));
+			Accessor.bQuantized = bMeshQuantized;
+			AttributeAccessors.Emplace(AttributeType, Accessor);
+		}
 
-		const FAccessor& Indices = AccessorAtIndex(A, GetIndex(Object, TEXT("indices")));
+		//Color_0 ... do not need the bMeshQuantized setting.
+		for (uint8 AttributeTypeCounter = EMeshAttributeType::COLOR_0; AttributeTypeCounter < EMeshAttributeType::COUNT; AttributeTypeCounter++)
+		{
+			EMeshAttributeType AttributeType = EMeshAttributeType(AttributeTypeCounter);
+			FAccessor& Accessor = AccessorAtIndex(Asset->Accessors, GetIndex(Attributes, *ToString(AttributeType)));
+			AttributeAccessors.Emplace(AttributeType, Accessor);
+		}
 
-		// the only required attribute is POSITION
-		const FJsonObject& Attributes = *Object.GetObjectField(TEXT("attributes"));
-		FAccessor&   Position   = AccessorAtIndex(Asset->Accessors, GetIndex(Attributes, TEXT("POSITION")));
-		FAccessor&   Normal     = AccessorAtIndex(Asset->Accessors, GetIndex(Attributes, TEXT("NORMAL")));
-		FAccessor&   Tangent    = AccessorAtIndex(Asset->Accessors, GetIndex(Attributes, TEXT("TANGENT")));
-		FAccessor&   TexCoord0  = AccessorAtIndex(Asset->Accessors, GetIndex(Attributes, TEXT("TEXCOORD_0")));
-		FAccessor&   TexCoord1  = AccessorAtIndex(Asset->Accessors, GetIndex(Attributes, TEXT("TEXCOORD_1")));
-		const FAccessor&   Color0     = AccessorAtIndex(A, GetIndex(Attributes, TEXT("COLOR_0")));
-		const FAccessor&   Joints0    = AccessorAtIndex(A, GetIndex(Attributes, TEXT("JOINTS_0")));
-		const FAccessor&   Weights0   = AccessorAtIndex(A, GetIndex(Attributes, TEXT("WEIGHTS_0")));
-
-		Position.bQuantized = bMeshQuantized;
-		Normal.bQuantized = bMeshQuantized;
-		Tangent.bQuantized = bMeshQuantized;
-		TexCoord0.bQuantized = bMeshQuantized;
-		TexCoord1.bQuantized = bMeshQuantized;
-
-		Mesh.Primitives.Emplace(Mode, MaterialIndex, Indices, Position, Normal, Tangent, TexCoord0, TexCoord1, Color0, Joints0, Weights0);
+		Mesh.Primitives.Emplace(Mode, MaterialIndex, Indices, AttributeAccessors);
 
 		//Morph Targets:
 		if (Object.HasField(TEXT("targets")))
diff --git a/Engine/Plugins/Interchange/Runtime/Source/Parsers/GLTFCore/Public/GLTFAccessor.h b/Engine/Plugins/Interchange/Runtime/Source/Parsers/GLTFCore/Public/GLTFAccessor.h
--- a/Engine/Plugins/Interchange/Runtime/Source/Parsers/GLTFCore/Public/GLTFAccessor.h
+++ b/Engine/Plugins/Interchange/Runtime/Source/Parsers/GLTFCore/Public/GLTFAccessor.h
@@ -3,6 +3,7 @@
 #pragma once
 
 #include "CoreMinimal.h"
+#include "Components.h" //For MAX_STATIC_TEXCOORDS check.
 
 struct FMD5Hash;
 
@@ -93,8 +94,16 @@ namespace GLTF
 		POSITION = 0,
 		NORMAL,
 		TANGENT,
+		
 		TEXCOORD_0, 
 		TEXCOORD_1,
+		TEXCOORD_2,
+		TEXCOORD_3,
+		TEXCOORD_4,
+		TEXCOORD_5,
+		TEXCOORD_6,
+		TEXCOORD_7,
+		
 		COLOR_0,
 
 		JOINTS_0,
@@ -103,6 +112,9 @@ namespace GLTF
 		COUNT,
 	};
 
+	constexpr uint8 SupportedTexCoordNum = EMeshAttributeType::TEXCOORD_7 - EMeshAttributeType::TEXCOORD_0 + 1;
+	static_assert(SupportedTexCoordNum == MAX_STATIC_TEXCOORDS, "The maximum UE supported Texcoord count has changed. Update the EMeshAttributeType and related accordingly.");
+
 	GLTFCORE_API FString ToString(const EMeshAttributeType& Type);
 
 	struct GLTFCORE_API FAccessor
diff --git a/Engine/Plugins/Interchange/Runtime/Source/Parsers/GLTFCore/Public/GLTFMesh.h b/Engine/Plugins/Interchange/Runtime/Source/Parsers/GLTFCore/Public/GLTFMesh.h
--- a/Engine/Plugins/Interchange/Runtime/Source/Parsers/GLTFCore/Public/GLTFMesh.h
+++ b/Engine/Plugins/Interchange/Runtime/Source/Parsers/GLTFCore/Public/GLTFMesh.h
@@ -43,16 +43,8 @@ namespace GLTF
 
 	struct GLTFCORE_API FAttributeAccessors
 	{
-		FAttributeAccessors(const FAccessor& InPositionDisplacements, const FAccessor& InNormalDisplacements,
-							const FAccessor& InTangentDisplacements, 
-							const FAccessor& InTexCoord0Displacements, const FAccessor& InTexCoord1Displacements, 
-							const FAccessor& InColor0Deltas);
-
-		FAttributeAccessors(const FAccessor& InPosition, const FAccessor& InNormal,
-							const FAccessor& InTangent, 
-							const FAccessor& InTexCoord0, const FAccessor& InTexCoord1, 
-							const FAccessor& InColor0, 
-							const FAccessor& InJoints0, const FAccessor& InWeights0);
+		//Expects elements in the exact same order as the EMeshAttributeType.
+		FAttributeAccessors(const TArray<TPair<EMeshAttributeType, FAccessor&>>& InAttributeAccessors, bool bInMorphTarget);
 
 	protected:
 		bool HasAttributeAccessor(const EMeshAttributeType& Type) const;
@@ -63,14 +55,13 @@ namespace GLTF
 		void HashAttributes(FMD5& MD5) const;
 
 	private:
-		TArray<TPair<EMeshAttributeType, const FAccessor&>> AttributeAccessors;
-		const bool bMorphTarget;
+		TArray<TPair<EMeshAttributeType, FAccessor&>> AttributeAccessors;
+		bool bMorphTarget;
 	};
 
 	struct GLTFCORE_API FMorphTarget : FAttributeAccessors
 	{
-		FMorphTarget(const FAccessor& InPositionDisplacements, const FAccessor& InNormalDisplacements,
-			const FAccessor& InTangentDisplacements, const FAccessor& InTexCoord0Displacements, const FAccessor& InTexCoord1Displacements, const FAccessor& InColor0Deltas);
+		FMorphTarget(const TArray<TPair<EMeshAttributeType, FAccessor&>>& InAttributeAccessors /*Displacements and Deltas*/);
 
 		bool IsValid(uint32 ExpectedElementCount) const;
 		FMD5Hash GetHash() const;
@@ -123,9 +114,7 @@ namespace GLTF
 
 		TMap<FString, FString> Extras;
 
-		FPrimitive(EMode InMode, int32 InMaterial, const FAccessor& InIndices, const FAccessor& InPosition, const FAccessor& InNormal,
-		           const FAccessor& InTangent, const FAccessor& InTexCoord0, const FAccessor& InTexCoord1, const FAccessor& InColor0,
-		           const FAccessor& InJoints0, const FAccessor& InWeights0);
+		FPrimitive(EMode InMode, int32 InMaterial, const FAccessor& InIndices, const TArray<TPair<EMeshAttributeType, FAccessor&>>& InAttributeAccessors);
 
 		void GenerateIsValidCache();
 		bool IsValid() const;
