Interchange MaterialX:
- Create material expression range, in order to reduce shader graphs that use it (such as randomfloat, randomcolor)
- Add material function randomfloat and randomcolor using cellnoise2d algorithm (technically a cellnoise3d with a seed of 0 for the 'z' channel)
- Add a method AddIntegerAttribute in order to have named parameters for integers
- Set the default unit of LocalPosition to 0.01 (meters->centimeters), as it is now fixed that the issue came from the wrong Sine period
- Refactor the <rotate3d> node to use the RotateAboutAxis expression using the period and handling correctly the output (RotateAboutAxis returns an offset where rotate3d rotates the vector)

UMaterialExpressionMaterialXFractal3D:
- Refactor the Fractal3D expression to use directly the GradientNoise3D_TEX algorithm which gives better and closer results to MaterialXViewer
- Removes unnecessary Inputs of the fractal

GenericMaterialPipeline:
- Handles Sine/Cosine/Tangent expressions by defaulting the period to 2Pi, because the respective UMaterialExpressions default it to 1, resulting in wrong results for test cases
- Add the support of the period for RotateAboutAxis nodes

EngineTest:
- Update screenshots

#virtualized

diff --git a/Engine/Plugins/Interchange/Runtime/Source/Common/Public/MaterialX/InterchangeMaterialXDefinitions.h b/Engine/Plugins/Interchange/Runtime/Source/Common/Public/MaterialX/InterchangeMaterialXDefinitions.h
--- a/Engine/Plugins/Interchange/Runtime/Source/Common/Public/MaterialX/InterchangeMaterialXDefinitions.h
+++ b/Engine/Plugins/Interchange/Runtime/Source/Common/Public/MaterialX/InterchangeMaterialXDefinitions.h
@@ -183,6 +183,7 @@ namespace UE
 			static constexpr const TCHAR* MaterialXRamp4 = TEXT("MaterialXRamp4");
 			static constexpr const TCHAR* MaterialXRampLeftRight = TEXT("MaterialXRampLeftRight");
 			static constexpr const TCHAR* MaterialXRampTopBottom = TEXT("MaterialXRampTopBottom");
+			static constexpr const TCHAR* MaterialXRange = TEXT("MaterialXRange");
 			static constexpr const TCHAR* MaterialXRemap = TEXT("MaterialXRemap");
 			static constexpr const TCHAR* MaterialXScreen = TEXT("MaterialXScreen");
 			static constexpr const TCHAR* MaterialXSplitLeftRight = TEXT("MaterialXSplitLeftRight");
@@ -194,6 +195,7 @@ namespace UE
 			static constexpr const TCHAR* Normalize = TEXT("Normalize");
 			static constexpr const TCHAR* Power = TEXT("Power");
 			static constexpr const TCHAR* RgbToHsv = TEXT("RgbToHsv");
+			static constexpr const TCHAR* RotateAboutAxis = TEXT("RotateAboutAxis");
 			static constexpr const TCHAR* Round = TEXT("Round");
 			static constexpr const TCHAR* Sign = TEXT("Sign");
 			static constexpr const TCHAR* Sine = TEXT("Sine");
@@ -227,6 +229,7 @@ namespace UE
 			static constexpr const TCHAR* Base = TEXT("Base");
 			static constexpr const TCHAR* Bottom = TEXT("Bottom");
 			static constexpr const TCHAR* C = TEXT("C");
+			static constexpr const TCHAR* Clamp = TEXT("Clamp");
 			static constexpr const TCHAR* Center = TEXT("Center");
 			static constexpr const TCHAR* Coordinate = TEXT("Coordinate");
 			static constexpr const TCHAR* Coordinates = TEXT("Coordinates");
@@ -238,6 +241,7 @@ namespace UE
 			static constexpr const TCHAR* Flatness = TEXT("Flatness");
 			static constexpr const TCHAR* Foreground = TEXT("Foreground");
 			static constexpr const TCHAR* Fraction = TEXT("Fraction");
+			static constexpr const TCHAR* Gamma = TEXT("Gamma");
 			static constexpr const TCHAR* in1 = TEXT("in1");
 			static constexpr const TCHAR* in2 = TEXT("in2");
 			static constexpr const TCHAR* in3 = TEXT("in3");
@@ -699,6 +703,7 @@ MATERIALX_NAMESPACE_BEGIN
 		static constexpr const char* Ramp4 = "ramp4";
 		static constexpr const char* RampLR = "ramplr";
 		static constexpr const char* RampTB = "ramptb";
+		static constexpr const char* RandomFloat = "randomfloat";
 		static constexpr const char* SplitLR = "splitlr";
 		static constexpr const char* SplitTB = "splittb";
 		static constexpr const char* WorleyNoise2D = "worleynoise2d";
@@ -706,6 +711,7 @@ MATERIALX_NAMESPACE_BEGIN
 		static constexpr const char* CellNoise3D = "cellnoise3d";
 		static constexpr const char* Fractal3D = "fractal3d";
 		static constexpr const char* Noise3D = "noise3d";
+		static constexpr const char* RandomColor = "randomcolor";
 		static constexpr const char* WorleyNoise3D = "worleynoise3d";
 		// Organization nodes 
 		static constexpr const char* Dot = "dot";
@@ -726,6 +732,7 @@ MATERIALX_NAMESPACE_BEGIN
 		// Adjustment nodes
 		static constexpr const char* HsvToRgb = "hsvtorgb";
 		static constexpr const char* Luminance = "luminance";
+		static constexpr const char* Range = "range";
 		static constexpr const char* Remap = "remap";
 		static constexpr const char* RgbToHsv = "rgbtohsv";
 		static constexpr const char* Saturate = "saturate";
diff --git a/Engine/Plugins/Interchange/Runtime/Source/Import/Private/MaterialX/MaterialExpressions/MaterialExpressionFractal3D.cpp b/Engine/Plugins/Interchange/Runtime/Source/Import/Private/MaterialX/MaterialExpressions/MaterialExpressionFractal3D.cpp
--- a/Engine/Plugins/Interchange/Runtime/Source/Import/Private/MaterialX/MaterialExpressions/MaterialExpressionFractal3D.cpp
+++ b/Engine/Plugins/Interchange/Runtime/Source/Import/Private/MaterialX/MaterialExpressions/MaterialExpressionFractal3D.cpp
@@ -35,45 +35,28 @@ int32 UMaterialExpressionMaterialXFractal3D::Compile(FMaterialCompiler* Compiler
 	MaterialExpressionCustom->Inputs.Add({ TEXT("Lacunarity") });
 	MaterialExpressionCustom->Inputs.Add({ TEXT("Diminish") });
 	
-	MaterialExpressionCustom->OutputType = ECustomMaterialOutputType::CMOT_Float3;
+	MaterialExpressionCustom->OutputType = ECustomMaterialOutputType::CMOT_Float1; //Noise functions returns only a float
 	MaterialExpressionCustom->Code = 
-		TEXT("float Scale =") + FString::SanitizeFloat(Scale) + TEXT(";") +
-		TEXT(R"(int Quality = 1;
-             int Function = 2; // Fast Gradient
-             bool bTurbulence =)") + FString::FromInt(bTurbulence) + TEXT(";") +
-		TEXT("uint Levels =") + FString::FromInt(Levels) + TEXT(";") +
-		TEXT("float OutputMin =") + FString::SanitizeFloat(OutputMin) + TEXT(";") +
-		TEXT("float OutputMax =") + FString::SanitizeFloat(OutputMax) + TEXT(";") +
-		TEXT(R"(float LevelScale = 2;
-             float FilterWidth = 0;
-             bool bTiling = false;
-             float RepeatSize = 512;
-             
-             float3 result = float3(0.0,0.0,0.0);
-             float amplitude = 1.0;
-             for (int i = 0;  i < Octaves; ++i)
-             {
-                 result += amplitude * 
-                 MaterialExpressionNoise(
-                     Position,
-                     Scale,
-                     Quality,
-                     Function,
-                     bTurbulence,
-                     Levels,
-                     OutputMin,
-                     OutputMax,
-                     LevelScale,
-                     FilterWidth,
-                     bTiling,
-                     RepeatSize);
-             
-                 amplitude *= Diminish;
-                 Position *= Lacunarity;
-             }
-             return result;)");
+		TEXT(R"(const bool bTiling = false;
+        const float RepeatSize = 512;
+	
+		float3 result = float3(0.0,0.0,0.0);
+		float amplitude = 1.0;
+		for (int i = 0;  i < Octaves; ++i)
+		{
+			result += amplitude * GradientNoise3D_TEX(Position, bTiling, RepeatSize);
+			amplitude *= Diminish;
+			Position *= Lacunarity;
+		}
+		return result;)");
+
+	int32 IndexPosition = Position.GetTracedInput().Expression ? Position.Compile(Compiler) : Compiler->LocalPosition(EPositionIncludedOffsets::IncludeOffsets, ELocalPositionOrigin::Instance);
+	
+	if (IndexPosition == INDEX_NONE)
+	{
+		return Compiler->Errorf(TEXT("Failed to compile Position input."));
+	}
 
-	int32 IndexPosition = Position.GetTracedInput().Expression ? Position.Compile(Compiler) : Compiler->TransformPosition(EMaterialCommonBasis::MCB_World, EMaterialCommonBasis::MCB_Local, Compiler->WorldPosition(WPT_Default));
 	int32 IndexOctaves = Octaves.GetTracedInput().Expression ? Octaves.Compile(Compiler) : Compiler->Constant(ConstOctaves);
 	int32 IndexLacunarity = Lacunarity.GetTracedInput().Expression ? Lacunarity.Compile(Compiler) : Compiler->Constant(ConstLacunarity);
 	int32 IndexDiminish = Diminish.GetTracedInput().Expression ? Diminish.Compile(Compiler) : Compiler->Constant(ConstDiminish);
diff --git a/Engine/Plugins/Interchange/Runtime/Source/Import/Private/MaterialX/MaterialExpressions/MaterialExpressionFractal3D.h b/Engine/Plugins/Interchange/Runtime/Source/Import/Private/MaterialX/MaterialExpressions/MaterialExpressionFractal3D.h
--- a/Engine/Plugins/Interchange/Runtime/Source/Import/Private/MaterialX/MaterialExpressions/MaterialExpressionFractal3D.h
+++ b/Engine/Plugins/Interchange/Runtime/Source/Import/Private/MaterialX/MaterialExpressions/MaterialExpressionFractal3D.h
@@ -51,23 +51,20 @@ class UMaterialExpressionMaterialXFractal3D : public UMaterialExpression
 	UPROPERTY(EditAnywhere, Category = MaterialExpressionFractal, meta = (OverridingInputProperty = "Diminish"))
 	float ConstDiminish = 0.5;
 
-	/** can also be done with a multiply on the Position */
-	UPROPERTY(EditAnywhere, Category = MaterialExpressionNoise)
-	float Scale = 1.f;
+	UPROPERTY()
+	float Scale_DEPRECATED = 1.f;
 
-	/** How multiple frequencies are getting combined */
-	UPROPERTY(EditAnywhere, Category = MaterialExpressionNoise)
-	bool bTurbulence = false;
+	UPROPERTY()
+	bool bTurbulence_DEPRECATED = false;
 
-	/** 1 = fast but little detail, .. larger numbers cost more performance */
-	UPROPERTY(EditAnywhere, Category = MaterialExpressionNoise, meta = (UIMin = "1", UIMax = "10"))
-	int32 Levels = 6;
+	UPROPERTY()
+	int32 Levels_DEPRECATED = 6;
 
-	UPROPERTY(EditAnywhere, Category = MaterialExpressionNoise)
-	float OutputMin = 0.f;
+	UPROPERTY()
+	float OutputMin_DEPRECATED = 0.f;
 
-	UPROPERTY(EditAnywhere, Category = MaterialExpressionNoise)
-	float OutputMax = 1.f;
+	UPROPERTY()
+	float OutputMax_DEPRECATED = 1.f;
 
 	//~ Begin UMaterialExpressionMaterialX Interface
 #if WITH_EDITOR
diff --git a/Engine/Plugins/Interchange/Runtime/Source/Import/Private/MaterialX/MaterialExpressions/MaterialExpressionRange.cpp b/Engine/Plugins/Interchange/Runtime/Source/Import/Private/MaterialX/MaterialExpressions/MaterialExpressionRange.cpp
new file mode 100644
--- /dev/null
+++ b/Engine/Plugins/Interchange/Runtime/Source/Import/Private/MaterialX/MaterialExpressions/MaterialExpressionRange.cpp
@@ -0,0 +1,163 @@
+// Copyright Epic Games, Inc. All Rights Reserved.
+#include "MaterialExpressionRange.h"
+#include "MaterialExpressionRemap.h"
+#include "MaterialCompiler.h"
+#include "MaterialHLSLGenerator.h"
+#include "Materials/MaterialExpressionAbs.h"
+#include "Materials/MaterialExpressionClamp.h"
+#include "Materials/MaterialExpressionDivide.h"
+#include "Materials/MaterialExpressionIf.h"
+#include "Materials/MaterialExpressionMultiply.h"
+#include "Materials/MaterialExpressionPower.h"
+#include "Materials/MaterialExpressionSign.h"
+
+#include UE_INLINE_GENERATED_CPP_BY_NAME(MaterialExpressionRange)
+
+#define LOCTEXT_NAMESPACE "MaterialExpressionMaterialXRange"
+
+UMaterialExpressionMaterialXRange::UMaterialExpressionMaterialXRange(const FObjectInitializer& ObjectInitializer)
+	: Super(ObjectInitializer)
+{
+	// Structure to hold one-time initialization
+	struct FConstructorStatics
+	{
+		FText NAME_MaterialX;
+		FConstructorStatics()
+			: NAME_MaterialX(LOCTEXT("MaterialX", "MaterialX"))
+		{}
+	};
+	static FConstructorStatics ConstructorStatics;
+
+	ConstInputLow = 0.f;
+	ConstInputHigh = 1.f;
+	ConstTargetLow = 0.f;
+	ConstTargetHigh = 1.f;
+	ConstGamma = 1.f;
+	bConstClamp = false;
+
+#if WITH_EDITORONLY_DATA
+	MenuCategories.Add(ConstructorStatics.NAME_MaterialX);
+#endif
+}
+
+#if WITH_EDITOR
+int32 UMaterialExpressionMaterialXRange::Compile(FMaterialCompiler* Compiler, int32 OutputIndex)
+{
+	UMaterialExpressionMaterialXRemap* Remap1 = NewObject<UMaterialExpressionMaterialXRemap>();
+	
+	Remap1->Input = Input;
+	Remap1->InputLow = InputLow;
+	Remap1->InputHigh = InputHigh;
+	Remap1->InputLowDefault = ConstInputLow;
+	Remap1->InputHighDefault = ConstInputHigh;
+	
+	//First remap, we only take the inputs, the target is (0,1), constants will be applied when we compile
+	Remap1->TargetLowDefault = 0.f;
+	Remap1->TargetHighDefault = 1.f;
+
+	//Inverse Gamma
+	UMaterialExpressionDivide* InvGamma = NewObject<UMaterialExpressionDivide>();
+	InvGamma->ConstA = 1.f;
+	InvGamma->B = Gamma;
+	InvGamma->ConstB = ConstGamma;
+
+	//Abs Remap
+	UMaterialExpressionAbs* AbsRemap1 = NewObject<UMaterialExpressionAbs>();
+	AbsRemap1->Input.Connect(0, Remap1);
+
+	//Gamma correct
+	UMaterialExpressionPower* Power = NewObject<UMaterialExpressionPower>();
+	Power->Base.Connect(0, AbsRemap1);
+	Power->Exponent.Connect(0, InvGamma);
+
+	// Sign of remap
+	UMaterialExpressionSign* Sign = NewObject<UMaterialExpressionSign>();
+	Sign->Input.Connect(0, Remap1);
+
+	//Sign Gamma Correct
+	UMaterialExpressionMultiply* GammaSign = NewObject<UMaterialExpressionMultiply>();
+	GammaSign->A.Connect(0, Power);
+	GammaSign->B.Connect(0, Sign);
+
+	UMaterialExpressionMaterialXRemap* Remap2 = NewObject<UMaterialExpressionMaterialXRemap>();
+	Remap2->Input.Connect(0, GammaSign);
+	Remap2->InputLowDefault = 0.f;
+	Remap2->InputHighDefault = 1.f;
+	Remap2->TargetLow = TargetLow;
+	Remap2->TargetHigh = TargetHigh;
+	Remap2->TargetLowDefault = ConstTargetLow;
+	Remap2->TargetHighDefault = ConstTargetHigh;
+
+	// Only Create a Clamp expression if its an input
+	if (Clamp.GetTracedInput().Expression)
+	{
+		UMaterialExpressionClamp* ExpressionClamp = NewObject<UMaterialExpressionClamp>();
+		ExpressionClamp->Input.Connect(0, Remap2);
+		ExpressionClamp->Min = TargetLow;
+		ExpressionClamp->MinDefault = ConstTargetLow;
+		ExpressionClamp->Max = TargetHigh;
+		ExpressionClamp->MaxDefault = ConstTargetHigh;
+
+		UMaterialExpressionIf * If = NewObject<UMaterialExpressionIf>();
+		If->A = Clamp;
+		If->ConstB = 1;
+
+		If->AEqualsB.Connect(0, ExpressionClamp);
+		If->AGreaterThanB.Connect(0, Remap2);
+		If->ALessThanB.Connect(0, Remap2);
+
+		return If->Compile(Compiler, OutputIndex);
+	}
+	else
+	{
+		if (bConstClamp)
+		{
+			UMaterialExpressionClamp* ExpressionClamp = NewObject<UMaterialExpressionClamp>();
+			ExpressionClamp->Input.Connect(0, Remap2);
+			ExpressionClamp->Min = TargetLow;
+			ExpressionClamp->MinDefault = ConstTargetLow;
+			ExpressionClamp->Max = TargetHigh;
+			ExpressionClamp->MaxDefault = ConstTargetHigh;
+
+			return ExpressionClamp->Compile(Compiler, OutputIndex);
+		}
+		else
+		{
+			return Remap2->Compile(Compiler, OutputIndex);
+		}
+	}
+}
+
+bool UMaterialExpressionMaterialXRange::GenerateHLSLExpression(FMaterialHLSLGenerator& Generator, UE::HLSLTree::FScope& Scope, int32 OutputIndex, UE::HLSLTree::FExpression const*& OutExpression) const
+{
+	//using namespace UE::HLSLTree;
+
+	//if (!Input.GetTracedInput().Expression)
+	//{
+	//	return Generator.Errorf(TEXT("Missing remap input"));
+	//}
+
+	//// Min2 + (Input - Min1) x (Max2 - Min2) / (Max1 - Min1)
+	//const FExpression* InputLowExpression = InputLow.AcquireHLSLExpressionOrConstant(Generator, Scope, InputLowDefault);
+	//const FExpression* InputHighExpression = InputHigh.AcquireHLSLExpressionOrConstant(Generator, Scope, InputHighDefault);
+
+	//const FExpression* TargetLowExpression = TargetLow.AcquireHLSLExpressionOrConstant(Generator, Scope, TargetLowDefault);
+	//const FExpression* TargetHighExpression = TargetHigh.AcquireHLSLExpressionOrConstant(Generator, Scope, TargetHighDefault);
+
+	//const FExpression* InputDeltaExpression = Generator.GetTree().NewSub(InputHighExpression, InputLowExpression);
+	//const FExpression* TargetDeltaExpression = Generator.GetTree().NewSub(TargetHighExpression, TargetLowExpression);
+	//const FExpression* ScaleExpression = Generator.GetTree().NewDiv(TargetDeltaExpression, InputDeltaExpression);
+
+	//const FExpression* InputExpression = Input.AcquireHLSLExpression(Generator, Scope);
+
+	//OutExpression = Generator.GetTree().NewAdd(TargetLowExpression, Generator.GetTree().NewMul(Generator.GetTree().NewSub(InputExpression, InputLowExpression), ScaleExpression));
+	return true;
+}
+
+void UMaterialExpressionMaterialXRange::GetCaption(TArray<FString>& OutCaptions) const
+{
+	OutCaptions.Add(TEXT("MaterialX Range"));
+}
+#endif
+
+#undef LOCTEXT_NAMESPACE 
\ No newline at end of file
diff --git a/Engine/Plugins/Interchange/Runtime/Source/Import/Private/MaterialX/MaterialExpressions/MaterialExpressionRange.h b/Engine/Plugins/Interchange/Runtime/Source/Import/Private/MaterialX/MaterialExpressions/MaterialExpressionRange.h
new file mode 100644
--- /dev/null
+++ b/Engine/Plugins/Interchange/Runtime/Source/Import/Private/MaterialX/MaterialExpressions/MaterialExpressionRange.h
@@ -0,0 +1,66 @@
+// Copyright Epic Games, Inc. All Rights Reserved.
+#pragma once
+
+#include "CoreMinimal.h"
+#include "Materials/MaterialExpression.h"
+
+#include "MaterialExpressionRange.generated.h"
+
+/**
+ * A material expression that Remap a value from one range to another, optionally
+ * applying a gamma correction in the middle, and optionally clamping output values.
+ */
+UCLASS(collapsecategories, hidecategories = Object, MinimalAPI, meta = (Private))
+class UMaterialExpressionMaterialXRange : public UMaterialExpression
+{
+	GENERATED_UCLASS_BODY()
+
+	UPROPERTY()
+	FExpressionInput Input;
+
+	UPROPERTY(meta = (RequiredInput = "false", ToolTip = "Low value for input range. Defaults to 'ConstInputLowDefault' if not specified"))
+	FExpressionInput InputLow;
+
+	UPROPERTY(meta = (RequiredInput = "false", ToolTip = "High value for input range. Defaults to 'ConstInputHighDefault' if not specified"))
+	FExpressionInput InputHigh;
+
+	UPROPERTY(meta = (RequiredInput = "false", ToolTip = "Low value for target range. Defaults to 'ConstTargetLowDefault' if not specified"))
+	FExpressionInput TargetLow;
+
+	UPROPERTY(meta = (RequiredInput = "false", ToolTip = "High value for target range. Defaults to 'ConstTargetHighDefault' if not specified"))
+	FExpressionInput TargetHigh;
+
+	UPROPERTY(meta = (RequiredInput = "false", ToolTip = "Inverse exponent applied to Input after first transforming from InputLow..InputHigh to 0..1, gamma values greater than 1.0 make midtones brighter. Defaults to 'ConstGamma' if not specified"))
+	FExpressionInput Gamma;
+
+	UPROPERTY(meta = (RequiredInput = "false", ToolTip = "If true, the output is clamped to the range TargetLow..TargetHigh. Defaults to 'ConstClamp' if not specified"))
+	FExpressionInput Clamp;
+
+	UPROPERTY(EditAnywhere, Category = MaterialExpressionRange, meta = (OverridingInputProperty = "InputLow"))
+	float ConstInputLow;
+
+	UPROPERTY(EditAnywhere, Category = MaterialExpressionRange, meta = (OverridingInputProperty = "InputHigh"))
+	float ConstInputHigh;
+
+	UPROPERTY(EditAnywhere, Category = MaterialExpressionRange, meta = (OverridingInputProperty = "TargetLow"))
+	float ConstTargetLow;
+
+	UPROPERTY(EditAnywhere, Category = MaterialExpressionRange, meta = (OverridingInputProperty = "TargetHigh"))
+	float ConstTargetHigh;
+
+	UPROPERTY(EditAnywhere, Category = MaterialExpressionRange, meta = (OverridingInputProperty = "Gamma"))
+	float ConstGamma;
+
+	UPROPERTY(EditAnywhere, Category = MaterialExpressionRange, meta = (OverridingInputProperty = "Clamp"))
+	bool bConstClamp;
+
+	//~ Begin UMaterialExpressionMaterialX Interface
+#if WITH_EDITOR
+	virtual int32 Compile(class FMaterialCompiler* Compiler, int32 OutputIndex) override;
+	virtual void GetCaption(TArray<FString>& OutCaptions) const override;
+
+	virtual bool GenerateHLSLExpression(class FMaterialHLSLGenerator& Generator, class UE::HLSLTree::FScope& Scope, int32 OutputIndex, class UE::HLSLTree::FExpression const*& OutExpression) const override;
+#endif
+	//~ End UMaterialExpressionMaterialX Interface
+};
+
diff --git a/Engine/Plugins/Interchange/Runtime/Source/Import/Private/MaterialX/MaterialXUtils/MaterialXBase.cpp b/Engine/Plugins/Interchange/Runtime/Source/Import/Private/MaterialX/MaterialXUtils/MaterialXBase.cpp
--- a/Engine/Plugins/Interchange/Runtime/Source/Import/Private/MaterialX/MaterialXUtils/MaterialXBase.cpp
+++ b/Engine/Plugins/Interchange/Runtime/Source/Import/Private/MaterialX/MaterialXUtils/MaterialXBase.cpp
@@ -31,6 +31,9 @@ void FMaterialXBase::UpdateDocumentRecursively(MaterialX::GraphElementPtr Graph)
 		mx::Category::Overlay,
 		mx::Category::Place2D,
 		mx::Category::Ramp4,
+		mx::Category::RandomFloat,
+		mx::Category::RandomColor,
+		mx::Category::Range,
 		mx::Category::Refract,
 		mx::Category::Saturate,
 		mx::Category::Separate2,
diff --git a/Engine/Plugins/Interchange/Runtime/Source/Import/Private/MaterialX/MaterialXUtils/MaterialXManager.cpp b/Engine/Plugins/Interchange/Runtime/Source/Import/Private/MaterialX/MaterialXUtils/MaterialXManager.cpp
--- a/Engine/Plugins/Interchange/Runtime/Source/Import/Private/MaterialX/MaterialXUtils/MaterialXManager.cpp
+++ b/Engine/Plugins/Interchange/Runtime/Source/Import/Private/MaterialX/MaterialXUtils/MaterialXManager.cpp
@@ -35,7 +35,9 @@ FMaterialXManager::FMaterialXManager()
 	   {{TEXT(""),                            TEXT("bg")},          ExpressionInput(UE::Expressions::Inputs::B)},
 	   {{TEXT(""),                            TEXT("center")},      ExpressionInput(UE::Expressions::Inputs::Center)},
 	   {{TEXT(""),                            TEXT("diminish")},    ExpressionInput(UE::Expressions::Inputs::Diminish)},
+	   {{TEXT(""),                            TEXT("doclamp")},     ExpressionInput(UE::Expressions::Inputs::Clamp)},
 	   {{TEXT(""),                            TEXT("fg")},          ExpressionInput(UE::Expressions::Inputs::A)},
+	   {{TEXT(""),                            TEXT("gamma")},       ExpressionInput(UE::Expressions::Inputs::Gamma)},
 	   {{TEXT(""),                            TEXT("high")},        ExpressionInput(UE::Expressions::Inputs::Max)},
 	   {{TEXT(""),                            TEXT("in")},          ExpressionInput(UE::Expressions::Inputs::Input)},
 	   {{TEXT(""),                            TEXT("in1")},         ExpressionInput(UE::Expressions::Inputs::A)},
@@ -181,6 +183,7 @@ FMaterialXManager::FMaterialXManager()
 		// Adjustment nodes,
 		{mx::Category::HsvToRgb,     UE::Expressions::Names::HsvToRgb},
 		{mx::Category::Luminance,    UE::Expressions::Names::MaterialXLuminance},
+		{mx::Category::Range,        UE::Expressions::Names::MaterialXRange},
 		{mx::Category::Remap,        UE::Expressions::Names::MaterialXRemap},
 		{mx::Category::RgbToHsv,     UE::Expressions::Names::RgbToHsv},
 		{mx::Category::Saturate,     UE::Expressions::Names::Desaturation},
@@ -230,6 +233,9 @@ FMaterialXManager::FMaterialXManager()
 		{mx::Category::LinDisplayP3ToLinRec709,	FMaterialXMaterialFunction{TInPlaceType<FString>{}, TEXT("/InterchangeAssets/Functions/MX_LinearDisplayP3ToRec709.MX_LinearDisplayP3ToRec709")}},
 		{mx::Category::SrgbDisplayP3ToLinRec709,FMaterialXMaterialFunction{TInPlaceType<FString>{}, TEXT("/InterchangeAssets/Functions/MX_SrgbDisplayP3ToRec709.MX_SrgbDisplayP3ToRec709")}},
 		{mx::Category::SrgbTextureToLinRec709,	FMaterialXMaterialFunction{TInPlaceType<FString>{}, TEXT("/InterchangeAssets/Functions/MX_SrgbToRec709.MX_SrgbToRec709")}},
+		//Procedural
+		{mx::Category::RandomFloat,	            FMaterialXMaterialFunction{TInPlaceType<FString>{}, TEXT("/InterchangeAssets/Functions/MX_RandomFloat.MX_RandomFloat")}},
+		{mx::Category::RandomColor,	            FMaterialXMaterialFunction{TInPlaceType<FString>{}, TEXT("/InterchangeAssets/Functions/MX_RandomColor.MX_RandomColor")}},
 	}
 	, bIsSubstrateEnabled{ IInterchangeImportModule::IsAvailable() ? IInterchangeImportModule::Get().IsSubstrateEnabled() : false }
 {
@@ -382,6 +388,8 @@ namespace UE::Interchange::MaterialX
 																TEXT("MaterialFunction'/InterchangeAssets/Functions/MX_LinearAdobeRGBToRec709.MX_LinearAdobeRGBToRec709'"),
 																TEXT("MaterialFunction'/InterchangeAssets/Functions/MX_SrgbToRec709.MX_SrgbToRec709'"),
 																TEXT("MaterialFunction'/InterchangeAssets/Functions/MX_SrgbDisplayP3ToRec709.MX_SrgbDisplayP3ToRec709'"),
+																TEXT("MaterialFunction'/InterchangeAssets/Functions/MX_RandomFloat.MX_RandomFloat'"),
+																TEXT("MaterialFunction'/InterchangeAssets/Functions/MX_RandomColor.MX_RandomColor'"),
 															  });
 
 		return bPackagesLoaded;
diff --git a/Engine/Plugins/Interchange/Runtime/Source/Import/Private/MaterialX/MaterialXUtils/MaterialXSurfaceShaderAbstract.cpp b/Engine/Plugins/Interchange/Runtime/Source/Import/Private/MaterialX/MaterialXUtils/MaterialXSurfaceShaderAbstract.cpp
--- a/Engine/Plugins/Interchange/Runtime/Source/Import/Private/MaterialX/MaterialXUtils/MaterialXSurfaceShaderAbstract.cpp
+++ b/Engine/Plugins/Interchange/Runtime/Source/Import/Private/MaterialX/MaterialXUtils/MaterialXSurfaceShaderAbstract.cpp
@@ -41,7 +41,7 @@ bool FMaterialXSurfaceShaderAbstract::AddAttribute(MaterialX::InputPtr Input, co
 		}
 		else if(Input->getType() == mx::Type::Integer) //Let's add Float attribute, because Interchange doesn't create a scalar if it's an int
 		{
-			return ShaderNode->AddFloatAttribute(UInterchangeShaderPortsAPI::MakeInputValueKey(InputChannelName), mx::fromValueString<int32>(Input->getValueString()));
+			return AddIntegerAttribute(Input, InputChannelName, ShaderNode);
 		}
 		else if(Input->getType() == mx::Type::Color3 || Input->getType() == mx::Type::Color4)
 		{
@@ -142,6 +142,30 @@ bool FMaterialXSurfaceShaderAbstract::AddFloatAttribute(MaterialX::InputPtr Inpu
 	return false;
 }
 
+bool FMaterialXSurfaceShaderAbstract::AddIntegerAttribute(MaterialX::InputPtr Input, const FString& InputChannelName, UInterchangeShaderNode* ShaderNode, int DefaultValue)
+{
+	//We handle Integers as Scalars
+	using namespace UE::Interchange::Materials::Standard::Nodes;
+	if (Input && Input->hasValue())
+	{
+		int32 Value = mx::fromValueString<int32>(Input->getValueString());
+
+		if (Value != DefaultValue)
+		{
+			// The parent is either a node, or it's an interfacename and we just take the name of the input
+			mx::NodePtr Node = Input->getParent()->asA<mx::Node>();
+			FString NodeName = Node ? Node->getName().c_str() + FString{ TEXT("_") } : FString{};
+			NodeName += Input->getName().c_str();
+
+			UInterchangeShaderNode* ScalarParameterNode = CreateShaderNode(Input, NodeName, ScalarParameter::Name.ToString());
+			ScalarParameterNode->AddFloatAttribute(UInterchangeShaderPortsAPI::MakeInputParameterKey(ScalarParameter::Attributes::DefaultValue.ToString()), Value);
+			return UInterchangeShaderPortsAPI::ConnectDefaultOuputToInput(ShaderNode, InputChannelName, ScalarParameterNode->GetUniqueID());
+		}
+	}
+
+	return false;
+}
+
 bool FMaterialXSurfaceShaderAbstract::AddLinearColorAttribute(MaterialX::InputPtr Input, const FString& InputChannelName, UInterchangeShaderNode* ShaderNode, const FLinearColor& DefaultValue, int32 OutputIndex)
 {
 	using namespace UE::Interchange::Materials::Standard::Nodes;
@@ -498,32 +522,34 @@ void FMaterialXSurfaceShaderAbstract::ConnectRotate2DInputToOutput(const FConnec
 
 void FMaterialXSurfaceShaderAbstract::ConnectRotate3DInputToOutput(const FConnectNode& Connect)
 {
-	UInterchangeShaderNode* Rotate3DNode = CreateShaderNode(Connect.UpstreamNode, Connect.UpstreamNode->getName().c_str(), TEXT("RotateAboutAxis"));
-	Rotate3DNode->AddLinearColorAttribute(UInterchangeShaderPortsAPI::MakeInputValueKey(TEXT("PivotPoint")), FLinearColor(0.5, 0.5, 0));
+	using namespace UE::Interchange::Materials::Standard::Nodes;
+
+	UInterchangeShaderNode* Rotate3DNode = CreateShaderNode(Connect.UpstreamNode, Connect.UpstreamNode->getName().c_str(), RotateAboutAxis::Name.ToString());
+	Rotate3DNode->AddLinearColorAttribute(UInterchangeShaderPortsAPI::MakeInputValueKey(RotateAboutAxis::Inputs::PivotPoint.ToString()), FLinearColor(0., 0., 0));
 
-	AddAttributeFromValueOrInterface(Connect.UpstreamNode->getInput("in"), UE::Expressions::Inputs::Position, Rotate3DNode);
+	mx::InputPtr Input = Connect.UpstreamNode->getInput("in");
+	AddAttributeFromValueOrInterface(Input, UE::Expressions::Inputs::Position, Rotate3DNode);
 	AddAttributeFromValueOrInterface(Connect.UpstreamNode->getInput("axis"), UE::Expressions::Inputs::NormalizedRotationAxis, Rotate3DNode);
+	AddAttributeFromValueOrInterface(Connect.UpstreamNode->getInput("amount"), UE::Expressions::Inputs::RotationAngle, Rotate3DNode);
 
-	// we create a Divide node to convert MaterialX angle in degrees to UE's angle which is a value between [0,1]
-	if(mx::InputPtr Input = Connect.UpstreamNode->getInput("amount"))
+	// Convert Degrees into radians by setting the period to 360
+	Rotate3DNode->AddFloatAttribute(RotateAboutAxis::Attributes::Period.ToString(), 360.f);
+
+	//RotateAboutAxis returns an offset of the rotated vector, let's add to the vector to rotate it directly
+	FString AddNodeName{ (Connect.UpstreamNode->getName() + "_Add").c_str() };
+	UInterchangeShaderNode* AddNode = CreateShaderNode(Connect.UpstreamNode, AddNodeName, Add::Name.ToString());
+	UInterchangeShaderPortsAPI::ConnectDefaultOuputToInput(AddNode, Add::Inputs::A.ToString(), Rotate3DNode->GetUniqueID());
+
+	//Add another input to connect the Position to the Add node
+	if (!Connect.UpstreamNode->getInput("position_add"))
 	{
-		if(!AddAttributeFromValueOrInterface(Input, UE::Expressions::Inputs::RotationAngle, Rotate3DNode))
-		{
-			//copy "in" input into "in1" and create "amount" input as "in2" with the value 360 because UE's angle is a value between [0,1]
-			mx::NodePtr NewDivideNode = CreateNode(Connect.UpstreamNode->getParent(),
-												   Connect.UpstreamNode->getName().c_str(),
-												   mx::Category::Divide,
-												   { {"in1", Input} }, // rename input to match <divide> input "in1"
-												   { {"in2", FAttributeValueArray{{"type", "float"}, {"value", "360"}}} });
-
-			//We need to set the type otherwise we won't be able to find a nodedef
-			NewDivideNode->setType(Input->getType());
-			// Input now points to the new node
-			Input->setNodeName(NewDivideNode->getName());
-		}
+		mx::InputPtr InputPositionAdd = Connect.UpstreamNode->addInput("position_add");
+		InputPositionAdd->copyContentFrom(Input);
+		InputPositionAdd->setAttribute(mx::Attributes::ParentName, TCHAR_TO_ANSI(*AddNodeName));
+		SetAttributeNewName(InputPositionAdd, TCHAR_TO_ANSI(UE::Expressions::Inputs::B));
 	}
 
-	UInterchangeShaderPortsAPI::ConnectDefaultOuputToInput(Connect.ParentShaderNode, Connect.InputChannelName, Rotate3DNode->GetUniqueID());
+	UInterchangeShaderPortsAPI::ConnectDefaultOuputToInput(Connect.ParentShaderNode, Connect.InputChannelName, AddNode->GetUniqueID());
 }
 
 void FMaterialXSurfaceShaderAbstract::ConnectImageInputToOutput(const FConnectNode& Connect)
@@ -908,9 +934,7 @@ void FMaterialXSurfaceShaderAbstract::ConnectPositionInputToOutput(const FConnec
 
 	UInterchangeShaderNode* UnitNode = CreateShaderNode(Connect.UpstreamNode, Connect.UpstreamNode->getName().c_str(), Multiply::Name.ToString());
 
-	// In case of the position node, it seems that the unit is different, we assume for now a conversion from mm -> m, even if UE by default is cm
-	// See standard_surface_marble_solid file, especially on the fractal3d node
-	UnitNode->AddFloatAttribute(UInterchangeShaderPortsAPI::MakeInputValueKey(TEXT("B")), 0.001f);
+	UnitNode->AddFloatAttribute(UInterchangeShaderPortsAPI::MakeInputValueKey(TEXT("B")), 0.01f);
 	UInterchangeShaderPortsAPI::ConnectDefaultOuputToInput(UnitNode, TEXT("A"), PositionNode->GetUniqueID());
 
 	UInterchangeShaderPortsAPI::ConnectDefaultOuputToInput(Connect.ParentShaderNode, Connect.InputChannelName, UnitNode->GetUniqueID());
diff --git a/Engine/Plugins/Interchange/Runtime/Source/Import/Public/MaterialX/MaterialXUtils/MaterialXSurfaceShaderAbstract.h b/Engine/Plugins/Interchange/Runtime/Source/Import/Public/MaterialX/MaterialXUtils/MaterialXSurfaceShaderAbstract.h
--- a/Engine/Plugins/Interchange/Runtime/Source/Import/Public/MaterialX/MaterialXUtils/MaterialXSurfaceShaderAbstract.h
+++ b/Engine/Plugins/Interchange/Runtime/Source/Import/Public/MaterialX/MaterialXUtils/MaterialXSurfaceShaderAbstract.h
@@ -88,6 +88,18 @@ protected:
 	 */
 	bool AddFloatAttribute(MaterialX::InputPtr Input, const FString& InputChannelName, UInterchangeShaderNode* ShaderNode, float DefaultValue = std::numeric_limits<float>::max());
 
+	/**
+	 * Add an integer attribute to a shader node only if its value taken from the input is not equal to its default value. Return false if the attribute does not exist or if we cannot add it
+	 *
+	 * @param Input - The MaterialX input to retrieve and add the value from. Must be of type float.
+	 * @param InputChannelName - The name of the shader node's input to add the attribute.
+	 * @param ShaderNode - The shader node to which we want to add the attribute.
+	 * @param DefaultValue - the default value to test the input against.
+	 *
+	 * @return true if the attribute was successfully added.
+	 */
+	bool AddIntegerAttribute(MaterialX::InputPtr Input, const FString& InputChannelName, UInterchangeShaderNode* ShaderNode, int DefaultValue = std::numeric_limits<int>::max());
+
 	/**
 	 * Add an FLinearColor attribute to a shader node only if its value taken from the input is not equal to its default value. Return false if the attribute does not exist or if we cannot add it.
 	 *
diff --git a/Engine/Plugins/Interchange/Runtime/Source/Nodes/Public/InterchangeMaterialDefinitions.h b/Engine/Plugins/Interchange/Runtime/Source/Nodes/Public/InterchangeMaterialDefinitions.h
--- a/Engine/Plugins/Interchange/Runtime/Source/Nodes/Public/InterchangeMaterialDefinitions.h
+++ b/Engine/Plugins/Interchange/Runtime/Source/Nodes/Public/InterchangeMaterialDefinitions.h
@@ -155,6 +155,24 @@ namespace Interchange
 					}
 				}
 
+				namespace RotateAboutAxis
+				{
+					const FName Name = TEXT("RotateAboutAxis");
+
+					namespace Attributes
+					{
+						const FName Period = TEXT("Period");
+					}
+
+					namespace Inputs
+					{
+						const FName NormalizedRotationAxis = TEXT("NormalizedRotationAxis");
+						const FName RotationAngle = TEXT("RotationAngle");
+						const FName PivotPoint = TEXT("PivotPoint");
+						const FName Position = TEXT("Position");
+					}
+				}
+
 				namespace Switch
 				{
 					const FName Name = TEXT("Switch");
@@ -291,6 +309,37 @@ namespace Interchange
 					}
 				}
 
+				namespace Trigonometry
+				{
+					namespace Inputs
+					{
+						const FName Input = TEXT("Input");
+					}
+
+					namespace Attributes
+					{
+						const FName Period = TEXT("Period");
+					}
+				}
+
+				namespace Cosine
+				{
+					const FName Name = TEXT("Cosine");
+					using namespace Trigonometry;
+				}
+
+				namespace Sine
+				{
+					const FName Name = TEXT("Sine");
+					using namespace Trigonometry;
+				}
+
+				namespace Tangent
+				{
+					const FName Name = TEXT("Tangent");
+					using namespace Trigonometry;
+				}
+
 				namespace VectorNoise
 				{
 
diff --git a/Engine/Plugins/Interchange/Runtime/Source/Pipelines/Private/InterchangeGenericMaterialPipeline.cpp b/Engine/Plugins/Interchange/Runtime/Source/Pipelines/Private/InterchangeGenericMaterialPipeline.cpp
--- a/Engine/Plugins/Interchange/Runtime/Source/Pipelines/Private/InterchangeGenericMaterialPipeline.cpp
+++ b/Engine/Plugins/Interchange/Runtime/Source/Pipelines/Private/InterchangeGenericMaterialPipeline.cpp
@@ -31,12 +31,16 @@
 #include "Materials/MaterialExpressionMultiply.h"
 #include "Materials/MaterialExpressionNoise.h"
 #include "Materials/MaterialExpressionOneMinus.h"
+#include "Materials/MaterialExpressionRotateAboutAxis.h"
 #include "Materials/MaterialExpressionRotator.h"
 #include "Materials/MaterialExpressionScalarParameter.h"
 #include "Materials/MaterialExpressionStaticBoolParameter.h"
 #include "Materials/MaterialExpressionConstant.h"
 #include "Materials/MaterialExpressionConstant3Vector.h"
+#include "Materials/MaterialExpressionCosine.h"
+#include "Materials/MaterialExpressionSine.h"
 #include "Materials/MaterialExpressionSwitch.h"
+#include "Materials/MaterialExpressionTangent.h"
 #include "Materials/MaterialExpressionTextureCoordinate.h"
 #include "Materials/MaterialExpressionTextureObject.h"
 #include "Materials/MaterialExpressionTextureSample.h"
@@ -2403,6 +2407,65 @@ void UInterchangeGenericMaterialPipeline::HandleRotatorNode(const UInterchangeSh
 	}
 }
 
+void UInterchangeGenericMaterialPipeline::HandleRotateAboutAxisNode(const UInterchangeShaderNode* ShaderNode, UInterchangeBaseMaterialFactoryNode* MaterialFactoryNode, UInterchangeMaterialExpressionFactoryNode* RotateAboutAxisFactoryNode)
+{
+	using namespace UE::Interchange::Materials::Standard::Nodes;
+	RotateAboutAxisFactoryNode->SetCustomExpressionClassName(UMaterialExpressionRotateAboutAxis::StaticClass()->GetName());
+
+	// NormalizedRotationAxis
+	{
+		TTuple<UInterchangeMaterialExpressionFactoryNode*, FString> NormalizedRotationAxisExpression =
+			CreateMaterialExpressionForInput(MaterialFactoryNode, ShaderNode, RotateAboutAxis::Inputs::NormalizedRotationAxis.ToString(), RotateAboutAxisFactoryNode->GetUniqueID());
+		if (NormalizedRotationAxisExpression.Get<0>())
+		{
+			UInterchangeShaderPortsAPI::ConnectOuputToInputByName(RotateAboutAxisFactoryNode, GET_MEMBER_NAME_CHECKED(UMaterialExpressionRotateAboutAxis, NormalizedRotationAxis).ToString(),
+																  NormalizedRotationAxisExpression.Get<0>()->GetUniqueID(), NormalizedRotationAxisExpression.Get<1>());
+		}
+	}
+
+	// PivotPoint
+	{
+		TTuple<UInterchangeMaterialExpressionFactoryNode*, FString> RotateAboutAxisExpression =
+			CreateMaterialExpressionForInput(MaterialFactoryNode, ShaderNode, RotateAboutAxis::Inputs::PivotPoint.ToString(), RotateAboutAxisFactoryNode->GetUniqueID());
+		if (RotateAboutAxisExpression.Get<0>())
+		{
+			UInterchangeShaderPortsAPI::ConnectOuputToInputByName(RotateAboutAxisFactoryNode, GET_MEMBER_NAME_CHECKED(UMaterialExpressionRotateAboutAxis, PivotPoint).ToString(),
+																  RotateAboutAxisExpression.Get<0>()->GetUniqueID(), RotateAboutAxisExpression.Get<1>());
+		}
+	}
+
+	// Position
+	{
+		TTuple<UInterchangeMaterialExpressionFactoryNode*, FString> PositionExpression =
+			CreateMaterialExpressionForInput(MaterialFactoryNode, ShaderNode, RotateAboutAxis::Inputs::Position.ToString(), RotateAboutAxisFactoryNode->GetUniqueID());
+		if (PositionExpression.Get<0>())
+		{
+			UInterchangeShaderPortsAPI::ConnectOuputToInputByName(RotateAboutAxisFactoryNode, GET_MEMBER_NAME_CHECKED(UMaterialExpressionRotateAboutAxis, Position).ToString(),
+																  PositionExpression.Get<0>()->GetUniqueID(), PositionExpression.Get<1>());
+		}
+	}
+
+	// RotationAngle
+	{
+		TTuple<UInterchangeMaterialExpressionFactoryNode*, FString> RotationAngleExpression =
+			CreateMaterialExpressionForInput(MaterialFactoryNode, ShaderNode, RotateAboutAxis::Inputs::RotationAngle.ToString(), RotateAboutAxisFactoryNode->GetUniqueID());
+		if (RotationAngleExpression.Get<0>())
+		{
+			UInterchangeShaderPortsAPI::ConnectOuputToInputByName(RotateAboutAxisFactoryNode, GET_MEMBER_NAME_CHECKED(UMaterialExpressionRotateAboutAxis, RotationAngle).ToString(),
+																  RotationAngleExpression.Get<0>()->GetUniqueID(), RotationAngleExpression.Get<1>());
+		}
+	}
+
+	// Period (By default the period is 1, but the angle is in radians, let's default it to 2Pi
+	{
+		float Period = 2.f * UE_PI;
+		ShaderNode->GetFloatAttribute(RotateAboutAxis::Attributes::Period.ToString(), Period);
+		const FName PeriodMemberName = GET_MEMBER_NAME_CHECKED(UMaterialExpressionRotateAboutAxis, Period);
+		RotateAboutAxisFactoryNode->AddFloatAttribute(PeriodMemberName.ToString(), Period);
+		RotateAboutAxisFactoryNode->AddApplyAndFillDelegates<float>(PeriodMemberName.ToString(), UMaterialExpressionRotateAboutAxis::StaticClass(), PeriodMemberName);
+	}
+}
+
 void UInterchangeGenericMaterialPipeline::HandleTimeNode(const UInterchangeShaderNode* ShaderNode, UInterchangeBaseMaterialFactoryNode* MaterialFactoryNode, UInterchangeMaterialExpressionFactoryNode* TimeFactoryNode)
 {
 	using namespace UE::Interchange::Materials::Standard::Nodes;
@@ -2738,6 +2801,47 @@ void UInterchangeGenericMaterialPipeline::HandleSwitchNode(const UInterchangeSha
 	}
 }
 
+/** Trigonometry expressions are unitless and sets the period to 1 by default to let the user sets their period*/
+void UInterchangeGenericMaterialPipeline::HandleTrigonometryNode(const UInterchangeShaderNode* ShaderNode, UClass* StaticClass, UInterchangeBaseMaterialFactoryNode* MaterialFactoryNode, UInterchangeMaterialExpressionFactoryNode* TrigonometryFactoryNode)
+{
+	using namespace UE::Interchange::Materials::Standard::Nodes;
+
+	using StringView =
+#if PLATFORM_WIDECHAR_IS_CHAR16
+		std::u16string_view;
+#else
+		std::wstring_view;
+#endif
+
+	// ensure that trig expressions have the Input and Period members;
+	static_assert(StringView(GET_MEMBER_NAME_STRING_CHECKED(UMaterialExpressionSine, Input)) == StringView(GET_MEMBER_NAME_STRING_CHECKED(UMaterialExpressionCosine, Input)));
+	static_assert(StringView(GET_MEMBER_NAME_STRING_CHECKED(UMaterialExpressionSine, Input)) == StringView(GET_MEMBER_NAME_STRING_CHECKED(UMaterialExpressionTangent, Input)));
+	static_assert(StringView(GET_MEMBER_NAME_STRING_CHECKED(UMaterialExpressionSine, Period)) == StringView(GET_MEMBER_NAME_STRING_CHECKED(UMaterialExpressionCosine, Period)));
+	static_assert(StringView(GET_MEMBER_NAME_STRING_CHECKED(UMaterialExpressionSine, Period)) == StringView(GET_MEMBER_NAME_STRING_CHECKED(UMaterialExpressionTangent, Period)));
+
+	TrigonometryFactoryNode->SetCustomExpressionClassName(StaticClass->GetName());
+
+	// Input
+	{
+		TTuple<UInterchangeMaterialExpressionFactoryNode*, FString> InputExpression =
+			CreateMaterialExpressionForInput(MaterialFactoryNode, ShaderNode, Trigonometry::Inputs::Input.ToString(), TrigonometryFactoryNode->GetUniqueID());
+		if (InputExpression.Get<0>())
+		{
+			UInterchangeShaderPortsAPI::ConnectOuputToInputByName(TrigonometryFactoryNode, GET_MEMBER_NAME_CHECKED(UMaterialExpressionSine, Input).ToString(),
+																  InputExpression.Get<0>()->GetUniqueID(), InputExpression.Get<1>());
+		}
+	}
+
+	// Period (should be 2Pi, unlike sine, a period of 0 always return cos(0) / tan(0) instead of cos(Input)/tan(Input), to be consistent over all 3, let's just set by default 2Pi
+	{
+		float Period = 2.f * UE_PI;
+		ShaderNode->GetFloatAttribute(Trigonometry::Attributes::Period.ToString(), Period);
+		const FName PeriodMemberName = GET_MEMBER_NAME_CHECKED(UMaterialExpressionSine, Period);
+		TrigonometryFactoryNode->AddFloatAttribute(PeriodMemberName.ToString(), Period);
+		TrigonometryFactoryNode->AddApplyAndFillDelegates<float>(PeriodMemberName.ToString(), StaticClass, PeriodMemberName);
+	}
+}
+
 void UInterchangeGenericMaterialPipeline::HandleScalarParameterNode(const UInterchangeShaderNode* ShaderNode, UInterchangeMaterialExpressionFactoryNode* ScalarParameterFactoryNode)
 {
 	using namespace UE::Interchange::Materials::Standard::Nodes;
@@ -2846,6 +2950,10 @@ UInterchangeMaterialExpressionFactoryNode* UInterchangeGenericMaterialPipeline::
 	{
 		HandleNormalFromHeightMapNode(ShaderNode, MaterialFactoryNode, MaterialExpression);
 	}
+	else if(*ShaderType == Nodes::RotateAboutAxis::Name)
+	{
+		HandleRotateAboutAxisNode(ShaderNode, MaterialFactoryNode, MaterialExpression);
+	}
 	else if(*ShaderType == Nodes::Rotator::Name)
 	{
 		HandleRotatorNode(ShaderNode, MaterialFactoryNode, MaterialExpression);
@@ -2886,6 +2994,18 @@ UInterchangeMaterialExpressionFactoryNode* UInterchangeGenericMaterialPipeline::
 	{
 		HandleVectorNoiseNode(ShaderNode, MaterialFactoryNode, MaterialExpression);
 	}
+	else if(*ShaderType == Nodes::Cosine::Name)
+	{
+		HandleTrigonometryNode(ShaderNode, UMaterialExpressionCosine::StaticClass(), MaterialFactoryNode, MaterialExpression);
+	}
+	else if(*ShaderType == Nodes::Sine::Name)
+	{
+		HandleTrigonometryNode(ShaderNode, UMaterialExpressionSine::StaticClass(), MaterialFactoryNode, MaterialExpression);
+	}
+	else if(*ShaderType == Nodes::Tangent::Name)
+	{
+		HandleTrigonometryNode(ShaderNode, UMaterialExpressionTangent::StaticClass(), MaterialFactoryNode, MaterialExpression);
+	}
 	else if (*ShaderType == Nodes::ScalarParameter::Name)
 	{
 		HandleScalarParameterNode(ShaderNode, MaterialExpression);
diff --git a/Engine/Plugins/Interchange/Runtime/Source/Pipelines/Public/InterchangeGenericMaterialPipeline.h b/Engine/Plugins/Interchange/Runtime/Source/Pipelines/Public/InterchangeGenericMaterialPipeline.h
--- a/Engine/Plugins/Interchange/Runtime/Source/Pipelines/Public/InterchangeGenericMaterialPipeline.h
+++ b/Engine/Plugins/Interchange/Runtime/Source/Pipelines/Public/InterchangeGenericMaterialPipeline.h
@@ -212,6 +212,7 @@ private:
 	void HandleLerpNode(const UInterchangeShaderNode* ShaderNode, UInterchangeBaseMaterialFactoryNode* MaterialFactoryNode, UInterchangeMaterialExpressionFactoryNode* LerpFactoryNode);
 	void HandleMaskNode(const UInterchangeShaderNode* ShaderNode, UInterchangeBaseMaterialFactoryNode* MaterialFactoryNode, UInterchangeMaterialExpressionFactoryNode* MaskFactoryNode);
 	void HandleRotatorNode(const UInterchangeShaderNode* ShaderNode, UInterchangeBaseMaterialFactoryNode* MaterialFactoryNode, UInterchangeMaterialExpressionFactoryNode* RotatorFactoryNode);
+	void HandleRotateAboutAxisNode(const UInterchangeShaderNode* ShaderNode, UInterchangeBaseMaterialFactoryNode* MaterialFactoryNode, UInterchangeMaterialExpressionFactoryNode* RotateAboutAxisFactoryNode);
 	void HandleTimeNode(const UInterchangeShaderNode* ShaderNode, UInterchangeBaseMaterialFactoryNode* MaterialFactoryNode, UInterchangeMaterialExpressionFactoryNode* TimeFactoryNode);
 	void HandleTransformPositionNode(const UInterchangeShaderNode* ShaderNode, UInterchangeBaseMaterialFactoryNode* MaterialFactoryNode, UInterchangeMaterialExpressionFactoryNode* TransformPositionFactoryNode);
 	void HandleTransformVectorNode(const UInterchangeShaderNode* ShaderNode, UInterchangeBaseMaterialFactoryNode* MaterialFactoryNode, UInterchangeMaterialExpressionFactoryNode* TransformVectorFactoryNode);
@@ -219,6 +220,8 @@ private:
 	void HandleVectorNoiseNode(const UInterchangeShaderNode* ShaderNode, UInterchangeBaseMaterialFactoryNode* MaterialFactoryNode, UInterchangeMaterialExpressionFactoryNode* NoiseFactoryNode);
 	void HandleSwizzleNode(const UInterchangeShaderNode* ShaderNode, UInterchangeBaseMaterialFactoryNode* MaterialFactoryNode, UInterchangeMaterialExpressionFactoryNode* SwizzleFactoryNode);
 	void HandleSwitchNode(const UInterchangeShaderNode* ShaderNode, UInterchangeBaseMaterialFactoryNode* MaterialFactoryNode, UInterchangeMaterialExpressionFactoryNode* SwitchFactoryNode);
+	void HandleTrigonometryNode(const UInterchangeShaderNode* ShaderNode, UClass* StaticClass, UInterchangeBaseMaterialFactoryNode* MaterialFactoryNode, UInterchangeMaterialExpressionFactoryNode* TrigonometryFactoryNode);
+
 	void HandleScalarParameterNode(const UInterchangeShaderNode* ShaderNode, UInterchangeMaterialExpressionFactoryNode* ScalarParameterFactoryNode);
 	void HandleVectorParameterNode(const UInterchangeShaderNode* ShaderNode, UInterchangeMaterialExpressionFactoryNode* VectorParameterFactoryNode);
 	void HandleStaticBooleanParameterNode(const UInterchangeShaderNode* ShaderNode, UInterchangeMaterialExpressionFactoryNode* StaticBooleanParameterFactoryNode);
