MaterialX:
- Add default inputs from nodedef for ifgreater, ifgreatereq, ifequal

diff --git a/Engine/Plugins/Interchange/Runtime/Source/Import/Private/MaterialX/MaterialXUtils/MaterialXSurfaceShaderAbstract.cpp b/Engine/Plugins/Interchange/Runtime/Source/Import/Private/MaterialX/MaterialXUtils/MaterialXSurfaceShaderAbstract.cpp
--- a/Engine/Plugins/Interchange/Runtime/Source/Import/Private/MaterialX/MaterialXUtils/MaterialXSurfaceShaderAbstract.cpp
+++ b/Engine/Plugins/Interchange/Runtime/Source/Import/Private/MaterialX/MaterialXUtils/MaterialXSurfaceShaderAbstract.cpp
@@ -333,7 +333,7 @@ void FMaterialXSurfaceShaderAbstract::ConnectNodeCategoryOutputToInput(const Mat
 		// Add all the inputs from the nodedef in case the node doesn't have them all (we don't have a 1:1 match on the default values otherwise)
 		// For exemple min defaults to (0,0) whereas UE to (0, 1)
 		// We don't want to add every node, since we don't a have perfect match for each input
-		static const std::set<std::string> NodeDefsCategories{"min", "max", "modulo"};
+		static const std::set<std::string> NodeDefsCategories{"ifgreater", "ifgreatereq", "ifequal", "min", "max", "modulo"};
 		if(NodeDefsCategories.find(UpstreamNode->getCategory()) != NodeDefsCategories.cend())
 		{
 			UpstreamNode->addInputsFromNodeDef();
