[Interchange] Fix skin to time zero bind pose to support multiple bind poses for bones.
This happen when we have multiple mesh skin with different bind pose, but no valid bind pose in the fbx file.

diff --git a/Engine/Plugins/Interchange/Runtime/Source/Import/Private/Mesh/InterchangeSkeletalMeshFactory.cpp b/Engine/Plugins/Interchange/Runtime/Source/Import/Private/Mesh/InterchangeSkeletalMeshFactory.cpp
--- a/Engine/Plugins/Interchange/Runtime/Source/Import/Private/Mesh/InterchangeSkeletalMeshFactory.cpp
+++ b/Engine/Plugins/Interchange/Runtime/Source/Import/Private/Mesh/InterchangeSkeletalMeshFactory.cpp
@@ -377,6 +377,19 @@ namespace UE
 					CommonPipelineDataFactoryNode->GetCustomGlobalOffsetTransform(GlobalOffsetTransform);
 				}
 
+				//If we use transform, some scaling are not well supported, so we must do all the math with matrix
+				FMatrix SceneNodeMatrixInverse = SceneNodeTransform.ToMatrixWithScale().Inverse();
+				FMatrix GeometricMatrix = FMatrix::Identity;
+
+				//Scope any transform variable to be sure they are not use to create the vertex transformation matrix
+				{
+					FTransform GeometricTransform;
+					if (SceneNode->GetCustomGeometricTransform(GeometricTransform))
+					{
+						GeometricMatrix = GeometricTransform.ToMatrixWithScale();
+					}
+				}
+
 				const int32 JointCount = JointNames.Num();
 				for (int32 JointIndex = 0; JointIndex < JointCount; ++JointIndex)
 				{
@@ -388,35 +401,48 @@ namespace UE
 						continue;
 					}
 
-					FTransform BindPose;
-					if (!ensure(JointNode->GetCustomBindPoseGlobalTransform(NodeContainer, FTransform::Identity, BindPose)))
-					{
-						//BindPose will fall back on LocalTransform in case its not present.
-						//If neither is present: No Value to convert from, skip this joint.
-						continue;
-					}
+					FMatrix JointBindPoseMatrix = FMatrix::Identity;;
+					FMatrix JointT0Matrix = FMatrix::Identity;;
+					FMatrix JointReferenceMatrix = FMatrix::Identity;;
 
-					FTransform T0;
-					if (!JointNode->GetCustomTimeZeroGlobalTransform(NodeContainer, FTransform::Identity, T0))
+					//Scope any transform variable to be sure they are not use to create the vertex transformation matrix
 					{
-						//If there is no time zero global transform we cannot set the bind pose to time zero.
-						//We must skip this joint.
-						continue;
-					}
+						FString AttributeKey = TEXT("JointBindPosePerMesh_") + MeshNode->GetUniqueID();
+						if (!JointNode->GetAttribute<FMatrix>(AttributeKey, JointBindPoseMatrix))
+						{
+							FTransform JointBindPose;
+							if (!ensure(JointNode->GetCustomBindPoseGlobalTransform(NodeContainer, FTransform::Identity, JointBindPose)))
+							{
+								//BindPose will fall back on LocalTransform in case its not present.
+								//If neither is present: No Value to convert from, skip this joint.
+								continue;
+							}
+							JointBindPoseMatrix = JointBindPose.ToMatrixWithScale();
+						}
+					
+						FTransform JointT0;
+						if (!JointNode->GetCustomTimeZeroGlobalTransform(NodeContainer, FTransform::Identity, JointT0))
+						{
+							//If there is no time zero global transform we cannot set the bind pose to time zero.
+							//We must skip this joint.
+							continue;
+						}
+						JointT0Matrix = JointT0.ToMatrixWithScale();
 
-					FMatrix TransformMatrix = FMatrix::Identity;
-					JointNode->GetGlobalBindPoseReferenceForMeshUID(MeshNode->GetUniqueID(), TransformMatrix);
-					FTransform TransformMatrixTransform;
-					TransformMatrixTransform.SetFromMatrix(TransformMatrix);
+						if (!JointNode->GetGlobalBindPoseReferenceForMeshUID(MeshNode->GetUniqueID(), JointReferenceMatrix))
+						{
+							//Skip this joint
+							continue;
+						}
 
-					FTransform GeometricTransform;
-					if (SceneNode->GetCustomGeometricTransform(GeometricTransform))
-					{
 						//We must add the geometric node transform
-						TransformMatrixTransform = GeometricTransform * TransformMatrixTransform;
+						JointReferenceMatrix = GeometricMatrix * JointReferenceMatrix;
 					}
-
-					FMatrix VertexTransformMatrix = (SceneNodeTransform.Inverse() * ((TransformMatrixTransform * BindPose.Inverse()) * (T0 * SceneNodeTransform.Inverse()))).ToMatrixWithScale();
+					
+					FMatrix ReferenceMatrixRelativeToBindPose = JointReferenceMatrix * JointBindPoseMatrix.Inverse();
+					FMatrix jointTimeZeroRelativeToMeshNodeMatrix = JointT0Matrix * SceneNodeMatrixInverse;
+					//Compute the transform to apply to vertex to change the bind pose from the existing one to use time zero pose instead
+					FMatrix VertexTransformMatrix = SceneNodeMatrixInverse * (ReferenceMatrixRelativeToBindPose * jointTimeZeroRelativeToMeshNodeMatrix);
 
 					//Iterate all bone vertices
 					for (FVertexID VertexID : MeshDescription.Vertices().GetElementIDs())
@@ -1037,6 +1063,15 @@ namespace UE
 					UE::Interchange::Private::FSkeletonHelper::RecursiveBoneHasBindPose(Arguments.NodeContainer, RootJointNodeId, bHasBoneWithoutBindPose);
 					if (bHasBoneWithoutBindPose)
 					{
+						if (!GIsAutomationTesting && Arguments.Translator)
+						{
+							FString Filename = Arguments.SourceData ? Arguments.SourceData->GetFilename() : TEXT("No specified file");
+							UInterchangeResultWarning_Generic* Message = Arguments.Translator->AddMessage<UInterchangeResultWarning_Generic>();
+							Message->Text = NSLOCTEXT("InterchangeSkeletalMeshFactory", "CreatePayloadTasks_ForceRebindOfSkinWithTimeZeroPose", "Imported skeletal mesh has some invalid bind poses. Skeletal mesh skinning has been rebind using the time zero pose.");
+							Message->SourceAssetName = Filename;
+							Message->AssetFriendlyName = Arguments.AssetName;
+							Message->AssetType = USkeletalMesh::StaticClass();
+						}
 						ImportAssetObjectLODDataRef.bUseTimeZeroAsBindPose = true;
 					}
 				}
@@ -1063,12 +1098,14 @@ namespace UE
 				{
 					// Temporary 5.5.2: fbx legacy has a special way to bake the skeletal mesh that do not fit the interchange standard.
 					// This fix the issue with blender armature bone skip if the armature has transform not identity.
-					const UInterchangeSourceNode* SourceNode = UInterchangeSourceNode::GetUniqueInstance(Arguments.NodeContainer);
-					bool bUseSceneNodeGlobalTransform = false;
-					SourceNode->GetBooleanAttribute(TEXT("SkeletalMesh-Bake-UseSceneNodeGlobalTransform"), bUseSceneNodeGlobalTransform);
-					if (bUseSceneNodeGlobalTransform)
+					if(const UInterchangeSourceNode* SourceNode = UInterchangeSourceNode::GetUniqueInstance(Arguments.NodeContainer))
 					{
-						BakeFromRootJointTransfromModifier.SetIdentity();
+						bool bUseSceneNodeGlobalTransform = false;
+						SourceNode->GetBooleanAttribute(TEXT("SkeletalMesh-Bake-UseSceneNodeGlobalTransform"), bUseSceneNodeGlobalTransform);
+						if (bUseSceneNodeGlobalTransform)
+						{
+							BakeFromRootJointTransfromModifier.SetIdentity();
+						}
 					}
 				}
 
@@ -1242,7 +1279,6 @@ void UInterchangeSkeletalMeshFactory::CreatePayloadTasks(const FImportAssetObjec
 		FImportAssetObjectLODData& ImportAssetObjectLODData = ImportAssetObjectData.LodDatas.AddDefaulted_GetRef();
 		ImportAssetObjectLODData.LodIndex = CurrentLodIndex;
 
-		bool bUseTimeZeroAsBindPose = false;
 		SkeletonNode->GetCustomUseTimeZeroForBindPose(ImportAssetObjectLODData.bUseTimeZeroAsBindPose);
 		
 		Private::BuildMeshReferences(
diff --git a/Engine/Plugins/Interchange/Runtime/Source/Parsers/Fbx/Private/FbxMesh.cpp b/Engine/Plugins/Interchange/Runtime/Source/Parsers/Fbx/Private/FbxMesh.cpp
--- a/Engine/Plugins/Interchange/Runtime/Source/Parsers/Fbx/Private/FbxMesh.cpp
+++ b/Engine/Plugins/Interchange/Runtime/Source/Parsers/Fbx/Private/FbxMesh.cpp
@@ -1992,13 +1992,14 @@ void FFbxMesh::AddAllMeshes(FbxScene* SDKScene, FbxGeometryConverter* SDKGeometr
 	} // for GeometryCount
 }
 
-bool FFbxMesh::GetGlobalJointBindPoseTransform(FFbxParser* Parser, FbxScene* SDKScene, FbxNode* Joint, FbxAMatrix& GlobalBindPoseJointMatrix, TMap<FString, FMatrix>& MeshIdToGlobalBindPoseReferenceMap, bool& bBadBindPoseMessageDisplay)
+bool FFbxMesh::GetGlobalJointBindPoseTransform(FFbxParser* Parser, FbxScene* SDKScene, FbxNode* Joint, FbxAMatrix& GlobalBindPoseJointMatrix, TMap<FString, FMatrix>& MeshIdToGlobalBindPoseJointMap, TMap<FString, FMatrix>& MeshIdToGlobalBindPoseReferenceMap, bool& bBadBindPoseMessageDisplay)
 {
 	FbxManager* SDKManager = SDKScene->GetFbxManager();
 	//First look for Cluster and then look in the bind pose if no cluster was found
 
 	//Search all skeletalmesh(FbxGeometry with valid deformer) using this joint and see if there is a valid FbxCluster
 	//containing a TransformLinkMatrix for this joint
+	bool bGlobalBindPoseJointMatrixSet = false;
 	const int32 GeometryCount = SDKScene->GetGeometryCount();
 	for (int32 GeometryIndex = 0; GeometryIndex < GeometryCount; ++GeometryIndex)
 	{
@@ -2033,11 +2034,18 @@ bool FFbxMesh::GetGlobalJointBindPoseTransform(FFbxParser* Parser, FbxScene* SDK
 					uint64 ClusterId = Cluster->GetUniqueID();
 					uint64 MeshId = Geometry->GetUniqueID();
 
-					Cluster->GetTransformLinkMatrix(GlobalBindPoseJointMatrix);
+					FbxAMatrix GlobalBindPoseJointMatrixPerGeo;
+					Cluster->GetTransformLinkMatrix(GlobalBindPoseJointMatrixPerGeo);
+					MeshIdToGlobalBindPoseJointMap.Add(MeshUniqueID, FFbxConvert::ConvertMatrix<FMatrix>(GlobalBindPoseJointMatrixPerGeo));
+					if (!bGlobalBindPoseJointMatrixSet)
+					{
+						//We use the first geo bind pose
+						GlobalBindPoseJointMatrix = GlobalBindPoseJointMatrixPerGeo;
+						bGlobalBindPoseJointMatrixSet = true;
+					}
 
 					FbxAMatrix GlobalBindPoseReferenceMatrix;
 					Cluster->GetTransformMatrix(GlobalBindPoseReferenceMatrix);
-					
 					MeshIdToGlobalBindPoseReferenceMap.Add(MeshUniqueID, FFbxConvert::ConvertMatrix<FMatrix>(GlobalBindPoseReferenceMatrix));
 
 					//Presumes only 1 cluster per geometry that matches criteria:
diff --git a/Engine/Plugins/Interchange/Runtime/Source/Parsers/Fbx/Private/FbxMesh.h b/Engine/Plugins/Interchange/Runtime/Source/Parsers/Fbx/Private/FbxMesh.h
--- a/Engine/Plugins/Interchange/Runtime/Source/Parsers/Fbx/Private/FbxMesh.h
+++ b/Engine/Plugins/Interchange/Runtime/Source/Parsers/Fbx/Private/FbxMesh.h
@@ -196,7 +196,7 @@ namespace UE
 				{}
 
 				void AddAllMeshes(FbxScene* SDKScene, FbxGeometryConverter* SDKGeometryConverter, UInterchangeBaseNodeContainer& NodeContainer, TMap<FString, TSharedPtr<FPayloadContextBase>>& PayloadContexts);
-				static bool GetGlobalJointBindPoseTransform(FFbxParser* Parser, FbxScene* SDKScene, FbxNode* Joint, FbxAMatrix& GlobalBindPoseJointMatrix, TMap<FString, FMatrix>& MeshIdToGlobalBindPoseReferenceMap, bool& bBadBindPoseMessageDisplay);
+				static bool GetGlobalJointBindPoseTransform(FFbxParser* Parser, FbxScene* SDKScene, FbxNode* Joint, FbxAMatrix& GlobalBindPoseJointMatrix, TMap<FString, FMatrix>& MeshIdToGlobalBindPoseJointMap, TMap<FString, FMatrix>& MeshIdToGlobalBindPoseReferenceMap, bool& bBadBindPoseMessageDisplay);
 
 				TArray<FMorphTargetAnimationBuildingData>& GetMorphTargetAnimationsBuildingData() { return MorphTargetAnimationsBuildingData; }
 			protected:
diff --git a/Engine/Plugins/Interchange/Runtime/Source/Parsers/Fbx/Private/FbxScene.cpp b/Engine/Plugins/Interchange/Runtime/Source/Parsers/Fbx/Private/FbxScene.cpp
--- a/Engine/Plugins/Interchange/Runtime/Source/Parsers/Fbx/Private/FbxScene.cpp
+++ b/Engine/Plugins/Interchange/Runtime/Source/Parsers/Fbx/Private/FbxScene.cpp
@@ -206,11 +206,18 @@ namespace UE
 					//Get the bind pose transform for this joint
 					FbxAMatrix GlobalBindPoseJointMatrix = SDKScene->GetAnimationEvaluator()->GetNodeGlobalTransform(Node, 0);
 					TMap<FString, FMatrix> MeshIdToGlobalBindPoseReferenceMap;
-
-					FFbxMesh::GetGlobalJointBindPoseTransform(&Parser, SDKScene, Node, GlobalBindPoseJointMatrix, MeshIdToGlobalBindPoseReferenceMap, bBadBindPoseMessageDisplay);
+					TMap<FString, FMatrix> MeshIdToGlobalBindPoseJointMap;
+					FFbxMesh::GetGlobalJointBindPoseTransform(&Parser, SDKScene, Node, GlobalBindPoseJointMatrix, MeshIdToGlobalBindPoseJointMap, MeshIdToGlobalBindPoseReferenceMap, bBadBindPoseMessageDisplay);
 
 					FTransform GlobalBindPoseJointTransform = GetConvertedTransform(GlobalBindPoseJointMatrix);
 					UnrealNode->SetGlobalBindPoseReferenceForMeshUIDs(MeshIdToGlobalBindPoseReferenceMap);
+					//Add Transform attribute per Geometry
+					//Temporary (to fix in 5.6 main branch): We need to make an API like Get/Set GlobalBindPoseReferenceForMeshUIDs for the joint one
+					for (const TPair<FString, FMatrix>& MeshIdAndBindPoseJointMatrix : MeshIdToGlobalBindPoseJointMap)
+					{
+						FString AttributeKey = TEXT("JointBindPosePerMesh_") + MeshIdAndBindPoseJointMatrix.Key;
+						UnrealNode->RegisterAttribute<FMatrix>(FAttributeKey(AttributeKey), MeshIdAndBindPoseJointMatrix.Value);
+					}
 
 					FbxNode* ParentNode = Node->GetParent();
 					
@@ -218,7 +225,8 @@ namespace UE
 					{
 						FbxAMatrix GlobalFbxParentMatrix = SDKScene->GetAnimationEvaluator()->GetNodeGlobalTransform(ParentNode, 0);
 						TMap<FString, FMatrix> ParentMeshIdToGlobalBindPoseReferenceMap;
-						FFbxMesh::GetGlobalJointBindPoseTransform(&Parser, SDKScene, ParentNode, GlobalFbxParentMatrix, ParentMeshIdToGlobalBindPoseReferenceMap, bBadBindPoseMessageDisplay);
+						TMap<FString, FMatrix> ParentMeshIdToGlobalBindPoseJointMap;
+						FFbxMesh::GetGlobalJointBindPoseTransform(&Parser, SDKScene, ParentNode, GlobalFbxParentMatrix, ParentMeshIdToGlobalBindPoseJointMap, ParentMeshIdToGlobalBindPoseReferenceMap, bBadBindPoseMessageDisplay);
 						
 						FbxAMatrix LocalFbxMatrix = GlobalFbxParentMatrix.Inverse() * GlobalBindPoseJointMatrix;
 						FTransform LocalBindPoseJointTransform = GetConvertedTransform(LocalFbxMatrix);
@@ -235,10 +243,10 @@ namespace UE
 					{
 						//NOTE:
 						// Legacy FBX uses the following Matrix calculation for moving Vertices to T0:
-						//		VertexTransformMatrix = ((TransformMatrix * BindPose.Inverse()) * (T0 * GlobalMeshTransformMatrix.Inverse()));
-						//					TransformMatrix				:= GlobalBindPoseReferenceForMeshUIDs (this is joint and Mesh dependent) => seems very FBX specific
-						//					BindPose					:= GlobalBindPose
-						//					T0							:= TimeZero
+						//		VertexTransformMatrix = ((JointReference * JointBindPose.Inverse()) * (JointT0 * GlobalMeshTransformMatrix.Inverse()));
+						//					JointReference				:= MeshIdToGlobalBindPoseReferenceMap * GeometricTransform (each mesh part can have a different bind position for the node)
+						//					JointBindPose				:= MeshIdToGlobalBindPoseJointMap (each mesh part can have a different bind position for the node)
+						//					T0							:= The joint evaluate a time 0
 						//					GlobalMeshTransformMatrix	:= Mesh's Node's GlobalTransform * GeometricTransform (Interchange.SceneNodeTransform)
 
 						//Set the global node transform
