[Interchange] Fix the Source timeline option for fbx to be the same has legacy fbx importer. It should represent the cursor of the dcc timeline and not the entire timeline.

diff --git a/Engine/Plugins/Interchange/Runtime/Source/Import/Private/Gltf/InterchangeGltfTranslator.cpp b/Engine/Plugins/Interchange/Runtime/Source/Import/Private/Gltf/InterchangeGltfTranslator.cpp
--- a/Engine/Plugins/Interchange/Runtime/Source/Import/Private/Gltf/InterchangeGltfTranslator.cpp
+++ b/Engine/Plugins/Interchange/Runtime/Source/Import/Private/Gltf/InterchangeGltfTranslator.cpp
@@ -509,45 +509,6 @@ bool UInterchangeGLTFTranslator::Translate( UInterchangeBaseNodeContainer& NodeC
 		SourceNode->SetExtraInformation(Extra.Name, Extra.Value);
 	}
 
-	//Set Animation Time related values:
-	if (GltfAsset.Animations.Num())
-	{
-		// glTF animations are tied to seconds, not framerate
-		// glTF animations do no support Time Line.
-		// glTF animations always start at 0 seconds
-		SourceNode->SetCustomSourceTimelineStart(0);
-		
-		//To acquire the AnimatedTimeEnd:
-		//	Find the last entry of each animation.samplers and find the biggest entry (aka the time (in seconds) which is the biggest).
-		float AnimatedTimeEnd = 0;
-		for (const GLTF::FAnimation& Animation : GltfAsset.Animations)
-		{
-			for (const GLTF::FAnimation::FSampler& Sampler : Animation.Samplers)
-			{
-				TArray<float> FrameTimeBuffer;
-
-				Sampler.Input.GetFloatArray(FrameTimeBuffer);
-
-				if (AnimatedTimeEnd < FrameTimeBuffer.Last())
-				{
-					AnimatedTimeEnd = FrameTimeBuffer.Last();
-				}
-			}
-		}
-
-		//Making sure that all animations can fit into the FrameNumbers:
-		//Currently glTF Animations are calcualted with 30.0 frames (as glTF does not provide Frequency) :
-		constexpr float FrameRate = 30.0;
-		constexpr float SingleFrameDuration = 1.0 / FrameRate;
-
-		float CurrentFrameNumberCandidate = AnimatedTimeEnd / SingleFrameDuration;
-		int32 CurrentFrameNumber = FMath::CeilToInt(CurrentFrameNumberCandidate);
-
-		AnimatedTimeEnd = CurrentFrameNumber * SingleFrameDuration;
-
-		SourceNode->SetCustomSourceTimelineEnd(AnimatedTimeEnd);
-	}
-
 	{
 		//For sockets, crete the Axis System conversion inverse:
 
diff --git a/Engine/Plugins/Interchange/Runtime/Source/Parsers/Fbx/Private/FbxAPI.cpp b/Engine/Plugins/Interchange/Runtime/Source/Parsers/Fbx/Private/FbxAPI.cpp
--- a/Engine/Plugins/Interchange/Runtime/Source/Parsers/Fbx/Private/FbxAPI.cpp
+++ b/Engine/Plugins/Interchange/Runtime/Source/Parsers/Fbx/Private/FbxAPI.cpp
@@ -148,19 +148,13 @@ namespace UE
 				UInterchangeSourceNode* SourceNode = UInterchangeSourceNode::FindOrCreateUniqueInstance(&NodeContainer);
 				SourceNode->SetCustomAxisConversionInverseTransform(AxisConversionInverseTransform);
 
-				//Store the fbx frame rate and timeline
+				//Store the fbx frame rate
 				{
 					FrameRate = FbxTime::GetFrameRate(SDKScene->GetGlobalSettings().GetTimeMode());
 					FileDetails.FrameRate = FString::Printf(TEXT("%.2f"), FrameRate);
 					SourceNode->SetCustomSourceFrameRateNumerator(FrameRate);
 					constexpr double Denominator = 1.0;
 					SourceNode->SetCustomSourceFrameRateDenominator(Denominator);
-
-					//Timeline time span
-					FbxTimeSpan TimelineTimeSpan(FBXSDK_TIME_INFINITE, FBXSDK_TIME_MINUS_INFINITE);
-					SDKScene->GetGlobalSettings().GetTimelineDefaultTimeSpan(TimelineTimeSpan);
-					SourceNode->SetCustomSourceTimelineStart(TimelineTimeSpan.GetStart().GetSecondDouble());
-					SourceNode->SetCustomSourceTimelineEnd(TimelineTimeSpan.GetStop().GetSecondDouble());
 				}
 
 				// Temporary 5.5.2: fbx legacy has a special way to bake the skeletal mesh that do not fit the interchange standard
diff --git a/Engine/Plugins/Interchange/Runtime/Source/Parsers/Fbx/Private/FbxScene.cpp b/Engine/Plugins/Interchange/Runtime/Source/Parsers/Fbx/Private/FbxScene.cpp
--- a/Engine/Plugins/Interchange/Runtime/Source/Parsers/Fbx/Private/FbxScene.cpp
+++ b/Engine/Plugins/Interchange/Runtime/Source/Parsers/Fbx/Private/FbxScene.cpp
@@ -890,12 +890,21 @@ namespace UE
 							SkeletalAnimationTrackNode->SetCustomSkeletonNodeUid(UnrealNode->GetUniqueID());
 
 							//Calculate AnimationTime:
+
+							//Node animated time interval (Animated time for all concern bones under the skeletalmesh root bone)
 							FbxTimeSpan AnimatedInterval(FBXSDK_TIME_INFINITE, FBXSDK_TIME_MINUS_INFINITE);
 							Node->GetAnimationInterval(AnimatedInterval, CurrentAnimationStack);
-
 							SkeletalAnimationTrackNode->SetCustomAnimationStartTime(AnimatedInterval.GetStart().GetSecondDouble());
 							SkeletalAnimationTrackNode->SetCustomAnimationStopTime(AnimatedInterval.GetStop().GetSecondDouble());
 
+							//Animation stack time interval (Source timeline interval)
+							// Note - can be less or more the bones animationlength, it can be configure when exporting by manipulating the timeline in the dcc (Maya have a range cursor on the timeline)
+							FbxTimeSpan AnimStackInterval(FBXSDK_TIME_INFINITE, FBXSDK_TIME_MINUS_INFINITE);
+							AnimStackInterval = CurrentAnimationStack->GetLocalTimeSpan();
+
+							//Temporary 5.5.2: This should be replace by an API in UInterchangeSkeletalAnimationTrackNode after the hotfix is done
+							SkeletalAnimationTrackNode->AddDoubleAttribute(TEXT("SourceTimelineAnimationStartTime"), AnimStackInterval.GetStart().GetSecondDouble());
+							SkeletalAnimationTrackNode->AddDoubleAttribute(TEXT("SourceTimelineAnimationStopTime"), AnimStackInterval.GetStop().GetSecondDouble());
 							return true;
 						}
 
diff --git a/Engine/Plugins/Interchange/Runtime/Source/Pipelines/Private/InterchangeGenericAnimationPipeline.cpp b/Engine/Plugins/Interchange/Runtime/Source/Pipelines/Private/InterchangeGenericAnimationPipeline.cpp
--- a/Engine/Plugins/Interchange/Runtime/Source/Pipelines/Private/InterchangeGenericAnimationPipeline.cpp
+++ b/Engine/Plugins/Interchange/Runtime/Source/Pipelines/Private/InterchangeGenericAnimationPipeline.cpp
@@ -659,7 +659,8 @@ void UInterchangeGenericAnimationPipeline::CreateAnimSequenceFactoryNode(UInterc
 		//Some format don't fill the animation range data we fall back on the trackNode
 		if (AnimationRange == EInterchangeAnimationRange::Timeline)
 		{
-			bTimeRangeIsValid = (SourceNode && SourceNode->GetCustomSourceTimelineStart(StartTime) && SourceNode->GetCustomSourceTimelineEnd(StopTime));
+			//Temporary 5.5.2: This should be replace by an API in UInterchangeSkeletalAnimationTrackNode after the hotfix is done
+			bTimeRangeIsValid = TrackNode.GetDoubleAttribute(TEXT("SourceTimelineAnimationStartTime"), StartTime) && TrackNode.GetDoubleAttribute(TEXT("SourceTimelineAnimationStopTime"), StopTime);
 		}
 		else if (AnimationRange == EInterchangeAnimationRange::SetRange)
 		{
