USD: Add InterchangeUSDPipeline.

#virtualized

diff --git a/Engine/Plugins/Interchange/Runtime/Source/Import/Private/Usd/InterchangeUsdTranslator.cpp b/Engine/Plugins/Interchange/Runtime/Source/Import/Private/Usd/InterchangeUsdTranslator.cpp
--- a/Engine/Plugins/Interchange/Runtime/Source/Import/Private/Usd/InterchangeUsdTranslator.cpp
+++ b/Engine/Plugins/Interchange/Runtime/Source/Import/Private/Usd/InterchangeUsdTranslator.cpp
@@ -3,6 +3,7 @@
 #include "Usd/InterchangeUsdTranslator.h"
 
 #include "UnrealUSDWrapper.h"
+#include "Usd/InterchangeUsdDefinitions.h"
 #include "USDConversionUtils.h"
 #include "USDGeomMeshConversion.h"
 #include "USDLightConversion.h"
@@ -615,130 +616,6 @@ namespace UE::InterchangeUsdTranslator::Private
 		NodeContainer.AddNode(Node);
 	}
 
-	// We use this visitor to set UsdToUnreal::FParameterValue TVariant values onto UInterchangeMaterialInstanceNodes
-	struct FParameterValueVisitor
-	{
-		FParameterValueVisitor(
-			const UE::FUsdPrim& InPrim,
-			UInterchangeBaseNodeContainer& InNodeContainer,
-			UInterchangeMaterialInstanceNode& InMaterialNode,
-			const TMap<FString, int32>& InPrimvarToUVIndex
-		)
-			: Prim(InPrim)
-			, NodeContainer(InNodeContainer)
-			, MaterialNode(InMaterialNode)
-			, PrimvarToUVIndex(InPrimvarToUVIndex)
-		{
-		}
-
-		void operator()(const float Value) const
-		{
-			MaterialNode.AddScalarParameterValue(*ParameterName, Value);
-
-			// Disable the texture input since we have a direct value
-			MaterialNode.AddScalarParameterValue(*FString::Printf(TEXT("Use%sTexture"), **ParameterName), 0.0f);
-		}
-
-		void operator()(const FVector& Value) const
-		{
-			MaterialNode.AddVectorParameterValue(*ParameterName, FLinearColor{Value});
-
-			// Disable the texture input since we have a direct value
-			MaterialNode.AddScalarParameterValue(*FString::Printf(TEXT("Use%sTexture"), **ParameterName), 0.0f);
-		}
-
-		void operator()(const UsdToUnreal::FTextureParameterValue& Value) const
-		{
-			// Emit texture node itself (this is the main place where this happens)
-			// Note that the node name isn't just the texture path, as we may have multiple material users of this texture
-			// with different settings, and so we need separate translated nodes for each material and parameter
-			FString TextureUid = FString::Printf(TEXT("Texture:%s:%s"), *Prim.GetPrimPath().GetString(), **ParameterName);
-			AddTextureNode(Prim, TextureUid, Value, NodeContainer);
-
-			// Actual texture assignment
-			MaterialNode.AddTextureParameterValue(*FString::Printf(TEXT("%sTexture"), **ParameterName), TextureUid);
-			MaterialNode.AddScalarParameterValue(*FString::Printf(TEXT("Use%sTexture"), **ParameterName), 1.0f);
-
-			// UV transform
-			FLinearColor ScaleAndTranslation = FLinearColor{
-				Value.UVScale.GetVector()[0],
-				Value.UVScale.GetVector()[1],
-				Value.UVTranslation[0],
-				Value.UVTranslation[1]
-			};
-			MaterialNode.AddVectorParameterValue(*FString::Printf(TEXT("%sScaleTranslation"), **ParameterName), ScaleAndTranslation);
-			MaterialNode.AddScalarParameterValue(*FString::Printf(TEXT("%sRotation"), **ParameterName), Value.UVRotation);
-
-			// UV index
-			if (const int32* FoundIndex = PrimvarToUVIndex.Find(Value.Primvar))
-			{
-				MaterialNode.AddScalarParameterValue(*FString::Printf(TEXT("%sUVIndex"), **ParameterName), *FoundIndex);
-			}
-			else
-			{
-				UE_LOG(
-					LogUsd,
-					Warning,
-					TEXT("Failed to find primvar '%s' when setting material parameter '%s' on material '%s'. Available primvars and UV "
-						 "indices: %s.%s"),
-					*Value.Primvar,
-					**ParameterName,
-					*Prim.GetPrimPath().GetString(),
-					*UsdUtils::StringifyMap(PrimvarToUVIndex),
-					Value.Primvar.IsEmpty()
-						? TEXT(
-							  " Is your UsdUVTexture Shader missing the 'inputs:st' attribute? (It specifies which UV set to sample the texture with)"
-						  )
-						: TEXT("")
-				);
-			}
-
-			// Component mask (which channel of the texture to use)
-			FLinearColor ComponentMask = FLinearColor::Black;
-			switch (Value.OutputIndex)
-			{
-				case 0:	   // RGB
-					ComponentMask = FLinearColor{1.f, 1.f, 1.f, 0.f};
-					break;
-				case 1:	   // R
-					ComponentMask = FLinearColor{1.f, 0.f, 0.f, 0.f};
-					break;
-				case 2:	   // G
-					ComponentMask = FLinearColor{0.f, 1.f, 0.f, 0.f};
-					break;
-				case 3:	   // B
-					ComponentMask = FLinearColor{0.f, 0.f, 1.f, 0.f};
-					break;
-				case 4:	   // A
-					ComponentMask = FLinearColor{0.f, 0.f, 0.f, 1.f};
-					break;
-			}
-			MaterialNode.AddVectorParameterValue(*FString::Printf(TEXT("%sTextureComponent"), **ParameterName), ComponentMask);
-		}
-
-		void operator()(const UsdToUnreal::FPrimvarReaderParameterValue& Value) const
-		{
-			MaterialNode.AddVectorParameterValue(*ParameterName, FLinearColor{Value.FallbackValue});
-
-			if (Value.PrimvarName == TEXT("displayColor"))
-			{
-				MaterialNode.AddScalarParameterValue(TEXT("UseVertexColorForBaseColor"), 1.0f);
-			}
-		}
-
-		void operator()(const bool Value) const
-		{
-			MaterialNode.AddScalarParameterValue(*ParameterName, static_cast<float>(Value));
-		}
-
-	public:
-		const UE::FUsdPrim& Prim;
-		UInterchangeBaseNodeContainer& NodeContainer;
-		UInterchangeMaterialInstanceNode& MaterialNode;
-		const TMap<FString, int32>& PrimvarToUVIndex;
-		const FString* ParameterName = nullptr;
-	};
-
 	void UInterchangeUSDTranslatorImpl::SetMaterialSlotDependencies(
 		const FString& MaterialPrimName,
 		const FString& MaterialUid,
@@ -856,6 +733,160 @@ namespace UE::InterchangeUsdTranslator::Private
 		return false;
 	}
 
+	// We use this visitor to set UsdToUnreal::FParameterValue TVariant values onto UInterchangeShaderGraphNode.
+	//
+	// For now we only set attributes meant to be parsed as material instance parameters.
+	// If/whenever we want to support generating full material shader graphs from USD, we likely don't want to just fill out inputs
+	// into a rigid material function structures based on the shading model like the GLTF translator does, as USD materials can have
+	// custom shader graphs themselves. We'd either need to truly generate arbitrary interchange shader graphs here to be useful, or
+	// to delegate this work to MaterialX somehow (c.f. MaterialX materials baked into USD shader graphs)
+	struct FMaterialInstanceParameterValueVisitor
+	{
+		FMaterialInstanceParameterValueVisitor(
+			const UE::FUsdPrim& InPrim,
+			UInterchangeBaseNodeContainer& InNodeContainer,
+			UInterchangeShaderGraphNode& InMaterialNode,
+			const TMap<FString, int32>& InPrimvarToUVIndex
+		)
+			: Prim(InPrim)
+			, NodeContainer(InNodeContainer)
+			, MaterialNode(InMaterialNode)
+			, PrimvarToUVIndex(InPrimvarToUVIndex)
+		{
+		}
+
+		static FString MakeInstanceParameterName(const FString& InParameterName)
+		{
+			return UE::Interchange::USD::AttributeIdentifierPrefix + InParameterName;
+		}
+
+		void EnableTextureForChannel(bool bEnable) const
+		{
+			MaterialNode.AddFloatAttribute(	   //
+				MakeInstanceParameterName(FString::Printf(TEXT("Use%sTexture"), **BaseParameterName)),
+				bEnable ? 1.0f : 0.0f
+			);
+		}
+
+		void operator()(const float Value) const
+		{
+			MaterialNode.AddFloatAttribute(MakeInstanceParameterName(*BaseParameterName), Value);
+			EnableTextureForChannel(false);
+		}
+
+		void operator()(const FVector& Value) const
+		{
+			MaterialNode.AddLinearColorAttribute(MakeInstanceParameterName(*BaseParameterName), FLinearColor{Value});
+			EnableTextureForChannel(false);
+		}
+
+		void operator()(const UsdToUnreal::FTextureParameterValue& Value) const
+		{
+			// Emit texture node itself (this is the main place where this happens)
+			// Note that the node name isn't just the texture path, as we may have multiple material users of this texture
+			// with different settings, and so we need separate translated nodes for each material and parameter
+			const FString TextureUid = FString::Printf(TEXT("Texture:%s:%s"), *Prim.GetPrimPath().GetString(), **BaseParameterName);
+			AddTextureNode(Prim, TextureUid, Value, NodeContainer);
+
+			// Actual texture assignment
+			MaterialNode.AddStringAttribute(	//
+				MakeInstanceParameterName(FString::Printf(TEXT("%sTexture"), **BaseParameterName)),
+				TextureUid
+			);
+			EnableTextureForChannel(true);
+
+			// UV transform
+			FLinearColor ScaleAndTranslation = FLinearColor{
+				Value.UVScale.GetVector()[0],
+				Value.UVScale.GetVector()[1],
+				Value.UVTranslation[0],
+				Value.UVTranslation[1]
+			};
+			MaterialNode.AddLinearColorAttribute(
+				MakeInstanceParameterName(FString::Printf(TEXT("%sScaleTranslation"), **BaseParameterName)),
+				ScaleAndTranslation
+			);
+			MaterialNode.AddFloatAttribute(	   //
+				MakeInstanceParameterName(FString::Printf(TEXT("%sRotation"), **BaseParameterName)),
+				Value.UVRotation
+			);
+
+			// UV index
+			if (const int32* FoundIndex = PrimvarToUVIndex.Find(Value.Primvar))
+			{
+				MaterialNode.AddFloatAttribute(	   //
+					MakeInstanceParameterName(FString::Printf(TEXT("%sUVIndex"), **BaseParameterName)),
+					*FoundIndex
+				);
+			}
+			else
+			{
+				UE_LOG(
+					LogUsd,
+					Warning,
+					TEXT("Failed to find primvar '%s' when setting material parameter '%s' on material '%s'. Available primvars and UV "
+						 "indices: %s.%s"),
+					*Value.Primvar,
+					**BaseParameterName,
+					*Prim.GetPrimPath().GetString(),
+					*UsdUtils::StringifyMap(PrimvarToUVIndex),
+					Value.Primvar.IsEmpty()
+						? TEXT(
+							  " Is your UsdUVTexture Shader missing the 'inputs:st' attribute? (It specifies which UV set to sample the texture with)"
+						  )
+						: TEXT("")
+				);
+			}
+
+			// Component mask (which channel of the texture to use)
+			FLinearColor ComponentMask = FLinearColor::Black;
+			switch (Value.OutputIndex)
+			{
+				case 0:	   // RGB
+					ComponentMask = FLinearColor{1.f, 1.f, 1.f, 0.f};
+					break;
+				case 1:	   // R
+					ComponentMask = FLinearColor{1.f, 0.f, 0.f, 0.f};
+					break;
+				case 2:	   // G
+					ComponentMask = FLinearColor{0.f, 1.f, 0.f, 0.f};
+					break;
+				case 3:	   // B
+					ComponentMask = FLinearColor{0.f, 0.f, 1.f, 0.f};
+					break;
+				case 4:	   // A
+					ComponentMask = FLinearColor{0.f, 0.f, 0.f, 1.f};
+					break;
+			}
+			MaterialNode.AddLinearColorAttribute(
+				MakeInstanceParameterName(FString::Printf(TEXT("%sTextureComponent"), **BaseParameterName)),
+				ComponentMask
+			);
+		}
+
+		void operator()(const UsdToUnreal::FPrimvarReaderParameterValue& Value) const
+		{
+			MaterialNode.AddLinearColorAttribute(MakeInstanceParameterName(*BaseParameterName), FLinearColor{Value.FallbackValue});
+
+			if (Value.PrimvarName == TEXT("displayColor"))
+			{
+				MaterialNode.AddFloatAttribute(MakeInstanceParameterName(TEXT("UseVertexColorForBaseColor")), 1.0f);
+			}
+		}
+
+		void operator()(const bool Value) const
+		{
+			MaterialNode.AddBooleanAttribute(MakeInstanceParameterName(*BaseParameterName), static_cast<float>(Value));
+		}
+
+	public:
+		const UE::FUsdPrim& Prim;
+		UInterchangeBaseNodeContainer& NodeContainer;
+		UInterchangeShaderGraphNode& MaterialNode;
+		const TMap<FString, int32>& PrimvarToUVIndex;
+		const FString* BaseParameterName = nullptr;
+	};
+
 	void UInterchangeUSDTranslatorImpl::AddMaterialNode(
 		const UE::FUsdPrim& Prim,
 		UInterchangeUsdTranslatorSettings* TranslatorSettings,
@@ -867,7 +898,7 @@ namespace UE::InterchangeUsdTranslator::Private
 		FString MaterialPrimName(Prim.GetName().ToString());
 
 		// Check if Node already exist with this ID
-		if (const UInterchangeMaterialInstanceNode* Node = Cast<UInterchangeMaterialInstanceNode>(NodeContainer.GetNode(MaterialUid)))
+		if (const UInterchangeShaderGraphNode* Node = Cast<UInterchangeShaderGraphNode>(NodeContainer.GetNode(MaterialUid)))
 		{
 			return;
 		}
@@ -894,9 +925,10 @@ namespace UE::InterchangeUsdTranslator::Private
 			RenderContext = UnrealIdentifiers::UniversalRenderContext;
 		}
 
-		UInterchangeMaterialInstanceNode* MaterialNode = NewObject<UInterchangeMaterialInstanceNode>(&NodeContainer);
+		UInterchangeShaderGraphNode* MaterialNode = NewObject<UInterchangeShaderGraphNode>(&NodeContainer);
 		MaterialNode->InitializeNode(MaterialUid, MaterialPrimName, EInterchangeNodeContainerType::TranslatedAsset);
 		MaterialNode->SetAssetName(MaterialPrimName);
+		MaterialNode->AddBooleanAttribute(UE::Interchange::USD::ParseMaterialIdentifier, true);
 		NodeContainer.AddNode(MaterialNode);
 
 		// Set the material instance node to the correct mesh nodes
@@ -907,46 +939,42 @@ namespace UE::InterchangeUsdTranslator::Private
 		const bool bSuccess = UsdToUnreal::ConvertMaterial(Prim, MaterialData, *RenderContext.ToString());
 
 		// Set all the parameter values to the interchange node
-		bool bHasUDIMTexture = false;
-		FParameterValueVisitor Visitor{Prim, NodeContainer, *MaterialNode, MaterialData.PrimvarToUVIndex};
+		bool bNeedsVTParent = false;
+		FMaterialInstanceParameterValueVisitor Visitor{Prim, NodeContainer, *MaterialNode, MaterialData.PrimvarToUVIndex};
 		for (TPair<FString, UsdToUnreal::FParameterValue>& Pair : MaterialData.Parameters)
 		{
-			Visitor.ParameterName = &Pair.Key;
+			Visitor.BaseParameterName = &Pair.Key;
 			Visit(Visitor, Pair.Value);
 
 			// Also simultaneously check if any of these parameters wants to be an UDIM texture so that we can use the VT reference material later
-			if (!bHasUDIMTexture)
+			if (!bNeedsVTParent)
 			{
 				if (UsdToUnreal::FTextureParameterValue* TextureParameter = Pair.Value.TryGet<UsdToUnreal::FTextureParameterValue>())
 				{
 					if (TextureParameter->bIsUDIM)
 					{
-						bHasUDIMTexture = true;
+						bNeedsVTParent = true;
 					}
 				}
 			}
 		}
 
-		// Find and set the right reference material
+		// Find and set the relevant info so that we can pick the reference material on the pipeline
 		//
 		// TODO: Proper VT texture support (we'd need to know the texture resolution at this point, and we haven't parsed them yet...).
 		// The way it currently works on Interchange is that the factory will create a VT or nonVT version of the texture to match the
 		// material parameter slot. Since we'll currently never set the VT reference material, it essentially means it will always
 		// downgrade our VT textures to non-VT.
-		// The only exception is how we upgrade the reference material to VT in case we have any UDIM textures, as those are trivial to
-		// check for (we don't have to actually load the textures to do it)
-		EUsdReferenceMaterialProperties Properties = EUsdReferenceMaterialProperties::None;
+		// The only exception is how we upgrade the reference material to VT in case we have any UDIM textures a few lines above,
+		// as those are trivial to check for (we don't have to actually load the textures to do it)
 		if (UsdUtils::IsMaterialTranslucent(MaterialData))
 		{
-			Properties |= EUsdReferenceMaterialProperties::Translucent;
+			MaterialNode->SetCustomBlendMode(EBlendMode::BLEND_Translucent);
 		}
-		if (bHasUDIMTexture)
+		if (bNeedsVTParent)
 		{
-			Properties |= EUsdReferenceMaterialProperties::VT;
+			MaterialNode->AddBooleanAttribute(UE::Interchange::USD::NeedsVTParentIdentifier, true);
 		}
-
-		FSoftObjectPath ReferenceMaterial = UsdUnreal::MaterialUtils::GetReferencePreviewSurfaceMaterial(Properties);
-		MaterialNode->SetCustomParent(ReferenceMaterial.ToString());
 	}
 
 	void AddDisplayColorMaterialInstanceNodeIfNeeded(UInterchangeBaseNodeContainer& NodeContainer, const FString& DisplayColorDesc)
@@ -956,7 +984,7 @@ namespace UE::InterchangeUsdTranslator::Private
 		FString NodeUid = MaterialPrefix + DisplayColorDesc;
 
 		// We'll treat the DisplayColorDesc (something like "!DisplayColor_1_0") as the material instance UID here
-		const UInterchangeMaterialInstanceNode* Node = Cast<UInterchangeMaterialInstanceNode>(NodeContainer.GetNode(DisplayColorDesc));
+		const UInterchangeShaderGraphNode* Node = Cast<UInterchangeShaderGraphNode>(NodeContainer.GetNode(DisplayColorDesc));
 		if (Node)
 		{
 			return;
@@ -970,18 +998,12 @@ namespace UE::InterchangeUsdTranslator::Private
 		}
 		FString NodeName = ParsedMat->ToPrettyString();
 
-		const FSoftObjectPath* ReferenceMaterialPath = GetReferenceMaterialPath(ParsedMat.GetValue());
-		if (!ReferenceMaterialPath)
-		{
-			return;
-		}
-
-		// Not needed
-		const FString ParentNodeUid;
-		UInterchangeMaterialInstanceNode* NewNode = UInterchangeMaterialInstanceNode::Create(&NodeContainer, DisplayColorDesc, ParentNodeUid);
+		UInterchangeShaderGraphNode* NewNode = NewObject<UInterchangeShaderGraphNode>(&NodeContainer);
 		NewNode->InitializeNode(NodeUid, NodeName, EInterchangeNodeContainerType::TranslatedAsset);
+		NewNode->AddBooleanAttribute(UE::Interchange::USD::ParseMaterialIdentifier, true);
+		NewNode->AddStringAttribute(*UE::Interchange::USD::DisplayColorDescIdentifier, DisplayColorDesc);
 
-		NewNode->SetCustomParent(ReferenceMaterialPath->GetAssetPathString());
+		NodeContainer.AddNode(NewNode);
 	}
 
 	void AddLightNode(const UE::FUsdPrim& Prim, UInterchangeBaseNodeContainer& NodeContainer)
@@ -1808,14 +1830,6 @@ namespace UE::InterchangeUsdTranslator::Private
 		FTraversalInfo Info
 	)
 	{
-		// Ignore prim subtrees from disabled purposes
-		// TODO: Move this to the pipeline and filter only the factory nodes
-		EUsdPurpose PrimPurpose = IUsdPrim::GetPurpose(Prim);
-		if (!EnumHasAllFlags(TranslatorImpl.CachedMeshConversionOptions.PurposesToLoad, PrimPurpose))
-		{
-			return;
-		}
-
 		FString SceneNodeUid = Prim.GetPrimPath().GetString();
 		FString DisplayLabel(Prim.GetName().ToString());
 
@@ -1846,18 +1860,32 @@ namespace UE::InterchangeUsdTranslator::Private
 		}
 
 		// Only prims that require rendering (and have a renderable parent) get a scene node.
-		// This includes Xforms but also Scopes, which are not Xformable
+		// This includes Xforms but also Scopes, which are not Xformable.
+		//
+		// We add a scene node for the pseudoroot in order to make tree traversal easier on the pipeline,
+		// but the pipeline will strip the pseudoroot node as it's final step
+		const bool bIsImageable = Prim.IsA(TEXT("Imageable"));
+		const bool bNeedsSceneNode = Prim.IsPseudoRoot() || (bIsImageable && (Info.ParentNode || Prim.GetParent().IsPseudoRoot()));
+
 		UInterchangeSceneNode* SceneNode = nullptr;
-		if (Prim.IsA(TEXT("Imageable")) && (Info.ParentNode || Prim.GetParent().IsPseudoRoot()))
+		if (bNeedsSceneNode)
 		{
 			SceneNode = NewObject<UInterchangeSceneNode>(&NodeContainer);
 			SceneNode->InitializeNode(SceneNodeUid, DisplayLabel, EInterchangeNodeContainerType::TranslatedScene);
 			NodeContainer.AddNode(SceneNode);
 
+			// Store our purpose here to be filtered on the pipeline
+			const EUsdPurpose PrimPurpose = IUsdPrim::GetPurpose(Prim);
+			if (PrimPurpose != EUsdPurpose::Default)
+			{
+				SceneNode->AddInt32Attribute(UE::Interchange::USD::GeometryPurposeIdentifier, (int32)PrimPurpose);
+			}
+
 			// If we're an Xformable, get our transform
 			FTransform Transform = FTransform::Identity;
 			bool bResetTransformStack = false;
-			if (UsdToUnreal::ConvertXformable(
+			if (bIsImageable
+				&& UsdToUnreal::ConvertXformable(
 					Prim.GetStage(),
 					UE::FUsdTyped(Prim),
 					Transform,
@@ -1888,22 +1916,27 @@ namespace UE::InterchangeUsdTranslator::Private
 			}
 
 			// Add animation tracks
-			AddPropertyAnimationNodes(Prim, TranslatorImpl, NodeContainer);
-			if (UsdUtils::HasAnimatedTransform(Prim))
+			if (bIsImageable)
 			{
-				AddTransformAnimationNode(Prim, TranslatorImpl, NodeContainer);
+				AddPropertyAnimationNodes(Prim, TranslatorImpl, NodeContainer);
+				if (UsdUtils::HasAnimatedTransform(Prim))
+				{
+					AddTransformAnimationNode(Prim, TranslatorImpl, NodeContainer);
+				}
 			}
 		}
 
-		// Note: This has the effect of effectively shutting down the generation of scene nodes
-		// below any prim that is not a least an Imageable, as we check for a valid parent before
-		// generating one
-		Info.ParentNode = SceneNode;
-
 		// Recurse into child prims
-		for (const FUsdPrim& ChildPrim : Prim.GetChildren())
 		{
-			Traverse(ChildPrim, TranslatorImpl, NodeContainer, TranslatorSettings, Info);
+			// Note: This has the effect of effectively shutting down the generation of scene nodes
+			// below any prim that is not a least an Imageable, as we check for a valid parent before
+			// generating one
+			Info.ParentNode = SceneNode;
+
+			for (const FUsdPrim& ChildPrim : Prim.GetChildren())
+			{
+				Traverse(ChildPrim, TranslatorImpl, NodeContainer, TranslatorSettings, Info);
+			}
 		}
 	}
 
@@ -2512,7 +2545,9 @@ bool UInterchangeUSDTranslator::Translate(UInterchangeBaseNodeContainer& NodeCon
 
 	// Cache these so we don't have to keep converting these tokens over and over during translation
 	FUsdMeshConversionOptions& MeshOptions = ImplPtr->CachedMeshConversionOptions;
-	MeshOptions.PurposesToLoad = (EUsdPurpose)Settings->GeometryPurpose;
+
+	// We filter for this on the pipeline now
+	MeshOptions.PurposesToLoad = EUsdPurpose::Default | EUsdPurpose::Proxy | EUsdPurpose::Render | EUsdPurpose::Guide;
 
 	// TODO: Change FUsdMeshConversionOptions to not hold USD types directly, so we don't have to the conversion below everywhere.
 	// We can't use UsdToUnreal::ConvertToken() here because it returns a TUsdStore, and the template instantiation created in this module doesn't
@@ -2527,10 +2562,7 @@ bool UInterchangeUSDTranslator::Translate(UInterchangeBaseNodeContainer& NodeCon
 
 	// Traverse stage and emit translated nodes
 	FTraversalInfo Info;
-	for (const FUsdPrim& Prim : ImplPtr->UsdStage.GetPseudoRoot().GetChildren())
-	{
-		Traverse(Prim, *ImplPtr, NodeContainer, TranslatorSettings, Info);
-	}
+	Traverse(ImplPtr->UsdStage.GetPseudoRoot(), *ImplPtr, NodeContainer, TranslatorSettings, Info);
 
 	return true;
 #else
diff --git a/Engine/Plugins/Interchange/Runtime/Source/Import/Public/Usd/InterchangeUsdDefinitions.h b/Engine/Plugins/Interchange/Runtime/Source/Import/Public/Usd/InterchangeUsdDefinitions.h
new file mode 100644
--- /dev/null
+++ b/Engine/Plugins/Interchange/Runtime/Source/Import/Public/Usd/InterchangeUsdDefinitions.h
@@ -0,0 +1,32 @@
+// Copyright Epic Games, Inc. All Rights Reserved.
+
+#pragma once
+
+#include "CoreMinimal.h"
+
+namespace UE::Interchange::USD
+{
+	// Name of a custom attribute added to translated nodes to contain their geometry purpose (proxy, render, guide, etc.)
+	const FString GeometryPurposeIdentifier = TEXT("USD_Geometry_Purpose");
+
+	// Prefix for custom attributes that describe material instance parameters
+	const FString AttributeIdentifierPrefix = TEXT("USD_MI_Attribute_");
+
+	// Flag indicating whether we should parse a UInterchangeShaderGraphNode on the InterchangeUSDPipeline
+	const FString ParseMaterialIdentifier = TEXT("USD_MI_ParseMaterial");
+
+	// Means we should pick this material to use as parent material, regardless of what else is provided
+	// (NeedsVTParent, DisplayColor, etc.)
+	const FString ParentMaterialIdentifier = TEXT("USD_MI_Parent");
+
+	// Means that we should use a display color material parent for this material, according to the
+	// value provided in this attribute
+	const FString DisplayColorDescIdentifier = TEXT("USD_MI_DisplayColorDescription");
+
+	// Means that if we were to pick one of the UsdPreviewSurface parent materials for this material
+	// translated node, we should pick the "VT" variant of that material.
+	//
+	// Note: To indicate TwoSided/Translucent we use the CustomTwoSided/BlendMode standard UInterchangeShaderGraphNode
+	// attributes.
+	const FString NeedsVTParentIdentifier = TEXT("USD_MI_NeedsVTParent");
+}
diff --git a/Engine/Plugins/Interchange/Runtime/Source/Import/Public/Usd/InterchangeUsdTranslator.h b/Engine/Plugins/Interchange/Runtime/Source/Import/Public/Usd/InterchangeUsdTranslator.h
--- a/Engine/Plugins/Interchange/Runtime/Source/Import/Public/Usd/InterchangeUsdTranslator.h
+++ b/Engine/Plugins/Interchange/Runtime/Source/Import/Public/Usd/InterchangeUsdTranslator.h
@@ -27,20 +27,8 @@ class UInterchangeUsdTranslatorSettings : public UInterchangeTranslatorSettings
 
 public:
 
-	/**
-	 * TODO: Most (if not all) of these settings should probably be specific to the pipeline, and not the
-	 * translator, and will be moved there whenever we do make a USD pipeline.
-	 *
-	 * For example: Instead of filtering prims based on purpose on translation, we should emit all of the translated
-	 * nodes and filter only later on the pipeline, as translating the scene should be fast either way. That way users
-	 * can even customize/disable that behavior if they want to, in order to use their own pipelines.
-	 *
-	 * It's not clear what to do about StageOptions or RenderContext though: Maybe these should be here, as they
-	 * actively affect how we translate the scene? (e.g. the generated Interchange Material node for a Material
-	 * prim will be very different whether we use "universal", "unreal" or "mtlx" render contexts...)
-	 */
-
 	/** Only import geometry prims with these specific purposes from the USD file */
+	UE_DEPRECATED(5.6, "This option is no longer a translator setting and has moved to the InterchangeUSDPipeline")
 	UPROPERTY(EditAnywhere, Category = "USD Translator", meta = (Bitmask, BitmaskEnum = "/Script/UnrealUSDWrapper.EUsdPurpose"))
 	int32 GeometryPurpose;
 
diff --git a/Engine/Plugins/Interchange/Runtime/Source/Pipelines/InterchangePipelines.Build.cs b/Engine/Plugins/Interchange/Runtime/Source/Pipelines/InterchangePipelines.Build.cs
--- a/Engine/Plugins/Interchange/Runtime/Source/Pipelines/InterchangePipelines.Build.cs
+++ b/Engine/Plugins/Interchange/Runtime/Source/Pipelines/InterchangePipelines.Build.cs
@@ -25,13 +25,15 @@ namespace UnrealBuildTool.Rules
 				{
 					"ApplicationCore",
 					"CinematicCamera",
-					"InterchangeImport",
 					"DeveloperSettings",
 					"GLTFCore",
 					"InputCore",
+					"InterchangeImport",
+					"LevelSequence",
 					"Slate",
 					"SlateCore",
-					"LevelSequence",
+					"UnrealUSDWrapper",
+					"USDClasses",
 				}
 			);
 
diff --git a/Engine/Plugins/Interchange/Runtime/Source/Pipelines/Private/InterchangeUSDPipeline.cpp b/Engine/Plugins/Interchange/Runtime/Source/Pipelines/Private/InterchangeUSDPipeline.cpp
new file mode 100644
--- /dev/null
+++ b/Engine/Plugins/Interchange/Runtime/Source/Pipelines/Private/InterchangeUSDPipeline.cpp
@@ -0,0 +1,466 @@
+// Copyright Epic Games, Inc. All Rights Reserved.
+
+#include "InterchangeUSDPipeline.h"
+
+#include "InterchangeManager.h"
+#include "InterchangeMaterialFactoryNode.h"
+#include "InterchangeMaterialInstanceNode.h"
+#include "InterchangeSceneNode.h"
+#include "InterchangeShaderGraphNode.h"
+#include "InterchangeSkeletalMeshLodDataNode.h"
+#include "InterchangeStaticMeshLodDataNode.h"
+
+#include "UnrealUSDWrapper.h"
+#include "Usd/InterchangeUsdDefinitions.h"
+#include "USDMaterialUtils.h"
+#include "UsdWrappers/SdfPath.h"
+
+#include "Materials/MaterialInstanceConstant.h"
+#include "Materials/MaterialInstanceDynamic.h"
+
+namespace UE::InterchangeUsdPipeline::Private
+{
+	const FString& GetPseudoRootTranslatedNodeUid()
+	{
+		const static FString Result = UE::FSdfPath::AbsoluteRootPath().GetString();
+		return Result;
+	}
+
+	const FString& GetPseudoRootFactoryNodeUid()
+	{
+		const static FString Result = UInterchangeFactoryBaseNode::BuildFactoryNodeUid(GetPseudoRootTranslatedNodeUid());
+		return Result;
+	}
+
+	// Disables the factory nodes on NodeContainer if they target translated nodes that had disallowed purposes
+	void DisableNodesBasedOnGeometryPurpose(UInterchangeBaseNodeContainer& NodeContainer, const EUsdPurpose AllowedPurposes)
+	{
+		TSet<FString> DisabledTranslatedNodes;
+
+		TFunction<void(const FString&, bool)> RecursiveCollectFilteredTranslatedNodes = nullptr;
+		RecursiveCollectFilteredTranslatedNodes = [&NodeContainer,
+												   AllowedPurposes,
+												   &RecursiveCollectFilteredTranslatedNodes,
+												   &DisabledTranslatedNodes](const FString& NodeUid, bool bDisableSubtree)
+		{
+			const UInterchangeBaseNode* Node = NodeContainer.GetNode(NodeUid);
+			if (!Node)
+			{
+				return;
+			}
+
+			if (!bDisableSubtree)
+			{
+				EUsdPurpose AuthoredPurpose = EUsdPurpose::Default;
+				bool bHasPurpose = Node->GetInt32Attribute(UE::Interchange::USD::GeometryPurposeIdentifier, (int32&)AuthoredPurpose);
+				if (bHasPurpose && !EnumHasAllFlags(AllowedPurposes, AuthoredPurpose))
+				{
+					// Purpose inheritance according to UsdGeomImageable::ComputePurposeInfo seems to be:
+					// - Authored purposes inherit down to prims without any purpose;
+					// - If a prim has any purpose authored on them, that is their "computed purpose", and that is inherited to its children.
+					//
+					// We don't care much about the actual purpose then: We just need to check if any prim/node has a purpose
+					// that is not allowed. If that is the case, we turn off the entire subtree
+					bDisableSubtree = true;
+				}
+			}
+
+			if (bDisableSubtree)
+			{
+				// Mark the scene translated node as disabled
+				DisabledTranslatedNodes.Add(NodeUid);
+
+				// Mark the asset translated node as disabled
+				if (const UInterchangeSceneNode* SceneNode = Cast<UInterchangeSceneNode>(Node))
+				{
+					FString AssetNodeUid;
+					if (SceneNode->GetCustomAssetInstanceUid(AssetNodeUid))
+					{
+						DisabledTranslatedNodes.Add(AssetNodeUid);
+					}
+				}
+			}
+
+			for (const FString& ChildNodeUid : NodeContainer.GetNodeChildrenUids(NodeUid))
+			{
+				RecursiveCollectFilteredTranslatedNodes(ChildNodeUid, bDisableSubtree);
+			}
+		};
+		const bool bDisableSubtree = false;
+		RecursiveCollectFilteredTranslatedNodes(GetPseudoRootTranslatedNodeUid(), bDisableSubtree);
+
+		// Disable any factory node that targets/references one of our disabled scene nodes.
+		//
+		// We do this in a separate pass because we can say absolutely nothing about our factory nodes, as any number
+		// of pipelines may have done arbirary transformations before we got to run. We have to hope they kept the
+		// target node attribute updated at least, and use that.
+		NodeContainer.IterateNodesOfType<UInterchangeFactoryBaseNode>(
+			[&DisabledTranslatedNodes](const FString& NodeUid, UInterchangeFactoryBaseNode* FactoryNode)
+			{
+				// Handle standard factory nodes
+				{
+					TArray<FString> TargetTranslatedNodeUids;
+					FactoryNode->GetTargetNodeUids(TargetTranslatedNodeUids);
+
+					// If we have a single target node that is disabled, just disable the factory node itself.
+					// This is probably the general case, and should be more explicit on preview panels and such
+					if (TargetTranslatedNodeUids.Num() == 1)
+					{
+						if (DisabledTranslatedNodes.Contains(TargetTranslatedNodeUids[0]))
+						{
+							const bool bEnabled = false;
+							FactoryNode->SetEnabled(bEnabled);
+						}
+					}
+					// If we have multiple target nodes, instead of disabling the whole factory node just remove
+					// this one particular target (e.g. removing a single mesh from a collapsed mesh blob, in case that
+					// mesh doesn't match the purpose)
+					else
+					{
+						for (const FString& TargetNode : TargetTranslatedNodeUids)
+						{
+							if (DisabledTranslatedNodes.Contains(TargetNode))
+							{
+								FactoryNode->RemoveTargetNodeUid(TargetNode);
+							}
+						}
+					}
+				}
+
+				// These nodes don't use the common "target" mechanism and store their mesh nodes separately...
+				// TODO: Will have to expand these to also handle geometry caches?
+				if (UInterchangeStaticMeshLodDataNode* LodDataNode = Cast<UInterchangeStaticMeshLodDataNode>(FactoryNode))
+				{
+					TArray<FString> TranslatedMeshNodeUids;
+					LodDataNode->GetMeshUids(TranslatedMeshNodeUids);
+
+					if (TranslatedMeshNodeUids.Num() == 1)
+					{
+						if (DisabledTranslatedNodes.Contains(TranslatedMeshNodeUids[0]))
+						{
+							const bool bEnabled = false;
+							LodDataNode->SetEnabled(bEnabled);
+						}
+					}
+					else
+					{
+						for (const FString& MeshNodeUid : TranslatedMeshNodeUids)
+						{
+							if (DisabledTranslatedNodes.Contains(MeshNodeUid))
+							{
+								LodDataNode->RemoveMeshUid(MeshNodeUid);
+							}
+						}
+					}
+				}
+				else if (UInterchangeSkeletalMeshLodDataNode* SkeletalNode = Cast<UInterchangeSkeletalMeshLodDataNode>(FactoryNode))
+				{
+					TArray<FString> TranslatedMeshNodeUids;
+					SkeletalNode->GetMeshUids(TranslatedMeshNodeUids);
+
+					if (TranslatedMeshNodeUids.Num() == 1)
+					{
+						if (DisabledTranslatedNodes.Contains(TranslatedMeshNodeUids[0]))
+						{
+							const bool bEnabled = false;
+							SkeletalNode->SetEnabled(bEnabled);
+						}
+					}
+					else
+					{
+						for (const FString& MeshNodeUid : TranslatedMeshNodeUids)
+						{
+							if (DisabledTranslatedNodes.Contains(MeshNodeUid))
+							{
+								SkeletalNode->RemoveMeshUid(MeshNodeUid);
+							}
+						}
+					}
+				}
+			}
+		);
+	}
+
+	// Resolves which should be the parent material for NewFactoryNode and sets it
+	void SetCustomParentMaterial(const UInterchangeShaderGraphNode& ShaderGraphNode, UInterchangeMaterialInstanceFactoryNode& NewFactoryNode)
+	{
+		using namespace UsdUnreal::MaterialUtils;
+
+		// Priority 1) Explicitly set parent material
+		//
+		// This can be from an UnrealMaterial or from anything the user may have manually set
+		// TODO: Maybe use this to handle UnrealMaterials?
+		if (FString ParentMaterialPath;
+			ShaderGraphNode.GetStringAttribute(UE::Interchange::USD::ParentMaterialIdentifier, ParentMaterialPath) && !ParentMaterialPath.IsEmpty())
+		{
+			NewFactoryNode.SetCustomParent(ParentMaterialPath);
+			return;
+		}
+
+		// Priority 2) DisplayColor material
+		//
+		// This is usually a lower priority, but we explicitly flag these so it's easy to check for them
+		if (FString DisplayColorDesc;
+			ShaderGraphNode.GetStringAttribute(UE::Interchange::USD::DisplayColorDescIdentifier, DisplayColorDesc) && !DisplayColorDesc.IsEmpty())
+		{
+			TOptional<FDisplayColorMaterial> ParsedMat = FDisplayColorMaterial::FromString(DisplayColorDesc);
+			if (ParsedMat.IsSet())
+			{
+				if (const FSoftObjectPath* ParentMaterialPath = GetReferenceMaterialPath(ParsedMat.GetValue()))
+				{
+					NewFactoryNode.SetCustomParent(ParentMaterialPath->GetAssetPathString());
+					return;
+				}
+			}
+		}
+
+		// Priority 3) UsdPreviewSurface material
+		{
+			EUsdReferenceMaterialProperties Properties = EUsdReferenceMaterialProperties::None;
+
+			// Translucent
+			int BlendMode = EBlendMode::BLEND_Opaque;
+			if (ShaderGraphNode.GetCustomBlendMode(BlendMode) && BlendMode == EBlendMode::BLEND_Translucent)
+			{
+				Properties |= EUsdReferenceMaterialProperties::Translucent;
+			}
+
+			// TwoSided
+			// TODO: Actually set this somewhere
+			bool bTwoSided = false;
+			if (ShaderGraphNode.GetCustomTwoSided(bTwoSided) && bTwoSided)
+			{
+				Properties |= EUsdReferenceMaterialProperties::TwoSided;
+			}
+
+			// VT
+			bool bNeedsVT = false;
+			if (ShaderGraphNode.GetBooleanAttribute(UE::Interchange::USD::NeedsVTParentIdentifier, bNeedsVT) && bNeedsVT)
+			{
+				Properties |= EUsdReferenceMaterialProperties::VT;
+			}
+
+			FSoftObjectPath ParentMaterial = GetReferencePreviewSurfaceMaterial(Properties);
+			if (ParentMaterial.IsValid())
+			{
+				NewFactoryNode.SetCustomParent(ParentMaterial.GetAssetPathString());
+				return;
+			}
+		}
+	}
+
+	// Creates a new material instance factory node and fills it in with the parameter/parent material path
+	// stored as base translated node attributes on the provided ShaderGraphNode
+	UInterchangeMaterialInstanceFactoryNode* BuildMaterialInstance(
+		UInterchangeBaseNodeContainer& NodeContainer,
+		const UInterchangeShaderGraphNode& ShaderGraphNode
+	)
+	{
+		FString MaterialFactoryNodeUid = UInterchangeFactoryBaseNode::BuildFactoryNodeUid(ShaderGraphNode.GetUniqueID());
+		FString MaterialFactoryNodeName = ShaderGraphNode.GetDisplayLabel();
+		UInterchangeMaterialInstanceFactoryNode* NewFactoryNode = NewObject<UInterchangeMaterialInstanceFactoryNode>(&NodeContainer);
+		NewFactoryNode->InitializeNode(MaterialFactoryNodeUid, MaterialFactoryNodeName, EInterchangeNodeContainerType::FactoryData);
+
+		SetCustomParentMaterial(ShaderGraphNode, *NewFactoryNode);
+
+		// Pick material instance class
+		{
+			UInterchangeEditorUtilitiesBase* EditorUtilities = UInterchangeManager::GetInterchangeManager().GetEditorUtilities();
+			const UClass* MaterialClass = (EditorUtilities && EditorUtilities->IsRuntimeOrPIE()) ? UMaterialInstanceDynamic::StaticClass()
+																								 : UMaterialInstanceConstant::StaticClass();
+			NewFactoryNode->SetCustomInstanceClassName(MaterialClass->GetPathName());
+		}
+
+		// Handle material attributes
+		TArray<UE::Interchange::FAttributeKey> AllAttributeKeys;
+		ShaderGraphNode.GetAttributeKeys(AllAttributeKeys);
+		for (const UE::Interchange::FAttributeKey& AttributeKey : AllAttributeKeys)
+		{
+			const FString AttributeKeyString = AttributeKey.ToString();
+
+			if (!AttributeKeyString.Contains(UE::Interchange::USD::AttributeIdentifierPrefix))
+			{
+				// It's not one of our attributes
+				continue;
+			}
+
+			const FString CleanAttributeKey = AttributeKeyString.Replace(	 //
+				*UE::Interchange::USD::AttributeIdentifierPrefix,
+				TEXT(""),
+				ESearchCase::CaseSensitive
+			);
+			const FString StandardParameterName = UInterchangeShaderPortsAPI::MakeInputValueKey(CleanAttributeKey);
+
+			const UE::Interchange::EAttributeTypes AttributeType = ShaderGraphNode.GetAttributeType(AttributeKey);
+			switch (AttributeType)
+			{
+				case UE::Interchange::EAttributeTypes::Bool:
+				{
+					bool Value;
+					if (ShaderGraphNode.GetBooleanAttribute(AttributeKey.Key, Value))
+					{
+						NewFactoryNode->AddBooleanAttribute(StandardParameterName, Value);
+					}
+				}
+				break;
+				case UE::Interchange::EAttributeTypes::Float:
+				{
+					float Value;
+					if (ShaderGraphNode.GetFloatAttribute(AttributeKey.Key, Value))
+					{
+						NewFactoryNode->AddFloatAttribute(StandardParameterName, Value);
+					}
+				}
+				break;
+				case UE::Interchange::EAttributeTypes::LinearColor:
+				{
+					FLinearColor Value;
+					if (ShaderGraphNode.GetLinearColorAttribute(AttributeKey.Key, Value))
+					{
+						NewFactoryNode->AddLinearColorAttribute(StandardParameterName, Value);
+					}
+				}
+				break;
+				case UE::Interchange::EAttributeTypes::String:
+				{
+					FString TextureTranslatedNodeUid;
+					if (ShaderGraphNode.GetStringAttribute(AttributeKey.Key, TextureTranslatedNodeUid))
+					{
+						// We only use string attributes for textures for now
+						FString FactoryTextureUid = UInterchangeFactoryBaseNode::BuildFactoryNodeUid(TextureTranslatedNodeUid);
+
+						NewFactoryNode->AddStringAttribute(StandardParameterName, FactoryTextureUid);
+						NewFactoryNode->AddFactoryDependencyUid(FactoryTextureUid);
+					}
+				}
+				break;
+				default:
+					break;
+			}
+		}
+
+		return NewFactoryNode;
+	}
+
+	// Replaces dummy material instance factory nodes with new nodes where we properly fill in the
+	// material parameters and parent materials.
+	//
+	// We always translate the USD materials into UInterchangeShaderGraphNode (as opposed to directly
+	// UInterchangeMaterialInstanceNode) and fill in all parameters as base translated node attributes,
+	// in order to defer to the UInterchangeGenericMaterialPipeline (or other pipelines) whether to emit
+	// new material graphs or material instances.
+	//
+	// Here, we will likely run *after* UInterchangeGenericMaterialPipeline, which means is will already have
+	// converted our UInterchangeShaderGraphNodes into either material or material instance factory nodes,
+	// according to its "MaterialImport" option.
+	//
+	// If MaterialImport is set to "ImportAsMaterialInstances" however, UInterchangeGenericMaterialPipeline
+	// will not convert our base translated node attributes into the factory node like we need it to.
+	// This is what this function aims to repair: It will go over any of these broken material instance factory
+	// nodes, and replace them with new nodes that have the correct material info extracted from the
+	// corresponding ("target") UInterchangeShaderGraphNode base node attributes. It is also in charge of
+	// picking the correct parent material for these nodes, according to all the info added to the shader graph
+	// nodes.
+	//
+	// Note: This is all for the material instance code path. Unfortunately the full material graph code path
+	// won't work either, as we don't emit the shader graph translated node data from the USD translator just yet.
+	// We will likely want to support arbitrary shader graph generation for that though, which may come at a
+	// later point
+	void RepairMaterialInstanceFactoryNodes(UInterchangeBaseNodeContainer& NodeContainer)
+	{
+		TMap<FString, const UInterchangeShaderGraphNode*> MaterialFactoryNodeUidsToShaderGraphNodes;
+
+		auto CollectNodes =
+			[&NodeContainer, &MaterialFactoryNodeUidsToShaderGraphNodes](const FString& NodeUid, UInterchangeMaterialInstanceFactoryNode* FactoryNode)
+		{
+			TArray<FString> TargetNodeUids;
+			FactoryNode->GetTargetNodeUids(TargetNodeUids);
+
+			for (const FString& TargetNodeUid : TargetNodeUids)
+			{
+				const UInterchangeShaderGraphNode* ShaderGraphNode = Cast<UInterchangeShaderGraphNode>(NodeContainer.GetNode(TargetNodeUid));
+				if (!ShaderGraphNode)
+				{
+					continue;
+				}
+
+				// We explicitly flag our shader graph nodes with this attribute
+				bool bParse;
+				if (ShaderGraphNode->GetBooleanAttribute(*UE::Interchange::USD::ParseMaterialIdentifier, bParse) && bParse)
+				{
+					MaterialFactoryNodeUidsToShaderGraphNodes.Add(NodeUid, ShaderGraphNode);
+					break;
+				}
+			}
+		};
+		NodeContainer.IterateNodesOfType<UInterchangeMaterialInstanceFactoryNode>(CollectNodes);
+
+		for (const TPair<FString, const UInterchangeShaderGraphNode*>& Pair : MaterialFactoryNodeUidsToShaderGraphNodes)
+		{
+			const FString& FactoryNodeUid = Pair.Key;
+			const UInterchangeShaderGraphNode* ShaderGraphNode = Pair.Value;
+
+			UInterchangeMaterialInstanceFactoryNode* NewFactoryNode = BuildMaterialInstance(NodeContainer, *ShaderGraphNode);
+			if (NewFactoryNode)
+			{
+				NodeContainer.ReplaceNode(FactoryNodeUid, NewFactoryNode);
+			}
+		}
+	}
+
+	void RemovePseudoRootFactoryNode(UInterchangeBaseNodeContainer& NodeContainer)
+	{
+		const FString& FactoryNodeUid = GetPseudoRootFactoryNodeUid();
+		UInterchangeFactoryBaseNode* PseudoRootFactoryNode = NodeContainer.GetFactoryNode(FactoryNodeUid);
+		if (PseudoRootFactoryNode)
+		{
+			// Move all root node children to top level
+			for (const FString& ChildNodeUid : NodeContainer.GetNodeChildrenUids(FactoryNodeUid))
+			{
+				if (UInterchangeFactoryBaseNode* ChildNode = NodeContainer.GetFactoryNode(ChildNodeUid))
+				{
+					NodeContainer.ClearNodeParentUid(ChildNodeUid);
+				}
+			}
+
+			// Disable the pseudoroot itself (we shouldn't have any asset node for it though)
+			NodeContainer.ReplaceNode(FactoryNodeUid, nullptr);
+		}
+	}
+}	 // namespace UE::InterchangeUsdPipeline::Private
+
+UInterchangeUsdPipeline::UInterchangeUsdPipeline()
+	: GeometryPurpose((int32)(EUsdPurpose::Default | EUsdPurpose::Proxy | EUsdPurpose::Render | EUsdPurpose::Guide))
+	, bImportPseudoRoot(true)
+{
+}
+
+void UInterchangeUsdPipeline::ExecutePipeline(
+	UInterchangeBaseNodeContainer* NodeContainer,
+	const TArray<UInterchangeSourceData*>& InSourceDatas,
+	const FString& ContentBasePath
+)
+{
+	using namespace UE::InterchangeUsdPipeline::Private;
+
+	Super::ExecutePipeline(NodeContainer, InSourceDatas, ContentBasePath);
+
+	if (!NodeContainer)
+	{
+		return;
+	}
+
+	DisableNodesBasedOnGeometryPurpose(*NodeContainer, (EUsdPurpose)GeometryPurpose);
+
+	RepairMaterialInstanceFactoryNodes(*NodeContainer);
+
+	if (!bImportPseudoRoot)
+	{
+		RemovePseudoRootFactoryNode(*NodeContainer);
+	}
+}
+
+bool UInterchangeUsdPipeline::CanExecuteOnAnyThread(EInterchangePipelineTask PipelineTask)
+{
+	return true;
+}
diff --git a/Engine/Plugins/Interchange/Runtime/Source/Pipelines/Public/InterchangeUSDPipeline.h b/Engine/Plugins/Interchange/Runtime/Source/Pipelines/Public/InterchangeUSDPipeline.h
new file mode 100644
--- /dev/null
+++ b/Engine/Plugins/Interchange/Runtime/Source/Pipelines/Public/InterchangeUSDPipeline.h
@@ -0,0 +1,50 @@
+// Copyright Epic Games, Inc. All Rights Reserved.
+
+#pragma once
+
+#include "CoreMinimal.h"
+
+#include "InterchangePipelineBase.h"
+#include "InterchangeSourceData.h"
+#include "Nodes/InterchangeBaseNodeContainer.h"
+
+#include "InterchangeUSDPipeline.generated.h"
+
+class UInterchangeBaseNodeContainer;
+
+UCLASS(BlueprintType)
+class INTERCHANGEPIPELINES_API UInterchangeUsdPipeline : public UInterchangePipelineBase
+{
+	GENERATED_BODY()
+
+	UInterchangeUsdPipeline();
+
+public:
+	UPROPERTY(
+		EditAnywhere,
+		BlueprintReadWrite,
+		Category = "Common",
+		meta = (StandAlonePipelineProperty = "True", PipelineInternalEditionData = "True")
+	)
+	FString PipelineDisplayName;
+
+	/** Only import translated nodes from imageable prims (Xforms, Meshes, etc.) with these specific purposes from the USD file */
+	UPROPERTY(EditAnywhere, BlueprintReadWrite, Category = "USD", meta = (Bitmask, BitmaskEnum = "/Script/UnrealUSDWrapper.EUsdPurpose"))
+	int32 GeometryPurpose;
+
+	/**
+	 * The translator always emits a scene node for the stage pseudoroot. Setting this option to false means we will remove it from
+	 * the factory nodes, making top-level prims lead to top-level actors and components.
+	 */
+	UPROPERTY(EditAnywhere, BlueprintReadWrite, Category = "USD")
+	bool bImportPseudoRoot;
+
+protected:
+	virtual void ExecutePipeline(
+		UInterchangeBaseNodeContainer* BaseNodeContainer,
+		const TArray<UInterchangeSourceData*>& SourceDatas,
+		const FString& ContentBasePath
+	) override;
+
+	virtual bool CanExecuteOnAnyThread(EInterchangePipelineTask PipelineTask) override;
+};
