USD: Handle textures inside USDZ files when importing via Interchange.

diff --git a/Engine/Plugins/Interchange/Runtime/Source/Import/Private/Usd/InterchangeUsdTranslator.cpp b/Engine/Plugins/Interchange/Runtime/Source/Import/Private/Usd/InterchangeUsdTranslator.cpp
--- a/Engine/Plugins/Interchange/Runtime/Source/Import/Private/Usd/InterchangeUsdTranslator.cpp
+++ b/Engine/Plugins/Interchange/Runtime/Source/Import/Private/Usd/InterchangeUsdTranslator.cpp
@@ -6,6 +6,7 @@
 #include "Usd/InterchangeUsdDefinitions.h"
 #include "USDConversionUtils.h"
 #include "USDGeomMeshConversion.h"
+#include "USDLayerUtils.h"
 #include "USDLightConversion.h"
 #include "USDLog.h"
 #include "USDObjectUtils.h"
@@ -176,11 +177,27 @@ namespace UE::InterchangeUsdTranslator::Private
 			bool bIterateShaderGraphs = false
 		);
 
+		/**
+		 * If we're not translating a decompressed USD root, returns TexturePathOnDisk.
+		 * If we are translating a decompressed USD root, returns the path to the USDZ file itself.
+		 *
+		 * The intent here is that in the USDZ case the texture filepath will point at a temp file on disk, that we may dispose
+		 * of later after importing. In order to allow reimporting the texture at a later time, we'll just put the USDZ path
+		 * itself as it's source path, and tweak the USD translator to know what to do with this
+		 */
+		FString GetTextureSourcePath(const FString& TexturePathOnDisk);
+
+		/** If we decompressed a USDZ file to a temp folder this will delete everything from that folder */
+		void CleanUpDecompressedUSDZFolder();
+
 	public:
 		// We have to keep a stage reference so that we can parse the payloads after Translate() complete.
 		// ReleaseSource() clears this member, once translation is complete.
 		UE::FUsdStage UsdStage;
 
+		FString USDZFilePath;
+		FString DecompressedUSDZRoot;
+
 #if USE_USD_SDK
 		// On UInterchangeUSDTranslator::Translate we set this up based on our TranslatorSettings, and then
 		// we can reuse it (otherwise we have to keep converting the FNames into Tokens all the time)
@@ -977,6 +994,56 @@ namespace UE::InterchangeUsdTranslator::Private
 		}
 	}
 
+	FString UInterchangeUSDTranslatorImpl::GetTextureSourcePath(const FString& TexturePathOnDisk)
+	{
+		if (!USDZFilePath.IsEmpty())
+		{
+			return USDZFilePath;
+		}
+		else
+		{
+			return TexturePathOnDisk;
+		}
+	}
+
+	void UInterchangeUSDTranslatorImpl::CleanUpDecompressedUSDZFolder()
+	{
+		if (!DecompressedUSDZRoot.IsEmpty())
+		{
+			const bool bRequireExists = false;
+			const bool bTree = true;
+			IFileManager::Get().DeleteDirectory(*FPaths::GetPath(DecompressedUSDZRoot), bRequireExists, bTree);
+		}
+
+		USDZFilePath.Reset();
+		DecompressedUSDZRoot.Reset();
+	}
+
+	bool DecompressUSDZFileToTempFolder(const FString& InUSDZFilePath, FString& OutDecompressedUSDZRoot)
+	{
+		const bool bIncludeDot = false;
+		FString Extension = FPaths::GetExtension(InUSDZFilePath, bIncludeDot).ToLower();
+		if (Extension == TEXT("usdz"))
+		{
+			const FString Prefix = FPaths::GetBaseFilename(InUSDZFilePath);
+			const FString EmptyExtension = TEXT("");
+			const FString TempFolder = FPaths::CreateTempFilename(FPlatformProcess::UserTempDir(), *Prefix, *EmptyExtension);
+			FString DecompressedRoot;
+			const bool bSuccess = UsdUtils::DecompressUSDZFile(InUSDZFilePath, TempFolder, &DecompressedRoot);
+			if (bSuccess && !DecompressedRoot.IsEmpty())
+			{
+				OutDecompressedUSDZRoot = DecompressedRoot;
+				return true;
+			}
+			else
+			{
+				UE_LOG(LogUsd, Warning, TEXT("Failed to decompress USDZ file '%s': Textures may not be handled correctly."), *InUSDZFilePath);
+			}
+		}
+
+		return false;
+	}
+
 	void AddDisplayColorMaterialInstanceNodeIfNeeded(UInterchangeBaseNodeContainer& NodeContainer, const FString& DisplayColorDesc)
 	{
 		using namespace UsdUnreal::MaterialUtils;
@@ -2521,6 +2588,19 @@ bool UInterchangeUSDTranslator::Translate(UInterchangeBaseNodeContainer& NodeCon
 		return false;
 	}
 
+	// If we're provided a USDZ path, for now we will decompress it into a temp dir and redirect our paths.
+	//
+	// This is mainly because the texture factories must receive a simple file path in order to produce the
+	// their payloads. It's not practical to make them handle USDZ files, and it's not yet possible
+	// to provide them with raw binary buffers directly either
+	if (FString DecompressedUSDZRoot; DecompressUSDZFileToTempFolder(FilePath, DecompressedUSDZRoot))
+	{
+		ImplPtr->USDZFilePath = FilePath;
+		ImplPtr->DecompressedUSDZRoot = DecompressedUSDZRoot;
+
+		FilePath = DecompressedUSDZRoot;
+	}
+
 	// Import should always feel like it's directly from disk, so we ignore already loaded layers and stage cache
 	const bool bUseStageCache = false;
 	const bool bForceReloadLayersFromDisk = true;
@@ -2588,6 +2668,17 @@ void UInterchangeUSDTranslator::ReleaseSource()
 	}
 }
 
+void UInterchangeUSDTranslator::ImportFinish()
+{
+	UE::InterchangeUsdTranslator::Private::UInterchangeUSDTranslatorImpl* ImplPtr = Impl.Get();
+	if (!ImplPtr)
+	{
+		return;
+	}
+
+	ImplPtr->CleanUpDecompressedUSDZFolder();
+}
+
 UInterchangeTranslatorSettings* UInterchangeUSDTranslator::GetSettings() const
 {
 	using namespace UE::InterchangeUsdTranslator::Private;
@@ -2707,7 +2798,8 @@ namespace UE::InterchangeUsdTranslator::Private
 	TOptional<UE::Interchange::FImportImage> GetMaterialXTexturePayloadData(
 		const FString& PayloadKey,
 		TOptional<FString>& AlternateTexturePath,
-		const UInterchangeTranslatorBase* Translator
+		const UInterchangeTranslatorBase* Translator,
+		UInterchangeUSDTranslatorImpl& Impl
 	)
 	{
 		FString Filename = PayloadKey;
@@ -2739,7 +2831,7 @@ namespace UE::InterchangeUsdTranslator::Private
 
 		SourceTranslator->SetResultsContainer(Translator->Results);
 
-		AlternateTexturePath = Filename;
+		AlternateTexturePath = Impl.GetTextureSourcePath(Filename);
 
 		TOptional<UE::Interchange::FImportImage> TexturePayloadData = TextureTranslator->GetTexturePayloadData(PayloadKey, AlternateTexturePath);
 
@@ -2761,6 +2853,12 @@ TOptional<UE::Interchange::FImportImage> UInterchangeUSDTranslator::GetTexturePa
 
 	TOptional<UE::Interchange::FImportImage> TexturePayloadData;
 
+	UInterchangeUSDTranslatorImpl* ImplPtr = Impl.Get();
+	if (!ImplPtr)
+	{
+		return TexturePayloadData;
+	}
+
 #if USE_USD_SDK
 	FString FilePath;
 	TextureGroup TextureGroup;
@@ -2772,7 +2870,7 @@ TOptional<UE::Interchange::FImportImage> UInterchangeUSDTranslator::GetTexturePa
 		const IInterchangeTexturePayloadInterface* TextureTranslator = ScopedTranslator.GetPayLoadInterface<IInterchangeTexturePayloadInterface>();
 		if (TextureTranslator)
 		{
-			AlternateTexturePath = FilePath;
+			AlternateTexturePath = ImplPtr->GetTextureSourcePath(FilePath);
 
 			// The texture translators don't use the payload key, and read the texture directly from the SourceData's file path
 			const FString UnusedPayloadKey = {};
@@ -2806,11 +2904,11 @@ TOptional<UE::Interchange::FImportImage> UInterchangeUSDTranslator::GetTexturePa
 		}
 	}
 
-	// if we couldn't find a texture in either the USD Translator nor the Translators, then it's most likely coming from reading an mtlx in memory,
+	// If we couldn't find a texture in either the USD Translator nor the Translators, then it's most likely coming from reading an mtlx in memory,
 	// copying the behavior from InterchangeMaterialXTranslator
 	if (!TexturePayloadData.IsSet())
 	{
-		TexturePayloadData = GetMaterialXTexturePayloadData(PayloadKey, AlternateTexturePath, this);
+		TexturePayloadData = GetMaterialXTexturePayloadData(PayloadKey, AlternateTexturePath, this, *ImplPtr);
 	}
 
 	return TexturePayloadData;
@@ -2823,6 +2921,12 @@ TOptional<UE::Interchange::FImportBlockedImage> UInterchangeUSDTranslator::GetBl
 {
 	using namespace UE::InterchangeUsdTranslator::Private;
 
+	UInterchangeUSDTranslatorImpl* ImplPtr = Impl.Get();
+	if (!ImplPtr)
+	{
+		return {};
+	}
+
 	UE::Interchange::FImportBlockedImage BlockData;
 
 #if USE_USD_SDK
@@ -2834,7 +2938,7 @@ TOptional<UE::Interchange::FImportBlockedImage> UInterchangeUSDTranslator::GetBl
 		return {};
 	}
 
-	AlternateTexturePath = FilePath;
+	AlternateTexturePath = ImplPtr->GetTextureSourcePath(FilePath);
 
 	// Collect all the UDIM tile filepaths similar to this current tile. If we've been asked to translate
 	// a blocked texture then we must have some
diff --git a/Engine/Plugins/Interchange/Runtime/Source/Import/Public/Usd/InterchangeUsdTranslator.h b/Engine/Plugins/Interchange/Runtime/Source/Import/Public/Usd/InterchangeUsdTranslator.h
--- a/Engine/Plugins/Interchange/Runtime/Source/Import/Public/Usd/InterchangeUsdTranslator.h
+++ b/Engine/Plugins/Interchange/Runtime/Source/Import/Public/Usd/InterchangeUsdTranslator.h
@@ -76,6 +76,7 @@ public:
 	virtual TArray<FString> GetSupportedFormats() const override;
 	virtual bool Translate(UInterchangeBaseNodeContainer& BaseNodeContainer) const override;
 	virtual void ReleaseSource() override;
+	virtual void ImportFinish() override;
 	virtual UInterchangeTranslatorSettings* GetSettings() const override;
 	virtual void SetSettings(const UInterchangeTranslatorSettings* InterchangeTranslatorSettings) override;
 	/** End UInterchangeTranslatorBase API*/
diff --git a/Engine/Plugins/Runtime/USDCore/Source/USDUtilities/Private/USDLayerUtils.cpp b/Engine/Plugins/Runtime/USDCore/Source/USDUtilities/Private/USDLayerUtils.cpp
--- a/Engine/Plugins/Runtime/USDCore/Source/USDUtilities/Private/USDLayerUtils.cpp
+++ b/Engine/Plugins/Runtime/USDCore/Source/USDUtilities/Private/USDLayerUtils.cpp
@@ -15,6 +15,7 @@
 #include "UsdWrappers/UsdStage.h"
 
 #include "Framework/Application/SlateApplication.h"
+#include "HAL/FileManager.h"
 #include "Misc/Paths.h"
 
 #if WITH_EDITOR
@@ -33,6 +34,7 @@
 #include "pxr/usd/usd/modelAPI.h"
 #include "pxr/usd/usd/primCompositionQuery.h"
 #include "pxr/usd/usd/stage.h"
+#include "pxr/usd/usd/zipFile.h"
 #include "pxr/usd/usdGeom/xform.h"
 #include "pxr/usd/usdUtils/flattenLayerStack.h"
 #include "USDIncludesEnd.h"
@@ -323,6 +325,117 @@ FString UsdUtils::MakePathRelativeToProjectDir(const FString& Path)
 	return PathConverted;
 }
 
+bool UsdUtils::SplitUSDZAssetPath(const FString& InFilePathIntoUSDZArchive, FString& OutUSDZFilePath, FString& OutInnerAssetPath)
+{
+	// We need at least an opening and closing bracket
+	if (InFilePathIntoUSDZArchive.Len() < 3)
+	{
+		return false;
+	}
+
+	const int32 OpenBracketPos = InFilePathIntoUSDZArchive.Find(TEXT("["), ESearchCase::IgnoreCase, ESearchDir::FromEnd);
+	const int32 CloseBracketPos = InFilePathIntoUSDZArchive.Find(TEXT("]"), ESearchCase::IgnoreCase, ESearchDir::FromEnd);
+	if (OpenBracketPos != INDEX_NONE &&								 //
+		CloseBracketPos == InFilePathIntoUSDZArchive.Len() - 1 &&	 //
+		CloseBracketPos > OpenBracketPos - 1)
+	{
+		// Should be like "C:/MyFiles/scene.usdz"
+		OutUSDZFilePath = InFilePathIntoUSDZArchive.Left(OpenBracketPos);
+
+		// Should be like "0/texture.png"
+		OutInnerAssetPath = InFilePathIntoUSDZArchive.Mid(OpenBracketPos + 1, CloseBracketPos - OpenBracketPos - 1);
+
+		return true;
+	}
+
+	return false;
+}
+
+bool UsdUtils::DecompressUSDZFile(const FString& InUSDZFilePath, const FString& InOutputDir, FString* OutDecompressedRootLayer)
+{
+	if (!IFileManager::Get().FileExists(*InUSDZFilePath))
+	{
+		return false;
+	}
+
+	// Ensure directory exists
+	const bool bTree = true;
+	if (!IFileManager::Get().MakeDirectory(*InOutputDir, bTree))
+	{
+		return false;
+	}
+
+	FScopedUsdAllocs Allocs;
+
+	bool bTraversedFirstFile = false;
+
+	pxr::UsdZipFile File = pxr::UsdZipFile::Open(UnrealToUsd::ConvertString(*InUSDZFilePath).Get());
+	for (pxr::UsdZipFile::Iterator Iter = File.begin(), End = File.end(); Iter != End; ++Iter)
+	{
+		const pxr::UsdZipFile::FileInfo Info = Iter.GetFileInfo();
+		const std::string InnerFilePath = *Iter;
+		const FString UEInnerFilePath = UsdToUnreal::ConvertString(InnerFilePath);
+		const char* FileStart = Iter.GetFile();
+		const std::size_t InnerFileSizeBytes = Info.size;
+		const FString TargetFilePath = FPaths::Combine(InOutputDir, UEInnerFilePath);
+
+		// According to USDZ spec (https://openusd.org/release/spec_usdz.html),
+		// the very first file in the archive should be a native USD file, which will be the root layer.
+		if (!bTraversedFirstFile && OutDecompressedRootLayer != nullptr)
+		{
+			bTraversedFirstFile = true;
+
+			const bool bIncludeDot = false;
+			FString Extension = FPaths::GetExtension(UEInnerFilePath, bIncludeDot).ToLower();
+			if (Extension.StartsWith(TEXT("usd")))
+			{
+				*OutDecompressedRootLayer = TargetFilePath;
+			}
+			else
+			{
+				UE_LOG(
+					LogUsd,
+					Warning,
+					TEXT("Failed to decompress %s: First file '%s' should always be an USD file!"),
+					*InUSDZFilePath,
+					*UEInnerFilePath
+				);
+				return false;
+			}
+		}
+
+		// Check for these, but it's very unlikely they differ from these values as USD never uses
+		// compression/encryption, and Usd_UsdzResolver::OpenAsset also just errors out if it runs into these
+		if (Info.compressionMethod != 0)
+		{
+			UE_LOG(LogUsd, Warning, TEXT("Cannot open %s in %s: Compressed files are not supported!"), *UEInnerFilePath, *InUSDZFilePath);
+			continue;
+		}
+		if (Info.encrypted)
+		{
+			UE_LOG(LogUsd, Warning, TEXT("Cannot open %s in %s: Encrypted files are not supported!"), *UEInnerFilePath, *InUSDZFilePath);
+			continue;
+		}
+
+		// Create any nested folder (usually a folder named "0")
+		if (!IFileManager::Get().MakeDirectory(*FPaths::GetPath(TargetFilePath), bTree))
+		{
+			continue;
+		}
+
+		TUniquePtr<FArchive> Ar{IFileManager::Get().CreateFileWriter(*TargetFilePath)};
+		if (!Ar)
+		{
+			continue;
+		}
+
+		Ar->Serialize((void*)FileStart, InnerFileSizeBytes);
+		Ar->Close();
+	}
+
+	return true;
+}
+
 TUsdStore<pxr::SdfLayerRefPtr> UsdUtils::CreateNewLayer(
 	TUsdStore<pxr::UsdStageRefPtr> UsdStage,
 	const TUsdStore<pxr::SdfLayerRefPtr>& ParentLayer,
diff --git a/Engine/Plugins/Runtime/USDCore/Source/USDUtilities/Private/USDShadeConversion.cpp b/Engine/Plugins/Runtime/USDCore/Source/USDUtilities/Private/USDShadeConversion.cpp
--- a/Engine/Plugins/Runtime/USDCore/Source/USDUtilities/Private/USDShadeConversion.cpp
+++ b/Engine/Plugins/Runtime/USDCore/Source/USDUtilities/Private/USDShadeConversion.cpp
@@ -197,31 +197,17 @@ namespace UE::UsdShadeConversion::Private
 	// return true
 	bool IsInsideUsdzArchive(const FString& ResolvedTexturePath, FString& OutExtensionNoDot)
 	{
-		// We need at least an opening and closing bracket
-		if (ResolvedTexturePath.Len() < 3)
+		FString USDZFilePath;
+		FString InnerAssetPath;
+		const bool bSplit = UsdUtils::SplitUSDZAssetPath(ResolvedTexturePath, USDZFilePath, InnerAssetPath);
+		if (!bSplit)
 		{
 			return false;
 		}
 
-		int32 OpenBracketPos = ResolvedTexturePath.Find(TEXT("["), ESearchCase::IgnoreCase, ESearchDir::FromEnd);
-		int32 CloseBracketPos = ResolvedTexturePath.Find(TEXT("]"), ESearchCase::IgnoreCase, ESearchDir::FromEnd);
-		if (OpenBracketPos != INDEX_NONE && CloseBracketPos != INDEX_NONE && CloseBracketPos > OpenBracketPos - 1)
-		{
-			// Should be like "C:/MyFiles/scene.usdz"
-			FString UsdzFilePath = ResolvedTexturePath.Left(OpenBracketPos);
-
-			// Should be like "0/texture.png"
-			FString TextureInUsdzPath = ResolvedTexturePath.Mid(OpenBracketPos + 1, CloseBracketPos - OpenBracketPos - 1);
-
-			if (FPaths::GetExtension(UsdzFilePath).ToLower() == TEXT("usdz"))
-			{
-				// Should be something like "png"
-				OutExtensionNoDot = FPaths::GetExtension(TextureInUsdzPath);
-				return true;
-			}
-		}
-
-		return false;
+		// Should be something like "png"
+		OutExtensionNoDot = FPaths::GetExtension(InnerAssetPath);
+		return true;
 	}
 
 	TUsdStore<std::shared_ptr<const char>> ReadTextureBufferFromUsdzArchive(const FString& ResolvedTexturePath, uint64& OutBufferSize)
diff --git a/Engine/Plugins/Runtime/USDCore/Source/USDUtilities/Public/USDLayerUtils.h b/Engine/Plugins/Runtime/USDCore/Source/USDUtilities/Public/USDLayerUtils.h
--- a/Engine/Plugins/Runtime/USDCore/Source/USDUtilities/Public/USDLayerUtils.h
+++ b/Engine/Plugins/Runtime/USDCore/Source/USDUtilities/Public/USDLayerUtils.h
@@ -97,6 +97,20 @@ namespace UsdUtils
 	 */
 	USDUTILITIES_API FString MakePathRelativeToProjectDir(const FString& Path);
 
+	/**
+	 * If InFilePathIntoUSDZArchive is a path to something inside an USDZ file, like "C:/MyFiles/scene.usdz[0/texture.png]",
+	 * this will put "C:/MyFiles/scene.usdz" on OutUSDZFilePath, and "0/texture.png" on OutInnerAssetPath and return true.
+	 * Returns false when it can't split the provided path in that way.
+	 */
+	USDUTILITIES_API bool SplitUSDZAssetPath(const FString& InFilePathIntoUSDZArchive, FString& OutUSDZFilePath, FString& OutInnerAssetPath);
+
+	/**
+	 * If InUSDZFilePath points at an actual USDZ file, this will decompress the archive and put its contents inside InOutputDir,
+	 * returning true if succeeding. In that case, if OutDecompressedRootLayer is not nullptr, this function will also set it to
+	 * contain the path to the root layer of the decompressed USDZ archive
+	 */
+	USDUTILITIES_API bool DecompressUSDZFile(const FString& InUSDZFilePath, const FString& InOutputDir, FString* OutDecompressedRootLayer = nullptr);
+
 	/** Creates a new layer with a default prim */
 	USDUTILITIES_API TUsdStore<pxr::SdfLayerRefPtr> CreateNewLayer(
 		TUsdStore<pxr::UsdStageRefPtr> UsdStage,
