USD: Fix broken build due to missing USE_USD_SDK checks.

diff --git a/Engine/Plugins/Interchange/Runtime/Source/Import/Private/Usd/InterchangeUsdTranslator.cpp b/Engine/Plugins/Interchange/Runtime/Source/Import/Private/Usd/InterchangeUsdTranslator.cpp
--- a/Engine/Plugins/Interchange/Runtime/Source/Import/Private/Usd/InterchangeUsdTranslator.cpp
+++ b/Engine/Plugins/Interchange/Runtime/Source/Import/Private/Usd/InterchangeUsdTranslator.cpp
@@ -2670,6 +2670,7 @@ void UInterchangeUSDTranslator::ReleaseSource()
 
 void UInterchangeUSDTranslator::ImportFinish()
 {
+#if USE_USD_SDK
 	UE::InterchangeUsdTranslator::Private::UInterchangeUSDTranslatorImpl* ImplPtr = Impl.Get();
 	if (!ImplPtr)
 	{
@@ -2677,6 +2678,7 @@ void UInterchangeUSDTranslator::ImportFinish()
 	}
 
 	ImplPtr->CleanUpDecompressedUSDZFolder();
+#endif	  // USE_USD_SDK
 }
 
 UInterchangeTranslatorSettings* UInterchangeUSDTranslator::GetSettings() const
@@ -2802,6 +2804,9 @@ namespace UE::InterchangeUsdTranslator::Private
 		UInterchangeUSDTranslatorImpl& Impl
 	)
 	{
+		TOptional<UE::Interchange::FImportImage> TexturePayloadData;
+
+#if USE_USD_SDK
 		FString Filename = PayloadKey;
 		TextureCompressionSettings CompressionSettings = TextureCompressionSettings::TC_Default;
 
@@ -2817,7 +2822,7 @@ namespace UE::InterchangeUsdTranslator::Private
 
 		if (!PayloadSourceData)
 		{
-			return TOptional<UE::Interchange::FImportImage>();
+			return TexturePayloadData;
 		}
 
 		UInterchangeTranslatorBase* SourceTranslator = UInterchangeManager::GetInterchangeManager().GetTranslatorForSourceData(PayloadSourceData);
@@ -2826,19 +2831,19 @@ namespace UE::InterchangeUsdTranslator::Private
 
 		if (!TextureTranslator)
 		{
-			return TOptional<UE::Interchange::FImportImage>();
+			return TexturePayloadData;
 		}
 
 		SourceTranslator->SetResultsContainer(Translator->Results);
 
 		AlternateTexturePath = Impl.GetTextureSourcePath(Filename);
 
-		TOptional<UE::Interchange::FImportImage> TexturePayloadData = TextureTranslator->GetTexturePayloadData(PayloadKey, AlternateTexturePath);
-
+		TexturePayloadData = TextureTranslator->GetTexturePayloadData(PayloadKey, AlternateTexturePath);
 		if (TexturePayloadData.IsSet())
 		{
 			TexturePayloadData.GetValue().CompressionSettings = CompressionSettings;
 		}
+#endif	  // USE_USD_SDK
 
 		return TexturePayloadData;
 	}
