Interchange: Add support for specifying component/actor visibility on scene and factory nodes;

Rename "Visibility" to ActorHiddenInGame for clarity;
Fix component visibility tracks not using the correct track type, preventing them from toggling the component's bVisible alongside bHiddenInGame.

diff --git a/Engine/Plugins/Experimental/Enterprise/DatasmithInterchange/Source/DatasmithInterchange/Private/InterchangeDatasmithAnimationUtils.cpp b/Engine/Plugins/Experimental/Enterprise/DatasmithInterchange/Source/DatasmithInterchange/Private/InterchangeDatasmithAnimationUtils.cpp
--- a/Engine/Plugins/Experimental/Enterprise/DatasmithInterchange/Source/DatasmithInterchange/Private/InterchangeDatasmithAnimationUtils.cpp
+++ b/Engine/Plugins/Experimental/Enterprise/DatasmithInterchange/Source/DatasmithInterchange/Private/InterchangeDatasmithAnimationUtils.cpp
@@ -138,7 +138,7 @@ namespace UE::DatasmithInterchange::AnimUtils
 		const FString ActorNodeUid = NodeUtils::GetActorUid(VisibilityAnimation.GetName());
 		TrackNode->SetCustomActorDependencyUid(ActorNodeUid);
 
-		TrackNode->SetCustomPropertyTrack(EInterchangePropertyTracks::Visibility);
+		TrackNode->SetCustomPropertyTrack(EInterchangePropertyTracks::ActorHiddenInGame);
 		TrackNode->SetCustomAnimationPayloadKey(TrackNode->GetUniqueID(), EInterchangeAnimationPayLoadType::STEPCURVE);
 
 		LevelSequenceNode->AddCustomAnimationTrackUid(TrackNode->GetUniqueID());
diff --git a/Engine/Plugins/Interchange/Runtime/Config/BaseInterchange.ini b/Engine/Plugins/Interchange/Runtime/Config/BaseInterchange.ini
--- a/Engine/Plugins/Interchange/Runtime/Config/BaseInterchange.ini
+++ b/Engine/Plugins/Interchange/Runtime/Config/BaseInterchange.ini
@@ -3,6 +3,7 @@ bCanSaveAllSections=True
 
 [CoreRedirects]
 +EnumRedirects=(OldName="/Script/InterchangePipelines.EInterchangeMaterialImportOption",ValueChanges=(("DoNotImport","ImportAsMaterials")))
++EnumRedirects=(OldName="/Script/InterchangeNodes.EInterchangePropertyTracks",ValueChanges=(("Visibility","ActorHiddenInGame")))
 
 ; Rename Interchange CineCamera to PhysicalCamera FactoryNode
 +ClassRedirects=(OldName="/Script/InterchangeFactoryNodes.InterchangeCineCameraFactoryNode", NewName="/Script/InterchangeFactoryNodes.InterchangePhysicalCameraFactoryNode")
diff --git a/Engine/Plugins/Interchange/Runtime/Source/FactoryNodes/Private/InterchangeActorFactoryNode.cpp b/Engine/Plugins/Interchange/Runtime/Source/FactoryNodes/Private/InterchangeActorFactoryNode.cpp
--- a/Engine/Plugins/Interchange/Runtime/Source/FactoryNodes/Private/InterchangeActorFactoryNode.cpp
+++ b/Engine/Plugins/Interchange/Runtime/Source/FactoryNodes/Private/InterchangeActorFactoryNode.cpp
@@ -35,7 +35,12 @@ bool UInterchangeActorFactoryNode::GetCustomGlobalTransform(FTransform& Attribut
 
 bool UInterchangeActorFactoryNode::SetCustomGlobalTransform(const FTransform& AttributeValue, bool bAddApplyDelegate)
 {
-	IMPLEMENT_NODE_ATTRIBUTE_SETTER_WITH_CUSTOM_DELEGATE_WITH_CUSTOM_CLASS(UInterchangeActorFactoryNode, GlobalTransform, FTransform, USceneComponent);
+	IMPLEMENT_NODE_ATTRIBUTE_SETTER_WITH_CUSTOM_DELEGATE_WITH_CUSTOM_CLASS(
+		UInterchangeActorFactoryNode,
+		GlobalTransform,
+		FTransform,
+		USceneComponent
+	);
 }
 
 bool UInterchangeActorFactoryNode::GetCustomLocalTransform(FTransform& AttributeValue) const
@@ -45,7 +50,32 @@ bool UInterchangeActorFactoryNode::GetCustomLocalTransform(FTransform& Attribute
 
 bool UInterchangeActorFactoryNode::SetCustomLocalTransform(const FTransform& AttributeValue, bool bAddApplyDelegate)
 {
-	IMPLEMENT_NODE_ATTRIBUTE_SETTER_WITH_CUSTOM_DELEGATE_WITH_CUSTOM_CLASS(UInterchangeActorFactoryNode, GlobalTransform, FTransform, USceneComponent);
+	IMPLEMENT_NODE_ATTRIBUTE_SETTER_WITH_CUSTOM_DELEGATE_WITH_CUSTOM_CLASS(
+		UInterchangeActorFactoryNode,
+		GlobalTransform,
+		FTransform,
+		USceneComponent
+	);
+}
+
+bool UInterchangeActorFactoryNode::GetCustomComponentVisibility(bool& AttributeValue) const
+{
+	IMPLEMENT_NODE_ATTRIBUTE_GETTER(ComponentVisibility, bool);
+}
+
+bool UInterchangeActorFactoryNode::SetCustomComponentVisibility(bool AttributeValue, bool bAddApplyDelegate)
+{
+	IMPLEMENT_NODE_ATTRIBUTE_SETTER_WITH_CUSTOM_DELEGATE_WITH_CUSTOM_CLASS(UInterchangeActorFactoryNode, ComponentVisibility, bool, USceneComponent);
+}
+
+bool UInterchangeActorFactoryNode::GetCustomActorVisibility(bool& AttributeValue) const
+{
+	IMPLEMENT_NODE_ATTRIBUTE_GETTER(ActorVisibility, bool);
+}
+
+bool UInterchangeActorFactoryNode::SetCustomActorVisibility(bool AttributeValue, bool bAddApplyDelegate)
+{
+	IMPLEMENT_NODE_ATTRIBUTE_SETTER_WITH_CUSTOM_DELEGATE_WITH_CUSTOM_CLASS(UInterchangeActorFactoryNode, ActorVisibility, bool, USceneComponent);
 }
 
 bool UInterchangeActorFactoryNode::ApplyCustomGlobalTransformToAsset(UObject* Asset) const
@@ -86,13 +116,116 @@ bool UInterchangeActorFactoryNode::FillCustomGlobalTransformFromAsset(UObject* A
 	return false;
 }
 
+bool UInterchangeActorFactoryNode::ApplyCustomComponentVisibilityToAsset(UObject* Asset) const
+{
+	bool bVisible = true;
+	if (GetCustomComponentVisibility(bVisible))
+	{
+		if (USceneComponent* Component = Cast<USceneComponent>(Asset))
+		{
+			Component->SetVisibility(bVisible);
+			return true;
+		}
+	}
+
+	return false;
+}
+
+bool UInterchangeActorFactoryNode::FillCustomComponentVisibilityFromAsset(UObject* Asset)
+{
+	if (const USceneComponent* Component = Cast<USceneComponent>(Asset))
+	{
+		const bool bVisible = Component->GetVisibleFlag();
+		const bool bAddApplyDelegate = false;
+		return SetCustomComponentVisibility(bVisible, bAddApplyDelegate);
+	}
+
+	return false;
+}
+
+bool UInterchangeActorFactoryNode::ApplyCustomActorVisibilityToAsset(UObject* Asset) const
+{
+	bool bVisible = true;
+	if (GetCustomActorVisibility(bVisible))
+	{
+		AActor* Actor = nullptr;
+		if (USceneComponent* Component = Cast<USceneComponent>(Asset))
+		{
+			Actor = Component->GetOwner();
+		}
+		else
+		{
+			Actor = Cast<AActor>(Asset);
+		}
+
+		if (Actor)
+		{
+			Actor->SetActorHiddenInGame(!bVisible);
+#if WITH_EDITOR
+			// This also hides the actor on the editor viewport
+			Actor->SetIsTemporarilyHiddenInEditor(!bVisible);
+#endif	  // WITH_EDITOR
+			return true;
+		}
+	}
+
+	return false;
+}
+
+bool UInterchangeActorFactoryNode::FillCustomActorVisibilityFromAsset(UObject* Asset)
+{
+	AActor* Actor = nullptr;
+	if (USceneComponent* Component = Cast<USceneComponent>(Asset))
+	{
+		Actor = Component->GetOwner();
+	}
+	else
+	{
+		Actor = Cast<AActor>(Asset);
+	}
+
+	if (Actor)
+	{
+#if WITH_EDITOR
+		const bool bVisible = !Actor->IsHiddenEd() && !Actor->IsHidden();
+#else
+		const bool bVisible = !Actor->IsHidden();
+#endif	  // WITH_EDITOR
+
+		const bool bAddApplyDelegate = false;
+		return SetCustomActorVisibility(bVisible, bAddApplyDelegate);
+	}
+
+	return false;
+}
+
 void UInterchangeActorFactoryNode::CopyWithObject(const UInterchangeFactoryBaseNode* SourceNode, UObject* Object)
 {
 	Super::CopyWithObject(SourceNode, Object);
 
 	if (const UInterchangeActorFactoryNode* ActorFactoryNode = Cast<UInterchangeActorFactoryNode>(SourceNode))
 	{
-		COPY_NODE_DELEGATES_WITH_CUSTOM_DELEGATE(ActorFactoryNode, UInterchangeActorFactoryNode, GlobalTransform, FTransform, USceneComponent::StaticClass())
+		COPY_NODE_DELEGATES_WITH_CUSTOM_DELEGATE(
+			ActorFactoryNode,
+			UInterchangeActorFactoryNode,
+			GlobalTransform,
+			FTransform,
+			USceneComponent::StaticClass()
+		)
+		COPY_NODE_DELEGATES_WITH_CUSTOM_DELEGATE(
+			ActorFactoryNode,
+			UInterchangeActorFactoryNode,
+			ComponentVisibility,
+			bool,
+			USceneComponent::StaticClass()
+		)
+		COPY_NODE_DELEGATES_WITH_CUSTOM_DELEGATE(
+			ActorFactoryNode,
+			UInterchangeActorFactoryNode,
+			ActorVisibility,
+			bool,
+			USceneComponent::StaticClass()
+		)
 	}
 }
 
@@ -158,5 +291,6 @@ void UInterchangeActorFactoryNode::OnRestoreAllCustomAttributeDelegates()
 {
 	REFILL_CUSTOM_ATTRIBUTE_APPLY_DELEGATE(GlobalTransform, FTransform)
 	REFILL_CUSTOM_ATTRIBUTE_APPLY_DELEGATE(LocalTransform, FTransform)
+	REFILL_CUSTOM_ATTRIBUTE_APPLY_DELEGATE(ComponentVisibility, bool)
+	REFILL_CUSTOM_ATTRIBUTE_APPLY_DELEGATE(ActorVisibility, bool)
 }
-
diff --git a/Engine/Plugins/Interchange/Runtime/Source/FactoryNodes/Public/InterchangeActorFactoryNode.h b/Engine/Plugins/Interchange/Runtime/Source/FactoryNodes/Public/InterchangeActorFactoryNode.h
--- a/Engine/Plugins/Interchange/Runtime/Source/FactoryNodes/Public/InterchangeActorFactoryNode.h
+++ b/Engine/Plugins/Interchange/Runtime/Source/FactoryNodes/Public/InterchangeActorFactoryNode.h
@@ -28,6 +28,18 @@ public:
 	UFUNCTION(BlueprintCallable, Category = "Interchange | Node | ActorFactory")
 	bool SetCustomLocalTransform(const FTransform& AttributeValue, bool bAddApplyDelegate = true);
 
+	UFUNCTION(BlueprintCallable, Category = "Interchange | Node | ActorFactory")
+	bool GetCustomComponentVisibility(bool& AttributeValue) const;
+
+	UFUNCTION(BlueprintCallable, Category = "Interchange | Node | ActorFactory")
+	bool SetCustomComponentVisibility(bool AttributeValue, bool bAddApplyDelegate = true);
+
+	UFUNCTION(BlueprintCallable, Category = "Interchange | Node | ActorFactory")
+	bool GetCustomActorVisibility(bool& AttributeValue) const;
+
+	UFUNCTION(BlueprintCallable, Category = "Interchange | Node | ActorFactory")
+	bool SetCustomActorVisibility(bool AttributeValue, bool bAddApplyDelegate = true);
+
 	UFUNCTION(BlueprintCallable, Category = "Interchange | Node | ActorFactory")
 	bool GetCustomActorClassName(FString& AttributeValue) const
 	{
@@ -87,10 +99,18 @@ private:
 	bool ApplyCustomGlobalTransformToAsset(UObject* Asset) const;
 	bool FillCustomGlobalTransformFromAsset(UObject* Asset);
 
-	const UE::Interchange::FAttributeKey Macro_CustomGlobalTransformKey = UE::Interchange::FAttributeKey(TEXT("GlobalTransform"));
-	const UE::Interchange::FAttributeKey Macro_CustomLocalTransformKey = UE::Interchange::FAttributeKey(TEXT("LocalTransform"));
-	const UE::Interchange::FAttributeKey Macro_CustomActorClassNameKey = UE::Interchange::FAttributeKey(TEXT("ActorClassName"));
-	const UE::Interchange::FAttributeKey Macro_CustomMobilityKey = UE::Interchange::FAttributeKey(TEXT("Mobility"));
+	bool ApplyCustomComponentVisibilityToAsset(UObject* Asset) const;
+	bool FillCustomComponentVisibilityFromAsset(UObject* Asset);
+
+	bool ApplyCustomActorVisibilityToAsset(UObject* Asset) const;
+	bool FillCustomActorVisibilityFromAsset(UObject* Asset);
+
+	IMPLEMENT_NODE_ATTRIBUTE_KEY(GlobalTransform);
+	IMPLEMENT_NODE_ATTRIBUTE_KEY(LocalTransform);
+	IMPLEMENT_NODE_ATTRIBUTE_KEY(ComponentVisibility);
+	IMPLEMENT_NODE_ATTRIBUTE_KEY(ActorVisibility);
+	IMPLEMENT_NODE_ATTRIBUTE_KEY(ActorClassName);
+	IMPLEMENT_NODE_ATTRIBUTE_KEY(Mobility);
 
 	//A scene node can be part of multiple Layers.
 	UE::Interchange::TArrayAttributeHelper<FString> LayerNames;
diff --git a/Engine/Plugins/Interchange/Runtime/Source/Import/Private/Animation/InterchangeLevelSequenceFactory.cpp b/Engine/Plugins/Interchange/Runtime/Source/Import/Private/Animation/InterchangeLevelSequenceFactory.cpp
--- a/Engine/Plugins/Interchange/Runtime/Source/Import/Private/Animation/InterchangeLevelSequenceFactory.cpp
+++ b/Engine/Plugins/Interchange/Runtime/Source/Import/Private/Animation/InterchangeLevelSequenceFactory.cpp
@@ -603,7 +603,7 @@ namespace UE::Interchange::Private
 		UE::Interchange::FAnimationPayloadData* PayloadData = &PayloadDataArray[0];
 
 		FGuid ObjectBinding;
-		if(PropertyTrack == EInterchangePropertyTracks::Visibility)
+		if(PropertyTrack == EInterchangePropertyTracks::ActorHiddenInGame)
 		{
 			ObjectBinding = BindActorToLevelSequence(Actor);
 		}
diff --git a/Engine/Plugins/Interchange/Runtime/Source/Import/Private/Animation/InterchangeLevelSequenceHelper.cpp b/Engine/Plugins/Interchange/Runtime/Source/Import/Private/Animation/InterchangeLevelSequenceHelper.cpp
--- a/Engine/Plugins/Interchange/Runtime/Source/Import/Private/Animation/InterchangeLevelSequenceHelper.cpp
+++ b/Engine/Plugins/Interchange/Runtime/Source/Import/Private/Animation/InterchangeLevelSequenceHelper.cpp
@@ -57,7 +57,7 @@ namespace UE::Interchange::Private
 			{EInterchangePropertyTracks::DrawDebug, {UMovieSceneBoolTrack::StaticClass()->GetName(),  GET_MEMBER_NAME_STRING_CHECKED(USplineComponent, bDrawDebug) , TEXT("Draw Debug")}},
 			{EInterchangePropertyTracks::EmissiveLightSource, {UMovieSceneBoolTrack::StaticClass()->GetName(),  GET_MEMBER_NAME_STRING_CHECKED(UPrimitiveComponent, bEmissiveLightSource) , TEXT("Emissive Light Source")}},
 			{EInterchangePropertyTracks::ExcludeFromLightAttachmentGroup, {UMovieSceneBoolTrack::StaticClass()->GetName(),  GET_MEMBER_NAME_STRING_CHECKED(UPrimitiveComponent, bExcludeFromLightAttachmentGroup) , TEXT("Exclude from Light Attachment Group")}},
-			{EInterchangePropertyTracks::HiddenInGame, {UMovieSceneBoolTrack::StaticClass()->GetName(),  GET_MEMBER_NAME_STRING_CHECKED(USceneComponent, bHiddenInGame) , TEXT("Hidden in Game")}},
+			{EInterchangePropertyTracks::HiddenInGame, {UMovieSceneVisibilityTrack::StaticClass()->GetName(),  GET_MEMBER_NAME_STRING_CHECKED(USceneComponent, bHiddenInGame) , TEXT("Hidden in Game")}},
 			{EInterchangePropertyTracks::HiddenInSceneCapture, {UMovieSceneBoolTrack::StaticClass()->GetName(),  GET_MEMBER_NAME_STRING_CHECKED(UPrimitiveComponent, bHiddenInSceneCapture) , TEXT("Hidden In Scene Capture")}},
 			{EInterchangePropertyTracks::Holdout, {UMovieSceneBoolTrack::StaticClass()->GetName(),  GET_MEMBER_NAME_STRING_CHECKED(UPrimitiveComponent, bHoldout) , TEXT("Holdout")}},
 			{EInterchangePropertyTracks::LightAttachmentsAsGroup, {UMovieSceneBoolTrack::StaticClass()->GetName(),  GET_MEMBER_NAME_STRING_CHECKED(UPrimitiveComponent, bLightAttachmentsAsGroup) , TEXT("Light Attachments as Group")}},
@@ -72,7 +72,7 @@ namespace UE::Interchange::Private
 			{EInterchangePropertyTracks::TranslucencySortDistanceOffset, {UMovieSceneFloatTrack::StaticClass()->GetName(),  GET_MEMBER_NAME_STRING_CHECKED(UPrimitiveComponent, TranslucencySortDistanceOffset) , TEXT("Translucency Sort Distance Offset")}},
 			{EInterchangePropertyTracks::VisibleInRayTracing, {UMovieSceneBoolTrack::StaticClass()->GetName(),  GET_MEMBER_NAME_STRING_CHECKED(UPrimitiveComponent, bVisibleInRayTracing) , TEXT("Visible in Ray Tracing")}},
 			{EInterchangePropertyTracks::VisibleInSceneCaptureOnly, {UMovieSceneBoolTrack::StaticClass()->GetName(),  GET_MEMBER_NAME_STRING_CHECKED(UPrimitiveComponent, bVisibleInSceneCaptureOnly) , TEXT("Visible In Scene Capture Only")}},
-			{EInterchangePropertyTracks::Visibility, {UMovieSceneVisibilityTrack::StaticClass()->GetName(), AActor::GetHiddenPropertyName().ToString(), TEXT("Visibility")}},
+			{EInterchangePropertyTracks::ActorHiddenInGame, {UMovieSceneVisibilityTrack::StaticClass()->GetName(), AActor::GetHiddenPropertyName().ToString(), TEXT("Visibility")}},
 
 			// Light
 			{EInterchangePropertyTracks::LightAffectGlobalIllumination, {UMovieSceneBoolTrack::StaticClass()->GetName(),  GET_MEMBER_NAME_STRING_CHECKED(ULightComponentBase, bAffectGlobalIllumination) , TEXT("Affect Ray Tracing Global Illumination")}},
diff --git a/Engine/Plugins/Interchange/Runtime/Source/Import/Private/Usd/InterchangeUsdTranslator.cpp b/Engine/Plugins/Interchange/Runtime/Source/Import/Private/Usd/InterchangeUsdTranslator.cpp
--- a/Engine/Plugins/Interchange/Runtime/Source/Import/Private/Usd/InterchangeUsdTranslator.cpp
+++ b/Engine/Plugins/Interchange/Runtime/Source/Import/Private/Usd/InterchangeUsdTranslator.cpp
@@ -109,7 +109,7 @@ namespace UE::InterchangeUsdTranslator::Private
 	// clang-format off
 	const static TMap<FName, EInterchangePropertyTracks> PropertyNameToTrackType = {
 		// Common properties
-		{UnrealIdentifiers::HiddenInGamePropertyName, 			EInterchangePropertyTracks::Visibility}, // Binding visibility to the actor works better for cameras
+		{UnrealIdentifiers::HiddenInGamePropertyName, 			EInterchangePropertyTracks::ActorHiddenInGame}, // Binding visibility to the actor works better for cameras
 
 		// Camera properties
 		// TODO: Need to add support for SensorHorizontalOffset and SensorVerticalOffset, once Interchange supports those
diff --git a/Engine/Plugins/Interchange/Runtime/Source/Nodes/Private/InterchangeSceneNode.cpp b/Engine/Plugins/Interchange/Runtime/Source/Nodes/Private/InterchangeSceneNode.cpp
--- a/Engine/Plugins/Interchange/Runtime/Source/Nodes/Private/InterchangeSceneNode.cpp
+++ b/Engine/Plugins/Interchange/Runtime/Source/Nodes/Private/InterchangeSceneNode.cpp
@@ -289,6 +289,26 @@ bool UInterchangeSceneNode::SetCustomPivotNodeTransform(const FTransform& Attrib
 	IMPLEMENT_NODE_ATTRIBUTE_SETTER_NODELEGATE(PivotNodeTransform, FTransform);
 }
 
+bool UInterchangeSceneNode::GetCustomComponentVisibility(bool& AttributeValue) const
+{
+	IMPLEMENT_NODE_ATTRIBUTE_GETTER(ComponentVisibility, bool);
+}
+
+bool UInterchangeSceneNode::SetCustomComponentVisibility(bool AttributeValue)
+{
+	IMPLEMENT_NODE_ATTRIBUTE_SETTER_NODELEGATE(ComponentVisibility, bool);
+}
+
+bool UInterchangeSceneNode::GetCustomActorVisibility(bool& AttributeValue) const
+{
+	IMPLEMENT_NODE_ATTRIBUTE_GETTER(ActorVisibility, bool);
+}
+
+bool UInterchangeSceneNode::SetCustomActorVisibility(bool AttributeValue)
+{
+	IMPLEMENT_NODE_ATTRIBUTE_SETTER_NODELEGATE(ActorVisibility, bool);
+}
+
 bool UInterchangeSceneNode::GetCustomAssetInstanceUid(FString& AttributeValue) const
 {
 	IMPLEMENT_NODE_ATTRIBUTE_GETTER(AssetInstanceUid, FString);
@@ -490,4 +510,4 @@ bool UInterchangeSceneNode::AddTag(const FString& Tag)
 bool UInterchangeSceneNode::RemoveTag(const FString& Tag)
 {
 	return Tags.RemoveItem(Tag);
-}
\ No newline at end of file
+}
diff --git a/Engine/Plugins/Interchange/Runtime/Source/Nodes/Public/InterchangeAnimationDefinitions.h b/Engine/Plugins/Interchange/Runtime/Source/Nodes/Public/InterchangeAnimationDefinitions.h
--- a/Engine/Plugins/Interchange/Runtime/Source/Nodes/Public/InterchangeAnimationDefinitions.h
+++ b/Engine/Plugins/Interchange/Runtime/Source/Nodes/Public/InterchangeAnimationDefinitions.h
@@ -49,7 +49,7 @@ enum class EInterchangePropertyTracks : int32
 	TranslucencySortDistanceOffset,
 	VisibleInRayTracing,
 	VisibleInSceneCaptureOnly,
-	Visibility, // This visibility property is set on the Actor
+	ActorHiddenInGame, // This visibility property is set on the Actor
 
 	/** Light Properties */
 	LightAffectGlobalIllumination,
@@ -317,7 +317,7 @@ enum class EInterchangePropertyTracks : int32
 	CameraUseFieldOfViewForLOD,
 
 	/** Common Mesh Properties */
-	MeshOverlayMaterialMaxDrawDistance, 
+	MeshOverlayMaterialMaxDrawDistance,
 
 	/** Skinned Mesh Properties */
 	SkinnedMeshCapsuleIndirectShadowMinVisibility,
@@ -344,6 +344,7 @@ enum class EInterchangePropertyTracks : int32
 	StaticMesh,
 
 	None = -1 UMETA(hidden),
+	Visibility UE_DEPRECATED(5.6, "Renamed to 'ActorHiddenInGame'") = ActorHiddenInGame UMETA(hidden)
 };
 
 namespace UE
diff --git a/Engine/Plugins/Interchange/Runtime/Source/Nodes/Public/InterchangeSceneNode.h b/Engine/Plugins/Interchange/Runtime/Source/Nodes/Public/InterchangeSceneNode.h
--- a/Engine/Plugins/Interchange/Runtime/Source/Nodes/Public/InterchangeSceneNode.h
+++ b/Engine/Plugins/Interchange/Runtime/Source/Nodes/Public/InterchangeSceneNode.h
@@ -135,12 +135,28 @@ public:
 	UFUNCTION(BlueprintCallable, Category = "Interchange | Node | Scene")
 	bool SetCustomPivotNodeTransform(const FTransform& AttributeValue);
 
+	/** Gets whether components spawned from this node should be visible */
+	UFUNCTION(BlueprintCallable, Category = "Interchange | Node | Scene")
+	bool GetCustomComponentVisibility(bool& bOutIsVisible) const;
+
+	/** Sets whether components spawned from this node should be visible */
+	UFUNCTION(BlueprintCallable, Category = "Interchange | Node | Scene")
+	bool SetCustomComponentVisibility(bool bInIsVisible);
+
+	/** Gets whether actors spawned from this node should be visible */
+	UFUNCTION(BlueprintCallable, Category = "Interchange | Node | Scene")
+	bool GetCustomActorVisibility(bool& bOutIsVisible) const;
+
+	/** Sets whether actors spawned from this node should be visible */
+	UFUNCTION(BlueprintCallable, Category = "Interchange | Node | Scene")
+	bool SetCustomActorVisibility(bool bInIsVisible);
+
 	/***********************************************************************************************
 	* Skeleton bind bone API Begin
-	* 
+	*
 	* Bind pose transform is the transform of the joint when the binding with the mesh was done.
 	* This attribute should be set only if this scene node represent a joint.
-	* 
+	*
 	* Time-zero transform is the transform of the node at time zero.
 	* Pipelines often have the option to evaluate the joint at time zero to create the bind pose.
 	* Time-zero bind pose is also used if the translator did not find any bind pose, or if we import
@@ -270,19 +286,23 @@ private:
 
 	//Scene node default local transforms.
 	const UE::Interchange::FAttributeKey Macro_CustomLocalTransformKey = UE::Interchange::FAttributeKey(TEXT("LocalTransform"));
-	
+
 	//Scene node Local bind pose transforms (the specialize type should be set to joint)
 	const UE::Interchange::FAttributeKey Macro_CustomBindPoseLocalTransformKey = UE::Interchange::FAttributeKey(TEXT("BindPoseLocalTransform"));
-	
+
 	//Scene node local transforms at time zero. This attribute is important for rigid mesh import or if the translator did not fill the bind pose.
 	const UE::Interchange::FAttributeKey Macro_CustomTimeZeroLocalTransformKey = UE::Interchange::FAttributeKey(TEXT("TimeZeroLocalTransform"));
-	
+
 	//A scene node can have a transform apply to the mesh it reference.
 	const UE::Interchange::FAttributeKey Macro_CustomGeometricTransformKey = UE::Interchange::FAttributeKey(TEXT("GeometricTransform"));
 
 	//A scene node can have a pivot transform apply to the mesh it reference (use this pivot only if you are not baking the vertices of the mesh).
 	const UE::Interchange::FAttributeKey Macro_CustomPivotNodeTransformKey = UE::Interchange::FAttributeKey(TEXT("PivotNodeTransform"));
 
+	//A scene node can be invisible, but still be imported. We have two of these to match how you can separately animate actor and component visibility
+	IMPLEMENT_NODE_ATTRIBUTE_KEY(ComponentVisibility);
+	IMPLEMENT_NODE_ATTRIBUTE_KEY(ActorVisibility);
+
 	//A scene node can reference an asset. Asset can be Mesh, Light, camera...
 	const UE::Interchange::FAttributeKey Macro_CustomAssetInstanceUidKey = UE::Interchange::FAttributeKey(TEXT("AssetInstanceUid"));
 
diff --git a/Engine/Plugins/Interchange/Runtime/Source/Parsers/Fbx/Private/InterchangeFbxSettings.cpp b/Engine/Plugins/Interchange/Runtime/Source/Parsers/Fbx/Private/InterchangeFbxSettings.cpp
--- a/Engine/Plugins/Interchange/Runtime/Source/Parsers/Fbx/Private/InterchangeFbxSettings.cpp
+++ b/Engine/Plugins/Interchange/Runtime/Source/Parsers/Fbx/Private/InterchangeFbxSettings.cpp
@@ -8,7 +8,7 @@ UInterchangeFbxSettings::UInterchangeFbxSettings()
 	{
 		PredefinedPropertyTracks =
 		{
-			{ TEXT("bHidden"), {EInterchangePropertyTracks::Visibility} },
+			{ TEXT("bHidden"), {EInterchangePropertyTracks::ActorHiddenInGame} },
 			{ TEXT("bAutoActivate"), {EInterchangePropertyTracks::AutoActivate} },
 
 			{ TEXT("Intensity"), { EInterchangePropertyTracks::LightIntensity}	},
diff --git a/Engine/Plugins/Interchange/Runtime/Source/Pipelines/Private/InterchangeGenericScenesPipeline.cpp b/Engine/Plugins/Interchange/Runtime/Source/Pipelines/Private/InterchangeGenericScenesPipeline.cpp
--- a/Engine/Plugins/Interchange/Runtime/Source/Pipelines/Private/InterchangeGenericScenesPipeline.cpp
+++ b/Engine/Plugins/Interchange/Runtime/Source/Pipelines/Private/InterchangeGenericScenesPipeline.cpp
@@ -712,6 +712,16 @@ void UInterchangeGenericLevelPipeline::ExecuteSceneNodePreImport(const FTransfor
 		}
 	}
 
+	bool bCustomVisibility = true;
+	if (SceneNode->GetCustomComponentVisibility(bCustomVisibility))
+	{
+		ActorFactoryNode->SetCustomComponentVisibility(bCustomVisibility);
+	}
+	if (SceneNode->GetCustomActorVisibility(bCustomVisibility))
+	{
+		ActorFactoryNode->SetCustomActorVisibility(bCustomVisibility);
+	}
+
 	ActorFactoryNode->SetCustomMobility(EComponentMobility::Static);
 
 	if (TranslatedAssetNode)
