USD: Import USD prim visibility onto actors/components on the level when importing via USD Interchange.

diff --git a/Engine/Plugins/Interchange/Runtime/Source/Import/Private/Usd/InterchangeUsdTranslator.cpp b/Engine/Plugins/Interchange/Runtime/Source/Import/Private/Usd/InterchangeUsdTranslator.cpp
--- a/Engine/Plugins/Interchange/Runtime/Source/Import/Private/Usd/InterchangeUsdTranslator.cpp
+++ b/Engine/Plugins/Interchange/Runtime/Source/Import/Private/Usd/InterchangeUsdTranslator.cpp
@@ -98,6 +98,7 @@ namespace UE::InterchangeUsdTranslator::Private
 	struct FTraversalInfo
 	{
 		UInterchangeBaseNode* ParentNode = nullptr;
+		bool bVisible = true;
 
 		TSharedPtr<UE::FUsdSkelCache> FurthestSkelCache;
 		UE::FUsdPrim ClosestParentSkelRoot;
@@ -1963,6 +1964,20 @@ namespace UE::InterchangeUsdTranslator::Private
 				SceneNode->SetCustomLocalTransform(&NodeContainer, Transform);
 			}
 
+			// Hide our scene node if it is meant to be hidden in USD
+			//
+			// We use actor visibility so that it matches how we map visibility timeSamples to actor visibility tracks.
+			// We do *that*, because it matches how Interchange always puts transform animations on the actors directly (so
+			// "scene component stuff" ends up as actor tracks), and also due to how it behaves better for cameras: Component
+			// visibility for camera nodes would hide the camera component itself, which has no effect. Actor visibility for
+			// camera actors does hide the entire camera actor however
+			const bool bComputedVisibility = Info.bVisible && UsdUtils::HasInheritedVisibility(Prim);
+			if (!bComputedVisibility)
+			{
+				SceneNode->SetCustomActorVisibility(bComputedVisibility);
+				Info.bVisible = bComputedVisibility;
+			}
+
 			// Skeleton joints are separate scene nodes in Interchange, so we need to emit that node hierarchy now
 			if (Prim.IsA(TEXT("Skeleton")))
 			{
