Interchange: Fix copy paste error when handling actor factory node local transforms.

diff --git a/Engine/Plugins/Interchange/Runtime/Source/FactoryNodes/Private/InterchangeActorFactoryNode.cpp b/Engine/Plugins/Interchange/Runtime/Source/FactoryNodes/Private/InterchangeActorFactoryNode.cpp
--- a/Engine/Plugins/Interchange/Runtime/Source/FactoryNodes/Private/InterchangeActorFactoryNode.cpp
+++ b/Engine/Plugins/Interchange/Runtime/Source/FactoryNodes/Private/InterchangeActorFactoryNode.cpp
@@ -35,27 +35,17 @@ bool UInterchangeActorFactoryNode::GetCustomGlobalTransform(FTransform& Attribut
 
 bool UInterchangeActorFactoryNode::SetCustomGlobalTransform(const FTransform& AttributeValue, bool bAddApplyDelegate)
 {
-	IMPLEMENT_NODE_ATTRIBUTE_SETTER_WITH_CUSTOM_DELEGATE_WITH_CUSTOM_CLASS(
-		UInterchangeActorFactoryNode,
-		GlobalTransform,
-		FTransform,
-		USceneComponent
-	);
+	IMPLEMENT_NODE_ATTRIBUTE_SETTER_WITH_CUSTOM_DELEGATE_WITH_CUSTOM_CLASS(UInterchangeActorFactoryNode, GlobalTransform, FTransform, USceneComponent);
 }
 
 bool UInterchangeActorFactoryNode::GetCustomLocalTransform(FTransform& AttributeValue) const
 {
-	IMPLEMENT_NODE_ATTRIBUTE_GETTER(GlobalTransform, FTransform);
+	IMPLEMENT_NODE_ATTRIBUTE_GETTER(LocalTransform, FTransform);
 }
 
 bool UInterchangeActorFactoryNode::SetCustomLocalTransform(const FTransform& AttributeValue, bool bAddApplyDelegate)
 {
-	IMPLEMENT_NODE_ATTRIBUTE_SETTER_WITH_CUSTOM_DELEGATE_WITH_CUSTOM_CLASS(
-		UInterchangeActorFactoryNode,
-		GlobalTransform,
-		FTransform,
-		USceneComponent
-	);
+	IMPLEMENT_NODE_ATTRIBUTE_SETTER_WITH_CUSTOM_DELEGATE_WITH_CUSTOM_CLASS(UInterchangeActorFactoryNode, LocalTransform, FTransform, USceneComponent);
 }
 
 bool UInterchangeActorFactoryNode::GetCustomComponentVisibility(bool& AttributeValue) const
@@ -78,24 +68,60 @@ bool UInterchangeActorFactoryNode::SetCustomActorVisibility(bool AttributeValue,
 	IMPLEMENT_NODE_ATTRIBUTE_SETTER_WITH_CUSTOM_DELEGATE_WITH_CUSTOM_CLASS(UInterchangeActorFactoryNode, ActorVisibility, bool, USceneComponent);
 }
 
+bool UInterchangeActorFactoryNode::GetCustomActorClassName(FString& AttributeValue) const
+{
+	IMPLEMENT_NODE_ATTRIBUTE_GETTER(ActorClassName, FString);
+}
+
+bool UInterchangeActorFactoryNode::SetCustomActorClassName(const FString& AttributeValue)
+{
+	IMPLEMENT_NODE_ATTRIBUTE_SETTER_NODELEGATE(ActorClassName, FString);
+}
+
+bool UInterchangeActorFactoryNode::GetCustomMobility(uint8& AttributeValue) const
+{
+	IMPLEMENT_NODE_ATTRIBUTE_GETTER(Mobility, uint8);
+}
+
+bool UInterchangeActorFactoryNode::SetCustomMobility(const uint8& AttributeValue, bool bAddApplyDelegate)
+{
+	IMPLEMENT_NODE_ATTRIBUTE_SETTER_NODELEGATE(Mobility, uint8);
+}
+
 bool UInterchangeActorFactoryNode::ApplyCustomGlobalTransformToAsset(UObject* Asset) const
 {
-	FTransform LocalTransform;
-	if (GetCustomLocalTransform(LocalTransform))
+	FTransform GlobalTransform;
+	if (GetCustomGlobalTransform(GlobalTransform))
 	{
 		if (USceneComponent* Component = Cast<USceneComponent>(Asset))
 		{
-			Component->SetRelativeTransform(LocalTransform);
+			Component->SetWorldTransform(GlobalTransform);
 			return true;
 		}
 	}
 
-	FTransform GlobalTransform;
-	if (GetCustomGlobalTransform(GlobalTransform))
+	return false;
+}
+
+bool UInterchangeActorFactoryNode::FillCustomGlobalTransformFromAsset(UObject* Asset)
+{
+	if (const USceneComponent* Component = Cast<USceneComponent>(Asset))
+	{
+		FTransform GlobalTransform = Component->GetComponentToWorld();
+		return SetCustomGlobalTransform(GlobalTransform, false);
+	}
+
+	return false;
+}
+
+bool UInterchangeActorFactoryNode::ApplyCustomLocalTransformToAsset(UObject* Asset) const
+{
+	FTransform LocalTransform;
+	if (GetCustomLocalTransform(LocalTransform))
 	{
 		if (USceneComponent* Component = Cast<USceneComponent>(Asset))
 		{
-			Component->SetWorldTransform(GlobalTransform);
+			Component->SetRelativeTransform(LocalTransform);
 			return true;
 		}
 	}
@@ -103,14 +129,12 @@ bool UInterchangeActorFactoryNode::ApplyCustomGlobalTransformToAsset(UObject* As
 	return false;
 }
 
-bool UInterchangeActorFactoryNode::FillCustomGlobalTransformFromAsset(UObject* Asset)
+bool UInterchangeActorFactoryNode::FillCustomLocalTransformFromAsset(UObject* Asset)
 {
 	if (const USceneComponent* Component = Cast<USceneComponent>(Asset))
 	{
 		FTransform LocalTransform = Component->GetRelativeTransform();
-		FTransform GlobalTransform = Component->GetComponentToWorld();
-
-		return this->SetCustomLocalTransform(LocalTransform, false) || this->SetCustomGlobalTransform(GlobalTransform, false);
+		return this->SetCustomLocalTransform(LocalTransform, false);
 	}
 
 	return false;
@@ -205,27 +229,10 @@ void UInterchangeActorFactoryNode::CopyWithObject(const UInterchangeFactoryBaseN
 
 	if (const UInterchangeActorFactoryNode* ActorFactoryNode = Cast<UInterchangeActorFactoryNode>(SourceNode))
 	{
-		COPY_NODE_DELEGATES_WITH_CUSTOM_DELEGATE(
-			ActorFactoryNode,
-			UInterchangeActorFactoryNode,
-			GlobalTransform,
-			FTransform,
-			USceneComponent::StaticClass()
-		)
-		COPY_NODE_DELEGATES_WITH_CUSTOM_DELEGATE(
-			ActorFactoryNode,
-			UInterchangeActorFactoryNode,
-			ComponentVisibility,
-			bool,
-			USceneComponent::StaticClass()
-		)
-		COPY_NODE_DELEGATES_WITH_CUSTOM_DELEGATE(
-			ActorFactoryNode,
-			UInterchangeActorFactoryNode,
-			ActorVisibility,
-			bool,
-			USceneComponent::StaticClass()
-		)
+		COPY_NODE_DELEGATES_WITH_CUSTOM_DELEGATE(ActorFactoryNode, UInterchangeActorFactoryNode, GlobalTransform, FTransform, USceneComponent::StaticClass())
+		COPY_NODE_DELEGATES_WITH_CUSTOM_DELEGATE(ActorFactoryNode, UInterchangeActorFactoryNode, LocalTransform, FTransform, USceneComponent::StaticClass())
+		COPY_NODE_DELEGATES_WITH_CUSTOM_DELEGATE(ActorFactoryNode, UInterchangeActorFactoryNode, ComponentVisibility, bool, USceneComponent::StaticClass())
+		COPY_NODE_DELEGATES_WITH_CUSTOM_DELEGATE(ActorFactoryNode, UInterchangeActorFactoryNode, ActorVisibility, bool, USceneComponent::StaticClass())
 	}
 }
 
diff --git a/Engine/Plugins/Interchange/Runtime/Source/FactoryNodes/Public/InterchangeActorFactoryNode.h b/Engine/Plugins/Interchange/Runtime/Source/FactoryNodes/Public/InterchangeActorFactoryNode.h
--- a/Engine/Plugins/Interchange/Runtime/Source/FactoryNodes/Public/InterchangeActorFactoryNode.h
+++ b/Engine/Plugins/Interchange/Runtime/Source/FactoryNodes/Public/InterchangeActorFactoryNode.h
@@ -41,28 +41,16 @@ public:
 	bool SetCustomActorVisibility(bool AttributeValue, bool bAddApplyDelegate = true);
 
 	UFUNCTION(BlueprintCallable, Category = "Interchange | Node | ActorFactory")
-	bool GetCustomActorClassName(FString& AttributeValue) const
-	{
-		IMPLEMENT_NODE_ATTRIBUTE_GETTER(ActorClassName, FString);
-	}
+	bool GetCustomActorClassName(FString& AttributeValue) const;
 
 	UFUNCTION(BlueprintCallable, Category = "Interchange | Node | ActorFactory")
-	bool SetCustomActorClassName(const FString& AttributeValue)
-	{
-		IMPLEMENT_NODE_ATTRIBUTE_SETTER_NODELEGATE(ActorClassName, FString);
-	}
+	bool SetCustomActorClassName(const FString& AttributeValue);
 
 	UFUNCTION(BlueprintCallable, Category = "Interchange | Node | ActorFactory")
-	bool GetCustomMobility(uint8& AttributeValue) const
-	{
-		IMPLEMENT_NODE_ATTRIBUTE_GETTER(Mobility, uint8);
-	}
+	bool GetCustomMobility(uint8& AttributeValue) const;
 
 	UFUNCTION(BlueprintCallable, Category = "Interchange | Node | ActorFactory")
-	bool SetCustomMobility(const uint8& AttributeValue, bool bAddApplyDelegate = true)
-	{
-		IMPLEMENT_NODE_ATTRIBUTE_SETTER_NODELEGATE(Mobility, uint8);
-	}
+	bool SetCustomMobility(const uint8& AttributeValue, bool bAddApplyDelegate = true);
 
 	virtual void CopyWithObject(const UInterchangeFactoryBaseNode* SourceNode, UObject* Object) override;
 
@@ -102,6 +90,9 @@ private:
 	bool ApplyCustomComponentVisibilityToAsset(UObject* Asset) const;
 	bool FillCustomComponentVisibilityFromAsset(UObject* Asset);
 
+	bool ApplyCustomLocalTransformToAsset(UObject* Asset) const;
+	bool FillCustomLocalTransformFromAsset(UObject* Asset);
+
 	bool ApplyCustomActorVisibilityToAsset(UObject* Asset) const;
 	bool FillCustomActorVisibilityFromAsset(UObject* Asset);
 
