[Interchange] Fixing Rigid body animation reimports.

#rnx

diff --git a/Engine/Plugins/Interchange/Runtime/Source/Pipelines/Private/InterchangeGenericAnimationPipeline.cpp b/Engine/Plugins/Interchange/Runtime/Source/Pipelines/Private/InterchangeGenericAnimationPipeline.cpp
--- a/Engine/Plugins/Interchange/Runtime/Source/Pipelines/Private/InterchangeGenericAnimationPipeline.cpp
+++ b/Engine/Plugins/Interchange/Runtime/Source/Pipelines/Private/InterchangeGenericAnimationPipeline.cpp
@@ -30,6 +30,8 @@
 
 namespace UE::Interchange::Private
 {
+	const FString ConvertedFromRigidAnimationPrefixIdentifier = TEXT("\\SkeletalAnimation\\ConvertedFromRigidAnimation\\");
+
 	bool IsTranslatedDataContainOnlyJointAnimation(const UInterchangeBaseNodeContainer* InBaseNodeContainer, bool bConvertStaticsWithMorphTargetsToSkeletals)
 	{
 		//Its valid to call GetMeshesInformationFromTranslatedData with a null container
@@ -305,7 +307,7 @@ void UInterchangeGenericAnimationPipeline::ExecutePipeline(UInterchangeBaseNodeC
 		for (UInterchangeSceneNode* SceneNode : SceneNodesWithMorphTargetCurveWeights)
 		{
 			UInterchangeSkeletalAnimationTrackNode* SkeletalAnimationNode = NewObject< UInterchangeSkeletalAnimationTrackNode >(BaseNodeContainer);
-			FString SkeletalAnimationNodeUid = "\\SkeletalAnimation\\MorphTargetCurveWeightInstantation\\" + SceneNode->GetUniqueID();
+			FString SkeletalAnimationNodeUid = UE::Interchange::Private::ConvertedFromRigidAnimationPrefixIdentifier + SceneNode->GetUniqueID();
 			SkeletalAnimationNode->InitializeNode(SkeletalAnimationNodeUid, SceneNode->GetDisplayLabel(), EInterchangeNodeContainerType::TranslatedAsset);
 
 			SkeletalAnimationNode->SetCustomAnimationSampleRate(30.f);
@@ -580,15 +582,18 @@ void UInterchangeGenericAnimationPipeline::CreateAnimSequenceFactoryNode(UInterc
 	const FString SkeletonFactoryNodeUid = UInterchangeFactoryBaseNode::BuildFactoryNodeUid(SkeletonNodeUid);
 	SkeletonFactoryNode = Cast<const UInterchangeSkeletonFactoryNode>(BaseNodeContainer->GetFactoryNode(SkeletonFactoryNodeUid));
 
+	const bool bRigidAnimationConverted = TrackNode.GetUniqueID().StartsWith(UE::Interchange::Private::ConvertedFromRigidAnimationPrefixIdentifier);
+
 	//If we import anim only and we do not have meshes and skeleton. We need to create a skeleton factory node
 	//base on the specified skeleton
 	if(bImportOnlyAnimation && !SkeletonFactoryNode && CommonSkeletalMeshesAndAnimationsProperties->Skeleton.IsValid())
 	{
 		TWeakObjectPtr<USkeleton> Skeleton = CommonSkeletalMeshesAndAnimationsProperties->Skeleton;
 		TPair<int32, FString> SkeletonRootNodeUidAndBoneIndex = TPair<int32, FString>(INDEX_NONE, FString());
-		BaseNodeContainer->IterateNodesOfType<UInterchangeSceneNode>([&SkeletonRootNodeUidAndBoneIndex, BaseNodeContainerClosure = BaseNodeContainer, Skeleton](const FString& NodeUid, UInterchangeSceneNode* Node)
+		BaseNodeContainer->IterateNodesOfType<UInterchangeSceneNode>([&SkeletonRootNodeUidAndBoneIndex, BaseNodeContainerClosure = BaseNodeContainer, Skeleton, bRigidAnimationConverted](const FString& NodeUid, UInterchangeSceneNode* Node)
 		{
-			if (Skeleton.IsValid() && Node->IsSpecializedTypeContains(UE::Interchange::FSceneNodeStaticData::GetJointSpecializeTypeString()))
+			//In case the AnimationTrackNode is created from a RigidAnimation conversion the SceneNodes won't have Joint Specialization:
+			if (Skeleton.IsValid() && (Node->IsSpecializedTypeContains(UE::Interchange::FSceneNodeStaticData::GetJointSpecializeTypeString()) || bRigidAnimationConverted))
 			{
 				const FReferenceSkeleton& ReferenceSkeleton = Skeleton->GetReferenceSkeleton();
 				int32 RefBoneIndex = ReferenceSkeleton.FindBoneIndex(FName(*Node->GetDisplayLabel()));
