[Interchange]: Make sure fbx bone animation and morph target animation on the same fbx animtrack index are on the same anim sequence.

diff --git a/Engine/Plugins/Interchange/Runtime/Source/Parsers/Fbx/Private/FbxScene.cpp b/Engine/Plugins/Interchange/Runtime/Source/Parsers/Fbx/Private/FbxScene.cpp
--- a/Engine/Plugins/Interchange/Runtime/Source/Parsers/Fbx/Private/FbxScene.cpp
+++ b/Engine/Plugins/Interchange/Runtime/Source/Parsers/Fbx/Private/FbxScene.cpp
@@ -879,7 +879,6 @@ namespace UE
 							if(!SkeletalAnimationTrackNode)
 							{
 								SkeletalAnimationTrackNode = NewObject< UInterchangeSkeletalAnimationTrackNode >(&NodeContainer);
-								TrackNodeUid = "\\SkeletalAnimation\\" + UnrealNode->GetUniqueID() + "_" + FString::FromInt(AnimationIndex);
 								SkeletalAnimationTrackNode->InitializeNode(TrackNodeUid, DisplayString, EInterchangeNodeContainerType::TranslatedAsset);
 							}
 							SkeletalAnimationTrackNode->AddBooleanAttribute(TEXT("RenameLikeLegacyFbx"), true);
