[Interchange] Fix fbx null cluster link preventing import of skeletal mesh

diff --git a/Engine/Plugins/Interchange/Runtime/Source/Parsers/Fbx/Private/FbxMesh.cpp b/Engine/Plugins/Interchange/Runtime/Source/Parsers/Fbx/Private/FbxMesh.cpp
--- a/Engine/Plugins/Interchange/Runtime/Source/Parsers/Fbx/Private/FbxMesh.cpp
+++ b/Engine/Plugins/Interchange/Runtime/Source/Parsers/Fbx/Private/FbxMesh.cpp
@@ -2202,6 +2202,10 @@ bool FFbxMesh::ExtractSkinnedMeshNodeJoints(FbxScene* SDKScene, UInterchangeBase
 				continue;
 			}
 			FbxNode* Link = Cluster->GetLink();
+			if (!Link)
+			{
+				continue;
+			}
 			bFoundValidJoint = true;
 
 			FString JointNodeUniqueID = Parser.GetFbxHelper()->GetFbxNodeHierarchyName(Link);
