[Interchange] Remove unreachable code

#rnx

diff --git a/Engine/Plugins/Interchange/Runtime/Source/Pipelines/Private/InterchangeGenericStaticMeshPipeline.cpp b/Engine/Plugins/Interchange/Runtime/Source/Pipelines/Private/InterchangeGenericStaticMeshPipeline.cpp
--- a/Engine/Plugins/Interchange/Runtime/Source/Pipelines/Private/InterchangeGenericStaticMeshPipeline.cpp
+++ b/Engine/Plugins/Interchange/Runtime/Source/Pipelines/Private/InterchangeGenericStaticMeshPipeline.cpp
@@ -289,7 +289,6 @@ void UInterchangeGenericMeshPipeline::ExecutePreImportPipelineStaticMesh()
 							continue;
 						}
 						
-						TArray<FString> ReferencingMeshInstanceUids; //This data is use only when bInstancedMesh is false
 						TMap<int32, TArray<FString>> MeshUidsPerLodIndex;
 						if (bInstancedMesh)
 						{
@@ -313,7 +312,6 @@ void UInterchangeGenericMeshPipeline::ExecutePreImportPipelineStaticMesh()
 							const FInterchangeMeshGeometry& MeshGeometry = PipelineMeshesUtilities->GetMeshGeometryByUid(MeshUid);
 							TArray<FString>& TranslatedNodes = MeshUidsPerLodIndex.FindOrAdd(LodIndexZero);
 							TranslatedNodes.Add(MeshGeometry.MeshUid);
-							ReferencingMeshInstanceUids = MeshGeometry.ReferencingMeshInstanceUids;
 						}
 
 						if (MeshUidsPerLodIndex.Num() > 0)
@@ -326,7 +324,7 @@ void UInterchangeGenericMeshPipeline::ExecutePreImportPipelineStaticMesh()
 								}
 							}
 
-							UInterchangeStaticMeshFactoryNode* StaticMeshFactoryNode = CreateStaticMeshFactoryNode(MeshUidsPerLodIndex, ReferencingMeshInstanceUids);
+							UInterchangeStaticMeshFactoryNode* StaticMeshFactoryNode = CreateStaticMeshFactoryNode(MeshUidsPerLodIndex);
 							StaticMeshFactoryNodes.Add(StaticMeshFactoryNode);
 							bFoundMeshes = true;
 						}
@@ -402,7 +400,7 @@ bool UInterchangeGenericMeshPipeline::MakeMeshFactoryNodeUidAndDisplayLabel(cons
 	return false;
 }
 
-UInterchangeStaticMeshFactoryNode* UInterchangeGenericMeshPipeline::CreateStaticMeshFactoryNode(const TMap<int32, TArray<FString>>& MeshUidsPerLodIndex, const TArray<FString>& ReferencingMeshInstanceUids)
+UInterchangeStaticMeshFactoryNode* UInterchangeGenericMeshPipeline::CreateStaticMeshFactoryNode(const TMap<int32, TArray<FString>>& MeshUidsPerLodIndex)
 {
 	check(CommonMeshesProperties.IsValid());
 	if (MeshUidsPerLodIndex.Num() == 0)
@@ -484,20 +482,6 @@ UInterchangeStaticMeshFactoryNode* UInterchangeGenericMeshPipeline::CreateStatic
 	StaticMeshFactoryNode->SetCustomAutoComputeLODScreenSizes(bAutoComputeLODScreenSizes);
 	StaticMeshFactoryNode->SetLODScreenSizes(LODScreenSizes);
 	
-
-	if (ReferencingMeshInstanceUids.Num() > 0)
-	{
-		constexpr bool bAddSourceNodeName = true;
-
-		for (const FString& ReferencingMeshInstanceUid : ReferencingMeshInstanceUids)
-		{
-			if (const UInterchangeSceneNode* SceneNode = Cast<UInterchangeSceneNode>(BaseNodeContainer->GetNode(ReferencingMeshInstanceUid)))
-			{
-				UInterchangeUserDefinedAttributesAPI::DuplicateAllUserDefinedAttribute(SceneNode, StaticMeshFactoryNode, bAddSourceNodeName);
-			}
-		}
-	}
-
 	return StaticMeshFactoryNode;
 }
 
diff --git a/Engine/Plugins/Interchange/Runtime/Source/Pipelines/Public/InterchangeGenericMeshPipeline.h b/Engine/Plugins/Interchange/Runtime/Source/Pipelines/Public/InterchangeGenericMeshPipeline.h
--- a/Engine/Plugins/Interchange/Runtime/Source/Pipelines/Public/InterchangeGenericMeshPipeline.h
+++ b/Engine/Plugins/Interchange/Runtime/Source/Pipelines/Public/InterchangeGenericMeshPipeline.h
@@ -352,7 +352,7 @@ private:
 	 * This function can create a UInterchangeStaticMeshFactoryNode
 	 * @param MeshUidsPerLodIndex - The MeshUids can represent a SceneNode pointing on a MeshNode or directly a MeshNode
 	 */
-	UInterchangeStaticMeshFactoryNode* CreateStaticMeshFactoryNode(const TMap<int32, TArray<FString>>& MeshUidsPerLodIndex, const TArray<FString>& ReferencingMeshInstanceUids = TArray<FString>());
+	UInterchangeStaticMeshFactoryNode* CreateStaticMeshFactoryNode(const TMap<int32, TArray<FString>>& MeshUidsPerLodIndex);
 
 	/** This function can create a UInterchangeStaticMeshLodDataNode which represents the LOD data needed by the factory to create a lod mesh */
 	UInterchangeStaticMeshLodDataNode* CreateStaticMeshLodDataNode(const FString& NodeName, const FString& NodeUniqueID);
