[Interchange] Animation conversion PrefixIdentifier consistency fix.

#rnx

diff --git a/Engine/Plugins/Interchange/Runtime/Source/Pipelines/Private/InterchangeGenericAnimationPipeline.cpp b/Engine/Plugins/Interchange/Runtime/Source/Pipelines/Private/InterchangeGenericAnimationPipeline.cpp
--- a/Engine/Plugins/Interchange/Runtime/Source/Pipelines/Private/InterchangeGenericAnimationPipeline.cpp
+++ b/Engine/Plugins/Interchange/Runtime/Source/Pipelines/Private/InterchangeGenericAnimationPipeline.cpp
@@ -394,7 +394,7 @@ void UInterchangeGenericAnimationPipeline::ExecutePipeline(UInterchangeBaseNodeC
 			}
 
 			UInterchangeSkeletalAnimationTrackNode* SkeletalAnimationNode = NewObject< UInterchangeSkeletalAnimationTrackNode >(BaseNodeContainer);
-			FString SkeletalAnimationNodeUid = "\\SkeletalAnimation\\ConvertedFromRigidAnimation\\" + TrackSetNode->GetUniqueID();
+			FString SkeletalAnimationNodeUid = UE::Interchange::Private::ConvertedFromRigidAnimationPrefixIdentifier + TrackSetNode->GetUniqueID();
 			SkeletalAnimationNode->InitializeNode(SkeletalAnimationNodeUid, TrackSetNode->GetDisplayLabel(), EInterchangeNodeContainerType::TranslatedAsset);
 
 			bool bCustomSkeletonNodeUidSet = false;
