[Interchange - USD] Implement support for primitive shapes.

#rnx

diff --git a/Engine/Plugins/Interchange/Runtime/Source/Import/Private/Usd/InterchangeUsdTranslator.cpp b/Engine/Plugins/Interchange/Runtime/Source/Import/Private/Usd/InterchangeUsdTranslator.cpp
--- a/Engine/Plugins/Interchange/Runtime/Source/Import/Private/Usd/InterchangeUsdTranslator.cpp
+++ b/Engine/Plugins/Interchange/Runtime/Source/Import/Private/Usd/InterchangeUsdTranslator.cpp
@@ -91,6 +91,7 @@ namespace UE::InterchangeUsdTranslator::Private
 	const static FString LightPrefix = TEXT("\\Light\\");
 	const static FString MaterialPrefix = TEXT("\\Material\\");
 	const static FString MeshPrefix = TEXT("\\Mesh\\");
+	const static FString PrimitiveShapePrefix = TEXT("\\PrimitiveShape\\");
 	const static FString MorphTargetPrefix = TEXT("\\MorphTarget\\");
 	const static FString BonePrefix = TEXT("\\Bone\\");
 
@@ -156,7 +157,7 @@ namespace UE::InterchangeUsdTranslator::Private
 			UInterchangeBaseNodeContainer& NodeContainer
 		);
 
-		void AddMeshNode(const UE::FUsdPrim& Prim, UInterchangeBaseNodeContainer& NodeContainer, const FTraversalInfo& Info);
+		void AddMeshNode(const UE::FUsdPrim& Prim, UInterchangeBaseNodeContainer& NodeContainer, const FTraversalInfo& Info, bool bPrimitiveShape = false);
 
 		/**
 		 * Search for a MaterialX file embedded in the file, in case there's none, create a ShaderGraph using MaterialX
@@ -1309,11 +1310,12 @@ namespace UE::InterchangeUsdTranslator::Private
 	void UInterchangeUSDTranslatorImpl::AddMeshNode(
 		const UE::FUsdPrim& Prim,
 		UInterchangeBaseNodeContainer& NodeContainer,
-		const FTraversalInfo& Info
+		const FTraversalInfo& Info,
+		bool bPrimitiveShape
 	)
 	{
 		FString PrimPath = Prim.GetPrimPath().GetString();
-		FString NodeUid = MeshPrefix + PrimPath;
+		FString NodeUid = (bPrimitiveShape ? PrimitiveShapePrefix : MeshPrefix) + PrimPath;
 		FString NodeName(Prim.GetName().ToString());
 
 		// Check if Node already exist with this ID
@@ -1327,7 +1329,14 @@ namespace UE::InterchangeUsdTranslator::Private
 		MeshNode->InitializeNode(NodeUid, NodeName, EInterchangeNodeContainerType::TranslatedAsset);
 		MeshNode->SetAssetName(NodeName);
 		const bool bIsSkinned = static_cast<bool>(Info.ClosestParentSkelRoot) && Prim.HasAPI(TEXT("SkelBindingAPI"));
-		if (bIsSkinned)
+
+		// We are currently not supporting Skinned Primitive Shapes.
+		// In theory a Skinned Mesh needs Joint influences and weights information provided,
+		//		However, there does not seem to be a ruleset against a Primitive Shape having SkelBindingAPI set.
+		//		Which means, that on a theoretical level there could be such a scenario. 
+		ensureMsgf(!(bPrimitiveShape && bIsSkinned), TEXT("Unexpected scenario: Primitive Shape is skinned."));
+
+		if (bIsSkinned && !bPrimitiveShape)
 		{
 			MeshNode->SetSkinnedMesh(true);
 			MeshNode->SetPayLoadKey(PrimPath, EInterchangeMeshPayLoadType::SKELETAL);
@@ -1347,7 +1356,10 @@ namespace UE::InterchangeUsdTranslator::Private
 		}
 		else
 		{
-			MeshNode->SetPayLoadKey(PrimPath, EInterchangeMeshPayLoadType::STATIC);
+			// For Primitive Shapes we add PrimitiveShapePrefix for the PayLoadKey, 
+			// in order to be able to identify the PrimitiveShape in the PayloadData acquisition phase,
+			// as the Primitive Shapes require a different MeshDescription acquisition path, compared to Static Meshes.
+			MeshNode->SetPayLoadKey((bPrimitiveShape ? PrimitiveShapePrefix : TEXT("")) + PrimPath, EInterchangeMeshPayLoadType::STATIC);
 		}
 
 		// Material assignments
@@ -1926,6 +1938,15 @@ namespace UE::InterchangeUsdTranslator::Private
 			Prefix = &LightPrefix;
 			AddLightNode(Prim, NodeContainer);
 		}
+		else if (Prim.IsA(TEXT("Gprim")) && !Prim.IsA(TEXT("PointBased")) && !Prim.IsA(TEXT("Volume")))
+		{
+			//PointBased prims are currently not supported apart from Meshes (which is taken care of in a previous branch).
+			//Volumes are also not currently supported.
+
+			Prefix = &PrimitiveShapePrefix;
+			constexpr bool bPrimitiveShape_True = true;
+			TranslatorImpl.AddMeshNode(Prim, NodeContainer, Info, bPrimitiveShape_True);
+		}
 
 		// Only prims that require rendering (and have a renderable parent) get a scene node.
 		// This includes Xforms but also Scopes, which are not Xformable.
@@ -1952,17 +1973,17 @@ namespace UE::InterchangeUsdTranslator::Private
 			// If we're an Xformable, get our transform
 			FTransform Transform = FTransform::Identity;
 			bool bResetTransformStack = false;
-			if (bIsImageable
-				&& UsdToUnreal::ConvertXformable(
-					Prim.GetStage(),
-					UE::FUsdTyped(Prim),
-					Transform,
-					UsdUtils::GetEarliestTimeCode(),
-					&bResetTransformStack
-				))
-			{
-				SceneNode->SetCustomLocalTransform(&NodeContainer, Transform);
-			}
+			UsdToUnreal::ConvertXformable(
+				Prim.GetStage(),
+				UE::FUsdTyped(Prim),
+				Transform,
+				UsdUtils::GetEarliestTimeCode(),
+				&bResetTransformStack
+			);
+
+			// All SceneNodes should have their LocalTransform set.
+			//	Not setting will cause ensure hits in Skeleton generations for example.
+			SceneNode->SetCustomLocalTransform(&NodeContainer, Transform);
 
 			// Hide our scene node if it is meant to be hidden in USD
 			//
@@ -2029,7 +2050,10 @@ namespace UE::InterchangeUsdTranslator::Private
 		FMeshDescription& OutMeshDescription
 	)
 	{
-		const FString& PrimPath = PayloadKey;
+		bool bIsPrimitiveShape = PayloadKey.StartsWith(PrimitiveShapePrefix);
+
+		FString PrimPath = bIsPrimitiveShape ? PayloadKey.RightChop(PrimitiveShapePrefix.Len()) : PayloadKey;
+
 		UE::FUsdPrim Prim = Impl.UsdStage.GetPrimAtPath(UE::FSdfPath{*PrimPath});
 		if (!Prim)
 		{
@@ -2039,7 +2063,28 @@ namespace UE::InterchangeUsdTranslator::Private
 		// TODO: We can't do much with these yet: There will be used to generate primvar-compatible
 		// versions of the materials that are assigned to this mesh, whenever we get a pipeline
 		UsdUtils::FUsdPrimMaterialAssignmentInfo TempMaterialInfo;
-		bool bSuccess = UsdToUnreal::ConvertGeomMesh(Prim, OutMeshDescription, TempMaterialInfo, Options);
+		bool bSuccess = false;
+		
+		if (bIsPrimitiveShape)
+		{
+			FMeshDescription TempMeshDescription;
+			UsdUtils::FUsdPrimMaterialAssignmentInfo PrimitiveTempMaterialInfo;
+
+			FStaticMeshAttributes StaticMeshAttributes(TempMeshDescription);
+			StaticMeshAttributes.Register();
+
+			bSuccess = UsdToUnreal::ConvertGeomPrimitive(Prim, TempMeshDescription, PrimitiveTempMaterialInfo, Options);
+			if (bSuccess)
+			{
+				OutMeshDescription = { MoveTemp(TempMeshDescription) };
+				TempMaterialInfo = { MoveTemp(PrimitiveTempMaterialInfo) };
+			}
+		}
+		else
+		{
+			bSuccess = UsdToUnreal::ConvertGeomMesh(Prim, OutMeshDescription, TempMaterialInfo, Options);
+		}
+		
 		if (!bSuccess)
 		{
 			return false;
