[Interchange] reimport skeletal mesh must re-import all alternate skin profiles

#rnx

diff --git a/Engine/Plugins/Interchange/Runtime/Source/Import/Private/Mesh/InterchangeSkeletalMeshFactory.cpp b/Engine/Plugins/Interchange/Runtime/Source/Import/Private/Mesh/InterchangeSkeletalMeshFactory.cpp
--- a/Engine/Plugins/Interchange/Runtime/Source/Import/Private/Mesh/InterchangeSkeletalMeshFactory.cpp
+++ b/Engine/Plugins/Interchange/Runtime/Source/Import/Private/Mesh/InterchangeSkeletalMeshFactory.cpp
@@ -2372,6 +2372,10 @@ void UInterchangeSkeletalMeshFactory::SetupObject_GameThread(const FSetupObjectP
 					return EmptyProfile;
 				};
 			
+			TMap<int32, FSkeletalMeshImportData> ImportDataDestPerLod;
+			TBitArray SaveLodIndexes;
+			SaveLodIndexes.Init(false, LodCount);
+
 			for(const FName& ProfileName : ProfileList)
 			{
 				const FSkinWeightProfileInfo& ProfileInfo = GetProfileFromName(ProfileName);
@@ -2397,14 +2401,19 @@ void UInterchangeSkeletalMeshFactory::SetupObject_GameThread(const FSetupObjectP
 
 					const FString ProfileNameStr = ProfileInfo.Name.ToString();
 
-					FSkeletalMeshImportData ImportDataDest;
-					PRAGMA_DISABLE_DEPRECATION_WARNINGS
-					SkeletalMesh->LoadLODImportedData(LodIndex, ImportDataDest);
-					PRAGMA_ENABLE_DEPRECATION_WARNINGS
+					if(!ImportDataDestPerLod.Contains(LodIndex))
+					{
+						PRAGMA_DISABLE_DEPRECATION_WARNINGS
+						SkeletalMesh->LoadLODImportedData(LodIndex, ImportDataDestPerLod.Add(LodIndex));
+						PRAGMA_ENABLE_DEPRECATION_WARNINGS
+					}
+					FSkeletalMeshImportData& ImportDataDest = ImportDataDestPerLod.FindChecked(LodIndex);
 
 					int32 PointNumberDest = ImportDataDest.Points.Num();
 					int32 VertexNumberDest = ImportDataDest.Points.Num();
 
+					const bool bIsGeneratedByEngine = !ProfileInfo.PerLODSourceFiles.Contains(LodIndex);
+
 					if (ExistingImportDataSrc.Points.Num() == PointNumberDest)
 					{
 						//Replace the data into the destination bulk data and save it, so the build dont complain until we reimport the new alternate skinning
@@ -2421,22 +2430,44 @@ void UInterchangeSkeletalMeshFactory::SetupObject_GameThread(const FSetupObjectP
 							ImportDataDest.AlternateInfluences.Add(ExistingImportDataSrc.AlternateInfluences[SrcProfileIndex]);
 						}
 
-						//Resave the bulk data with the new or refreshed data
-						PRAGMA_DISABLE_DEPRECATION_WARNINGS
-						SkeletalMesh->SaveLODImportedData(LodIndex, ImportDataDest);
-						PRAGMA_ENABLE_DEPRECATION_WARNINGS
+						SaveLodIndexes[LodIndex] = true;
 					}
-
-					//We must enqueue a post import task that will re-import all skin weight profile.
-					if (!SkeletalMeshPostImportTask.IsValid())
+					else if (bIsGeneratedByEngine)
 					{
-						SkeletalMeshPostImportTask = MakeShared<FInterchangeSkeletalMeshAlternateSkinWeightPostImportTask>(SkeletalMesh);
-						SkeletalMeshPostImportTask->ReimportAlternateSkinWeightDelegate.BindLambda([](USkeletalMesh* SkeletalMesh, int32 LodIndex)
-							{
-								return FSkinWeightsUtilities::ReimportAlternateSkinWeight(SkeletalMesh, LodIndex);
-							});
+						UInterchangeResultError_Generic* Message = AddMessage<UInterchangeResultError_Generic>();
+						Message->Text = FText::Format(NSLOCTEXT("InterchangeSkeletalMeshFactory", "SetupObject_GameThread_AlternateProfileLodCannotImport", "While re-importing skeletal mesh, the alternate skin weight profile '{0}' cannot be re-import for LOD {1}, because mesh topology differ.")
+							, FText::FromString(ProfileNameStr)
+							, FText::AsNumber(LodIndex));
+						Message->AssetFriendlyName = SkeletalMesh->GetName();
+						Message->AssetType = USkeletalMesh::StaticClass();
+						Message->SourceAssetName = Arguments.SourceData->GetFilename();
 					}
-					SkeletalMeshPostImportTask->AddLodToReimportAlternate(LodIndex);
+					
+					//We must enqueue a post import task that will re-import all skin weight profile that has some source file. Generated alternate skin cannot be re-imported
+					if(!bIsGeneratedByEngine)
+					{
+						if (!SkeletalMeshPostImportTask.IsValid())
+						{
+							SkeletalMeshPostImportTask = MakeShared<FInterchangeSkeletalMeshAlternateSkinWeightPostImportTask>(SkeletalMesh);
+							SkeletalMeshPostImportTask->ReimportAlternateSkinWeightDelegate.BindLambda([](USkeletalMesh* SkeletalMesh, int32 LodIndex)
+								{
+									return FSkinWeightsUtilities::ReimportAlternateSkinWeight(SkeletalMesh, LodIndex);
+								});
+						}
+						SkeletalMeshPostImportTask->AddLodToReimportAlternate(LodIndex);
+					}
+				}
+			}
+
+			//Save every modified lod source data
+			for (int32 LodIndex = 0; LodIndex < LodCount; ++LodIndex)
+			{
+				if(SaveLodIndexes[LodIndex] && ImportDataDestPerLod.Contains(LodIndex))
+				{
+					//Resave the bulk data with the new or refreshed data
+					PRAGMA_DISABLE_DEPRECATION_WARNINGS
+					SkeletalMesh->SaveLODImportedData(LodIndex, ImportDataDestPerLod.FindChecked(LodIndex));
+					PRAGMA_ENABLE_DEPRECATION_WARNINGS
 				}
 			}
 			
