[Interchange] Update the tooltip for "Use T0 As Ref Pose" so users can have more information on the effect of this option.

#rnx

diff --git a/Engine/Plugins/Interchange/Runtime/Source/Pipelines/Public/InterchangeGenericAssetsPipelineSharedSettings.h b/Engine/Plugins/Interchange/Runtime/Source/Pipelines/Public/InterchangeGenericAssetsPipelineSharedSettings.h
--- a/Engine/Plugins/Interchange/Runtime/Source/Pipelines/Public/InterchangeGenericAssetsPipelineSharedSettings.h
+++ b/Engine/Plugins/Interchange/Runtime/Source/Pipelines/Public/InterchangeGenericAssetsPipelineSharedSettings.h
@@ -173,7 +173,7 @@ public:
 	UPROPERTY(EditAnywhere, BlueprintReadWrite, Category = "Common Skeletal Meshes and Animations")
 	bool bImportMeshesInBoneHierarchy = true;
 
-	/** If enabled, frame 0 is used as the reference pose for skeletal meshes. */
+	/** If enabled, skinned meshes will be rebind using the joints transform pose at frame 0 instead of being import with the bind pose. */
 	UPROPERTY(EditAnywhere, BlueprintReadWrite, Category = "Common Skeletal Meshes and Animations", meta=(DisplayName="Use T0 As Ref Pose"))
 	bool bUseT0AsRefPose = false;
 
