Interchange MaterialX:
- Swap Normal input from the SurfaceShader to the DisplacementShader
- Plug the Normal DisplacementShader Output to the SurfaceShader.

MX_Displacement:
- Use the local position instead of the world position, since by default MaterialX sets the position in object space.

#virtualized

diff --git a/Engine/Plugins/Interchange/Runtime/Source/FactoryNodes/Private/InterchangeMaterialFactoryNode.cpp b/Engine/Plugins/Interchange/Runtime/Source/FactoryNodes/Private/InterchangeMaterialFactoryNode.cpp
--- a/Engine/Plugins/Interchange/Runtime/Source/FactoryNodes/Private/InterchangeMaterialFactoryNode.cpp
+++ b/Engine/Plugins/Interchange/Runtime/Source/FactoryNodes/Private/InterchangeMaterialFactoryNode.cpp
@@ -415,6 +415,16 @@ bool UInterchangeMaterialFactoryNode::SetCustomScreenSpaceReflections(const bool
 	IMPLEMENT_NODE_ATTRIBUTE_SETTER_NODELEGATE(ScreenSpaceReflections, bool);
 }
 
+bool UInterchangeMaterialFactoryNode::GetCustomDisplacementCenter(float& AttributeValue) const
+{
+	IMPLEMENT_NODE_ATTRIBUTE_GETTER(DisplacementCenter, float);
+}
+
+bool UInterchangeMaterialFactoryNode::SetCustomDisplacementCenter(float AttributeValue)
+{
+	IMPLEMENT_NODE_ATTRIBUTE_SETTER_NODELEGATE(DisplacementCenter, float);
+}
+
 FString UInterchangeMaterialExpressionFactoryNode::GetTypeName() const
 {
 	const FString TypeName = TEXT("MaterialExpressionFactoryNode");
diff --git a/Engine/Plugins/Interchange/Runtime/Source/FactoryNodes/Public/InterchangeMaterialFactoryNode.h b/Engine/Plugins/Interchange/Runtime/Source/FactoryNodes/Public/InterchangeMaterialFactoryNode.h
--- a/Engine/Plugins/Interchange/Runtime/Source/FactoryNodes/Public/InterchangeMaterialFactoryNode.h
+++ b/Engine/Plugins/Interchange/Runtime/Source/FactoryNodes/Public/InterchangeMaterialFactoryNode.h
@@ -270,14 +270,21 @@ public:
 	UFUNCTION(BlueprintCallable, Category = "Interchange | Node | Material")
 	bool SetCustomScreenSpaceReflections(const bool& AttributeValue);
 
+	UFUNCTION(BlueprintCallable, Category = "Interchange | Node | Material")
+	bool GetCustomDisplacementCenter(float& AttributeValue) const;
+
+	UFUNCTION(BlueprintCallable, Category = "Interchange | Node | Material")
+	bool SetCustomDisplacementCenter(float AttributeValue);
+
 private:
-	const UE::Interchange::FAttributeKey Macro_CustomShadingModelKey = UE::Interchange::FAttributeKey(TEXT("ShadingModel"));
-	const UE::Interchange::FAttributeKey Macro_CustomTranslucencyLightingModeKey = UE::Interchange::FAttributeKey(TEXT("TranslucencyLightingMode"));
-	const UE::Interchange::FAttributeKey Macro_CustomBlendModeKey = UE::Interchange::FAttributeKey(TEXT("BlendMode"));
-	const UE::Interchange::FAttributeKey Macro_CustomTwoSidedKey = UE::Interchange::FAttributeKey(TEXT("TwoSided"));
-	const UE::Interchange::FAttributeKey Macro_CustomOpacityMaskClipValueKey = UE::Interchange::FAttributeKey(TEXT("OpacityMaskClipValue"));
-	const UE::Interchange::FAttributeKey Macro_CustomRefractionMethodKey = UE::Interchange::FAttributeKey(TEXT("RefractionMethod"));
-	const UE::Interchange::FAttributeKey Macro_CustomScreenSpaceReflectionsKey = UE::Interchange::FAttributeKey(TEXT("ScreenSpaceReflections"));
+	IMPLEMENT_NODE_ATTRIBUTE_KEY(ShadingModel)
+	IMPLEMENT_NODE_ATTRIBUTE_KEY(TranslucencyLightingMode)
+	IMPLEMENT_NODE_ATTRIBUTE_KEY(BlendMode)
+	IMPLEMENT_NODE_ATTRIBUTE_KEY(TwoSided)
+	IMPLEMENT_NODE_ATTRIBUTE_KEY(OpacityMaskClipValue);
+	IMPLEMENT_NODE_ATTRIBUTE_KEY(RefractionMethod)
+	IMPLEMENT_NODE_ATTRIBUTE_KEY(ScreenSpaceReflections)
+	IMPLEMENT_NODE_ATTRIBUTE_KEY(DisplacementCenter)
 };
 
 UCLASS(BlueprintType)
diff --git a/Engine/Plugins/Interchange/Runtime/Source/Import/Private/Material/InterchangeMaterialFactory.cpp b/Engine/Plugins/Interchange/Runtime/Source/Import/Private/Material/InterchangeMaterialFactory.cpp
--- a/Engine/Plugins/Interchange/Runtime/Source/Import/Private/Material/InterchangeMaterialFactory.cpp
+++ b/Engine/Plugins/Interchange/Runtime/Source/Import/Private/Material/InterchangeMaterialFactory.cpp
@@ -1300,6 +1300,11 @@ void UInterchangeMaterialFactory::SetupMaterial(UMaterial* Material, const FImpo
 					{
 						DisplacementExpression->ConnectExpression(DisplacementInput, GetOutputIndex(*DisplacementExpression, OutputName));
 						Material->bEnableTessellation = true;
+
+						if (float DisplacementCenter; MaterialFactoryNode->GetCustomDisplacementCenter(DisplacementCenter))
+						{
+							Material->DisplacementScaling.Center = DisplacementCenter;
+						}
 					}
 				}
 			}
diff --git a/Engine/Plugins/Interchange/Runtime/Source/Import/Private/MaterialX/MaterialXUtils/MaterialXDirectionalLightShader.cpp b/Engine/Plugins/Interchange/Runtime/Source/Import/Private/MaterialX/MaterialXUtils/MaterialXDirectionalLightShader.cpp
--- a/Engine/Plugins/Interchange/Runtime/Source/Import/Private/MaterialX/MaterialXUtils/MaterialXDirectionalLightShader.cpp
+++ b/Engine/Plugins/Interchange/Runtime/Source/Import/Private/MaterialX/MaterialXUtils/MaterialXDirectionalLightShader.cpp
@@ -16,7 +16,7 @@ TSharedRef<FMaterialXBase> FMaterialXDirectionalLightShader::MakeInstance(UInter
 	return MakeShared<FMaterialXDirectionalLightShader>(BaseNodeContainer);
 }
 
-void FMaterialXDirectionalLightShader::Translate(MaterialX::NodePtr DirectionalLightShaderNode)
+UInterchangeBaseNode* FMaterialXDirectionalLightShader::Translate(MaterialX::NodePtr DirectionalLightShaderNode)
 {
 	this->LightShaderNode = DirectionalLightShaderNode;
 
@@ -33,6 +33,8 @@ void FMaterialXDirectionalLightShader::Translate(MaterialX::NodePtr DirectionalL
 	this->LightNode = DirectionalLightNode;
 
 	PostTranslate();
+
+	return DirectionalLightNode;
 }
 
 void FMaterialXDirectionalLightShader::GetDirection(MaterialX::InputPtr DirectionInput, FTransform& Transform)
diff --git a/Engine/Plugins/Interchange/Runtime/Source/Import/Private/MaterialX/MaterialXUtils/MaterialXDirectionalLightShader.h b/Engine/Plugins/Interchange/Runtime/Source/Import/Private/MaterialX/MaterialXUtils/MaterialXDirectionalLightShader.h
--- a/Engine/Plugins/Interchange/Runtime/Source/Import/Private/MaterialX/MaterialXUtils/MaterialXDirectionalLightShader.h
+++ b/Engine/Plugins/Interchange/Runtime/Source/Import/Private/MaterialX/MaterialXUtils/MaterialXDirectionalLightShader.h
@@ -19,7 +19,7 @@ public:
 
 	static TSharedRef<FMaterialXBase> MakeInstance(UInterchangeBaseNodeContainer& BaseNodeContainer);
 
-	virtual void Translate(MaterialX::NodePtr PointLightShaderNode) override;
+	virtual UInterchangeBaseNode* Translate(MaterialX::NodePtr PointLightShaderNode) override;
 
 	static void GetDirection(MaterialX::InputPtr DirectionInput, FTransform& Transform);
 };
diff --git a/Engine/Plugins/Interchange/Runtime/Source/Import/Private/MaterialX/MaterialXUtils/MaterialXDisplacementShader.cpp b/Engine/Plugins/Interchange/Runtime/Source/Import/Private/MaterialX/MaterialXUtils/MaterialXDisplacementShader.cpp
--- a/Engine/Plugins/Interchange/Runtime/Source/Import/Private/MaterialX/MaterialXUtils/MaterialXDisplacementShader.cpp
+++ b/Engine/Plugins/Interchange/Runtime/Source/Import/Private/MaterialX/MaterialXUtils/MaterialXDisplacementShader.cpp
@@ -20,15 +20,19 @@ TSharedRef<FMaterialXBase> FMaterialXDisplacementShader::MakeInstance(UInterchan
 	return Result;
 }
 
-void FMaterialXDisplacementShader::Translate(MaterialX::NodePtr DisplacementNode)
+UInterchangeBaseNode* FMaterialXDisplacementShader::Translate(MaterialX::NodePtr DisplacementNode)
 {
 	this->SurfaceShaderNode = DisplacementNode;
 
 	UInterchangeShaderNode* DisplacementShaderNode = FMaterialXSurfaceShaderAbstract::Translate(EInterchangeMaterialXShaders::Displacement);
+	// by default the center is at 0.5, but in order to compute correctly the normals from the displacement we need it at 0 in MX_Displacement material function
+	ShaderGraphNode->SetCustomDisplacementCenterMode(0.f);
 
 	using namespace UE::Interchange::Materials;
 
 	// Outputs
 	UInterchangeShaderPortsAPI::ConnectOuputToInputByName(ShaderGraphNode, Common::Parameters::Displacement.ToString(), DisplacementShaderNode->GetUniqueID(), Common::Parameters::Displacement.ToString());
+
+	return DisplacementShaderNode;
 }
 #endif
\ No newline at end of file
diff --git a/Engine/Plugins/Interchange/Runtime/Source/Import/Private/MaterialX/MaterialXUtils/MaterialXDisplacementShader.h b/Engine/Plugins/Interchange/Runtime/Source/Import/Private/MaterialX/MaterialXUtils/MaterialXDisplacementShader.h
--- a/Engine/Plugins/Interchange/Runtime/Source/Import/Private/MaterialX/MaterialXUtils/MaterialXDisplacementShader.h
+++ b/Engine/Plugins/Interchange/Runtime/Source/Import/Private/MaterialX/MaterialXUtils/MaterialXDisplacementShader.h
@@ -20,6 +20,6 @@ public:
 
 	static TSharedRef<FMaterialXBase> MakeInstance(UInterchangeBaseNodeContainer& BaseNodeContainer);
 
-	virtual void Translate(MaterialX::NodePtr DisplacementNode) override;
+	virtual UInterchangeBaseNode* Translate(MaterialX::NodePtr DisplacementNode) override;
 };
 #endif
\ No newline at end of file
diff --git a/Engine/Plugins/Interchange/Runtime/Source/Import/Private/MaterialX/MaterialXUtils/MaterialXOpenPBRSurfaceShader.cpp b/Engine/Plugins/Interchange/Runtime/Source/Import/Private/MaterialX/MaterialXUtils/MaterialXOpenPBRSurfaceShader.cpp
--- a/Engine/Plugins/Interchange/Runtime/Source/Import/Private/MaterialX/MaterialXUtils/MaterialXOpenPBRSurfaceShader.cpp
+++ b/Engine/Plugins/Interchange/Runtime/Source/Import/Private/MaterialX/MaterialXUtils/MaterialXOpenPBRSurfaceShader.cpp
@@ -20,7 +20,7 @@ TSharedRef<FMaterialXBase> FMaterialXOpenPBRSurfaceShader::MakeInstance(UInterch
 	return Result;
 }
 
-void FMaterialXOpenPBRSurfaceShader::Translate(MaterialX::NodePtr OpenPBRSurfaceNode)
+UInterchangeBaseNode* FMaterialXOpenPBRSurfaceShader::Translate(MaterialX::NodePtr OpenPBRSurfaceNode)
 {
 	using namespace UE::Interchange::Materials;
 
@@ -89,6 +89,14 @@ void FMaterialXOpenPBRSurfaceShader::Translate(MaterialX::NodePtr OpenPBRSurface
 			UInterchangeShaderPortsAPI::ConnectOuputToInputByName(ShaderGraphNode, Subsurface::Parameters::SubsurfaceColor.ToString(), OpenPBRSurfaceShaderNode->GetUniqueID(), Subsurface::Parameters::SubsurfaceColor.ToString());
 		}
 	}
+
+	return OpenPBRSurfaceShaderNode;
+}
+
+mx::InputPtr FMaterialXOpenPBRSurfaceShader::GetInputNormal(mx::NodePtr OpenPbrSurfaceShaderNode, const char*& InputNormal) const
+{
+	InputNormal = mx::OpenPBRSurface::Input::GeometryNormal;
+	return OpenPbrSurfaceShaderNode->getActiveInput(mx::OpenPBRSurface::Input::GeometryNormal);
 }
 
 #endif //WITH_EDITOR
diff --git a/Engine/Plugins/Interchange/Runtime/Source/Import/Private/MaterialX/MaterialXUtils/MaterialXOpenPBRSurfaceShader.h b/Engine/Plugins/Interchange/Runtime/Source/Import/Private/MaterialX/MaterialXUtils/MaterialXOpenPBRSurfaceShader.h
--- a/Engine/Plugins/Interchange/Runtime/Source/Import/Private/MaterialX/MaterialXUtils/MaterialXOpenPBRSurfaceShader.h
+++ b/Engine/Plugins/Interchange/Runtime/Source/Import/Private/MaterialX/MaterialXUtils/MaterialXOpenPBRSurfaceShader.h
@@ -18,7 +18,9 @@ public:
 
 	static TSharedRef<FMaterialXBase> MakeInstance(UInterchangeBaseNodeContainer& BaseNodeContainer);
 
-	virtual void Translate(MaterialX::NodePtr OpenPBRSurfaceNode) override;
+	virtual UInterchangeBaseNode* Translate(MaterialX::NodePtr OpenPBRSurfaceNode) override;
+
+	virtual MaterialX::InputPtr GetInputNormal(MaterialX::NodePtr OpenPbrSurfaceShaderNode, const char*& InputNormal) const override;
 };
 
 #endif //WITH_EDITOR
diff --git a/Engine/Plugins/Interchange/Runtime/Source/Import/Private/MaterialX/MaterialXUtils/MaterialXPointLightShader.cpp b/Engine/Plugins/Interchange/Runtime/Source/Import/Private/MaterialX/MaterialXUtils/MaterialXPointLightShader.cpp
--- a/Engine/Plugins/Interchange/Runtime/Source/Import/Private/MaterialX/MaterialXUtils/MaterialXPointLightShader.cpp
+++ b/Engine/Plugins/Interchange/Runtime/Source/Import/Private/MaterialX/MaterialXUtils/MaterialXPointLightShader.cpp
@@ -17,7 +17,7 @@ TSharedRef<FMaterialXBase> FMaterialXPointLightShader::MakeInstance(UInterchange
 	return MakeShared<FMaterialXPointLightShader>(BaseNodeContainer);
 }
 
-void FMaterialXPointLightShader::Translate(MaterialX::NodePtr PointLightShaderNode)
+UInterchangeBaseNode* FMaterialXPointLightShader::Translate(MaterialX::NodePtr PointLightShaderNode)
 {
 	this->LightShaderNode = PointLightShaderNode;
 
@@ -33,6 +33,8 @@ void FMaterialXPointLightShader::Translate(MaterialX::NodePtr PointLightShaderNo
 	SetPosition();
 
 	PostTranslate();
+
+	return LightNode;
 }
 
 float FMaterialXPointLightShader::GetDecayRate(mx::InputPtr DecayRateInput)
diff --git a/Engine/Plugins/Interchange/Runtime/Source/Import/Private/MaterialX/MaterialXUtils/MaterialXPointLightShader.h b/Engine/Plugins/Interchange/Runtime/Source/Import/Private/MaterialX/MaterialXUtils/MaterialXPointLightShader.h
--- a/Engine/Plugins/Interchange/Runtime/Source/Import/Private/MaterialX/MaterialXUtils/MaterialXPointLightShader.h
+++ b/Engine/Plugins/Interchange/Runtime/Source/Import/Private/MaterialX/MaterialXUtils/MaterialXPointLightShader.h
@@ -18,7 +18,7 @@ public:
 
 	static TSharedRef<FMaterialXBase> MakeInstance(UInterchangeBaseNodeContainer& BaseNodeContainer);
 
-	virtual void Translate(MaterialX::NodePtr PointLightShaderNode) override;
+	virtual UInterchangeBaseNode* Translate(MaterialX::NodePtr PointLightShaderNode) override;
 
 	static float GetDecayRate(MaterialX::InputPtr DecayRateInput);
 
diff --git a/Engine/Plugins/Interchange/Runtime/Source/Import/Private/MaterialX/MaterialXUtils/MaterialXSpotLightShader.cpp b/Engine/Plugins/Interchange/Runtime/Source/Import/Private/MaterialX/MaterialXUtils/MaterialXSpotLightShader.cpp
--- a/Engine/Plugins/Interchange/Runtime/Source/Import/Private/MaterialX/MaterialXUtils/MaterialXSpotLightShader.cpp
+++ b/Engine/Plugins/Interchange/Runtime/Source/Import/Private/MaterialX/MaterialXUtils/MaterialXSpotLightShader.cpp
@@ -18,7 +18,7 @@ TSharedRef<FMaterialXBase> FMaterialXSpotLightShader::MakeInstance(UInterchangeB
 	return MakeShared<FMaterialXSpotLightShader>(BaseNodeContainer);
 }
 
-void FMaterialXSpotLightShader::Translate(MaterialX::NodePtr SpotLightShaderNode)
+UInterchangeBaseNode* FMaterialXSpotLightShader::Translate(MaterialX::NodePtr SpotLightShaderNode)
 {
 	this->LightShaderNode = SpotLightShaderNode;
 
@@ -49,6 +49,8 @@ void FMaterialXSpotLightShader::Translate(MaterialX::NodePtr SpotLightShaderNode
 	SetOuterAngle();
 
 	PostTranslate();
+
+	return LightNode;
 }
 
 void FMaterialXSpotLightShader::SetInnerAngle()
diff --git a/Engine/Plugins/Interchange/Runtime/Source/Import/Private/MaterialX/MaterialXUtils/MaterialXSpotLightShader.h b/Engine/Plugins/Interchange/Runtime/Source/Import/Private/MaterialX/MaterialXUtils/MaterialXSpotLightShader.h
--- a/Engine/Plugins/Interchange/Runtime/Source/Import/Private/MaterialX/MaterialXUtils/MaterialXSpotLightShader.h
+++ b/Engine/Plugins/Interchange/Runtime/Source/Import/Private/MaterialX/MaterialXUtils/MaterialXSpotLightShader.h
@@ -18,7 +18,7 @@ public:
 
 	static TSharedRef<FMaterialXBase> MakeInstance(UInterchangeBaseNodeContainer& BaseNodeContainer);
 
-	virtual void Translate(MaterialX::NodePtr SpotLightShaderNode) override;
+	virtual UInterchangeBaseNode* Translate(MaterialX::NodePtr SpotLightShaderNode) override;
 
 	void SetInnerAngle();
 
diff --git a/Engine/Plugins/Interchange/Runtime/Source/Import/Private/MaterialX/MaterialXUtils/MaterialXStandardSurfaceShader.cpp b/Engine/Plugins/Interchange/Runtime/Source/Import/Private/MaterialX/MaterialXUtils/MaterialXStandardSurfaceShader.cpp
--- a/Engine/Plugins/Interchange/Runtime/Source/Import/Private/MaterialX/MaterialXUtils/MaterialXStandardSurfaceShader.cpp
+++ b/Engine/Plugins/Interchange/Runtime/Source/Import/Private/MaterialX/MaterialXUtils/MaterialXStandardSurfaceShader.cpp
@@ -19,7 +19,7 @@ TSharedRef<FMaterialXBase> FMaterialXStandardSurfaceShader::MakeInstance(UInterc
 	return Result;
 }
 
-void FMaterialXStandardSurfaceShader::Translate(mx::NodePtr StandardSurfaceNode)
+UInterchangeBaseNode* FMaterialXStandardSurfaceShader::Translate(mx::NodePtr StandardSurfaceNode)
 {
 	this->SurfaceShaderNode = StandardSurfaceNode;
 	
@@ -107,6 +107,14 @@ void FMaterialXStandardSurfaceShader::Translate(mx::NodePtr StandardSurfaceNode)
 			UInterchangeShaderPortsAPI::ConnectOuputToInputByName(ShaderGraphNode, ThinTranslucent::Parameters::SurfaceCoverage.ToString(), StandardSurfaceShaderNode->GetUniqueID(), ThinTranslucent::Parameters::SurfaceCoverage.ToString());
 		}
 	}
+
+	return StandardSurfaceShaderNode;
+}
+
+mx::InputPtr FMaterialXStandardSurfaceShader::GetInputNormal(mx::NodePtr StandardSurfaceNode, const char*& InputNormal) const
+{
+	InputNormal = mx::StandardSurface::Input::Normal;
+	return StandardSurfaceNode->getActiveInput(mx::StandardSurface::Input::Normal);
 }
 
 #endif
diff --git a/Engine/Plugins/Interchange/Runtime/Source/Import/Private/MaterialX/MaterialXUtils/MaterialXStandardSurfaceShader.h b/Engine/Plugins/Interchange/Runtime/Source/Import/Private/MaterialX/MaterialXUtils/MaterialXStandardSurfaceShader.h
--- a/Engine/Plugins/Interchange/Runtime/Source/Import/Private/MaterialX/MaterialXUtils/MaterialXStandardSurfaceShader.h
+++ b/Engine/Plugins/Interchange/Runtime/Source/Import/Private/MaterialX/MaterialXUtils/MaterialXStandardSurfaceShader.h
@@ -18,6 +18,8 @@ public:
 
 	static TSharedRef<FMaterialXBase> MakeInstance(UInterchangeBaseNodeContainer& BaseNodeContainer);
 
-	virtual void Translate(MaterialX::NodePtr StandardSurfaceNode) override;
+	virtual UInterchangeBaseNode* Translate(MaterialX::NodePtr StandardSurfaceNode) override;
+
+	virtual MaterialX::InputPtr GetInputNormal(MaterialX::NodePtr StandardSurfaceNode, const char*& InputNormal) const override;
 };
 #endif
\ No newline at end of file
diff --git a/Engine/Plugins/Interchange/Runtime/Source/Import/Private/MaterialX/MaterialXUtils/MaterialXSurfaceMaterial.cpp b/Engine/Plugins/Interchange/Runtime/Source/Import/Private/MaterialX/MaterialXUtils/MaterialXSurfaceMaterial.cpp
--- a/Engine/Plugins/Interchange/Runtime/Source/Import/Private/MaterialX/MaterialXUtils/MaterialXSurfaceMaterial.cpp
+++ b/Engine/Plugins/Interchange/Runtime/Source/Import/Private/MaterialX/MaterialXUtils/MaterialXSurfaceMaterial.cpp
@@ -19,8 +19,9 @@ TSharedRef<FMaterialXBase> FMaterialXSurfaceMaterial::MakeInstance(UInterchangeB
 	return MakeShared<FMaterialXSurfaceMaterial>(BaseNodeContainer);
 }
 
-void FMaterialXSurfaceMaterial::Translate(MaterialX::NodePtr SurfaceMaterialNode)
+UInterchangeBaseNode* FMaterialXSurfaceMaterial::Translate(MaterialX::NodePtr SurfaceMaterialNode)
 {
+	using namespace UE::Interchange::Materials;
 	bool bHasSurfaceShader = false;
 	bool bIsShaderGraphInContainer = true;
 
@@ -38,24 +39,64 @@ void FMaterialXSurfaceMaterial::Translate(MaterialX::NodePtr SurfaceMaterialNode
 		bIsShaderGraphInContainer = false;
 	}
 
-	for(mx::InputPtr Input : SurfaceMaterialNode->getInputs())
+	mx::InputPtr InputNormal;
+	mx::NodePtr ConnectedSurfaceShaderNode, ConnectedDisplacementShaderNode;
+	const char* InputNormalName = "normal";
+	auto InitializeShaderTranslator = [&]
+	(mx::InputPtr Input, bool bIsSurfaceShader = false) -> TSharedPtr<FMaterialXSurfaceShaderAbstract>
 	{
-		mx::NodePtr ConnectedNode = Input->getConnectedNode();
-
-		if(ConnectedNode)
+		if (Input)
 		{
-			// we use here a static_cast instead of a dynamic cast, because we know for sure that the underlying type is a surfaceshader
-			// otherwise the document is ill-formed and the Translator already handles it
+			mx::NodePtr ConnectedNode = Input->getConnectedNode();
+			if (!ConnectedNode)
+			{
+				return nullptr;
+			}
 
 			TSharedPtr<FMaterialXSurfaceShaderAbstract> ShaderTranslator = StaticCastSharedPtr<FMaterialXSurfaceShaderAbstract>(FMaterialXManager::GetInstance().GetShaderTranslator(ConnectedNode->getCategory().c_str(), NodeContainer));
-			bHasSurfaceShader = ShaderTranslator.IsValid();
-			if(bHasSurfaceShader)
+			if (ShaderTranslator.IsValid())
 			{
 				ShaderTranslator->ShaderGraphNode = ShaderGraphNode;
-				ShaderTranslator->SetNodeDefinition(ConnectedNode->getNodeDef()->getName().c_str()); // we set the nodedef here, because a displacement shader can either be a float or vector3, we want to be sure to retrieve the correct one
-				ShaderTranslator->Translate(ConnectedNode);
+				ShaderTranslator->SurfaceMaterialName = SurfaceMaterialNode->getName().c_str();
+				// we set the nodedef here, because a displacement shader can either be a float or vector3, we want to be sure to retrieve the correct one
+				ShaderTranslator->SetNodeDefinition(ConnectedNode->getNodeDef()->getName().c_str());
+				bHasSurfaceShader |= true;
+
+				if (bIsSurfaceShader)
+				{
+					ConnectedSurfaceShaderNode = ConnectedNode;
+					InputNormal = ShaderTranslator->GetInputNormal(ConnectedNode, InputNormalName);
+					if (ConnectedDisplacementShaderNode && InputNormal)
+					{
+						//add the input normal from the surfaceshader to the displacement
+						ConnectedDisplacementShaderNode->addInput()->copyContentFrom(InputNormal);
+						// now remove the normal input from the surfaceshader
+						ConnectedNode->removeInput(InputNormal->getName());
+					}
+				}
+				else
+				{
+					ConnectedDisplacementShaderNode = ConnectedNode;
+				}
 			}
+			return ShaderTranslator;
 		}
+		return nullptr;
+	};
+
+	// if we have a displacement shader, the normals have to be blended in MX_Displacement, 
+	// so we need to plug any shader graph linked to the normal input of the surfaceshader on the displacementshader
+	// and then plug the normal output of the displacement to the surfaceshader	
+	TSharedPtr<FMaterialXSurfaceShaderAbstract> DisplacementShaderTranslator = InitializeShaderTranslator(SurfaceMaterialNode->getInput("displacementshader"));
+	TSharedPtr<FMaterialXSurfaceShaderAbstract> SurfaceShaderTranslator = InitializeShaderTranslator(SurfaceMaterialNode->getInput("surfaceshader"), true);
+
+	UInterchangeShaderNode* SurfaceShaderNode = SurfaceShaderTranslator ? Cast<UInterchangeShaderNode>(SurfaceShaderTranslator->Translate(ConnectedSurfaceShaderNode)) : nullptr;
+	UInterchangeShaderNode* DisplacementShaderNode = DisplacementShaderTranslator ? Cast<UInterchangeShaderNode>(DisplacementShaderTranslator->Translate(ConnectedDisplacementShaderNode)) : nullptr;
+
+	//Now connect the displacementshader material function normal output to the input of the surfaceshader material function
+	if(SurfaceShaderNode && DisplacementShaderNode)
+	{
+		UInterchangeShaderPortsAPI::ConnectOuputToInputByName(SurfaceShaderNode, InputNormalName, DisplacementShaderNode->GetUniqueID(), PBRMR::Parameters::Normal.ToString());
 	}
 
 	// We only add the ShaderGraph to the container if we found a supported surfaceshader
@@ -67,6 +108,8 @@ void FMaterialXSurfaceMaterial::Translate(MaterialX::NodePtr SurfaceMaterialNode
 	{
 		UE_LOG(LogInterchangeImport, Warning, TEXT("the surfaceshader of <%s> is not supported"), UTF8_TO_TCHAR(SurfaceMaterialNode->getName().c_str()));
 	}
+
+	return ShaderGraphNode;
 }
 
 #endif
\ No newline at end of file
diff --git a/Engine/Plugins/Interchange/Runtime/Source/Import/Private/MaterialX/MaterialXUtils/MaterialXSurfaceMaterial.h b/Engine/Plugins/Interchange/Runtime/Source/Import/Private/MaterialX/MaterialXUtils/MaterialXSurfaceMaterial.h
--- a/Engine/Plugins/Interchange/Runtime/Source/Import/Private/MaterialX/MaterialXUtils/MaterialXSurfaceMaterial.h
+++ b/Engine/Plugins/Interchange/Runtime/Source/Import/Private/MaterialX/MaterialXUtils/MaterialXSurfaceMaterial.h
@@ -18,6 +18,6 @@ public:
 
 	static TSharedRef<FMaterialXBase> MakeInstance(UInterchangeBaseNodeContainer& BaseNodeContainer);
 
-	virtual void Translate(MaterialX::NodePtr SurfaceMaterialNode) override;
+	virtual UInterchangeBaseNode* Translate(MaterialX::NodePtr SurfaceMaterialNode) override;
 };
 #endif
\ No newline at end of file
diff --git a/Engine/Plugins/Interchange/Runtime/Source/Import/Private/MaterialX/MaterialXUtils/MaterialXSurfaceShader.cpp b/Engine/Plugins/Interchange/Runtime/Source/Import/Private/MaterialX/MaterialXUtils/MaterialXSurfaceShader.cpp
--- a/Engine/Plugins/Interchange/Runtime/Source/Import/Private/MaterialX/MaterialXUtils/MaterialXSurfaceShader.cpp
+++ b/Engine/Plugins/Interchange/Runtime/Source/Import/Private/MaterialX/MaterialXUtils/MaterialXSurfaceShader.cpp
@@ -20,7 +20,7 @@ TSharedRef<FMaterialXBase> FMaterialXSurfaceShader::MakeInstance(UInterchangeBas
 	return Result;
 }
 
-void FMaterialXSurfaceShader::Translate(MaterialX::NodePtr SurfaceNode)
+UInterchangeBaseNode* FMaterialXSurfaceShader::Translate(MaterialX::NodePtr SurfaceNode)
 {
 	this->SurfaceShaderNode = SurfaceNode;
 
@@ -41,5 +41,7 @@ void FMaterialXSurfaceShader::Translate(MaterialX::NodePtr SurfaceNode)
 			UInterchangeShaderPortsAPI::ConnectOuputToInputByName(ShaderGraphNode, SubstrateMaterial::Parameters::OpacityMask.ToString(), FunctionSurfaceShaderNode->GetUniqueID(), Surface::Substrate::Outputs::Opacity.ToString());
 		}
 	}
+
+	return FunctionSurfaceShaderNode;
 }
 #endif
\ No newline at end of file
diff --git a/Engine/Plugins/Interchange/Runtime/Source/Import/Private/MaterialX/MaterialXUtils/MaterialXSurfaceShader.h b/Engine/Plugins/Interchange/Runtime/Source/Import/Private/MaterialX/MaterialXUtils/MaterialXSurfaceShader.h
--- a/Engine/Plugins/Interchange/Runtime/Source/Import/Private/MaterialX/MaterialXUtils/MaterialXSurfaceShader.h
+++ b/Engine/Plugins/Interchange/Runtime/Source/Import/Private/MaterialX/MaterialXUtils/MaterialXSurfaceShader.h
@@ -20,6 +20,6 @@ public:
 
 	static TSharedRef<FMaterialXBase> MakeInstance(UInterchangeBaseNodeContainer& BaseNodeContainer);
 
-	virtual void Translate(MaterialX::NodePtr SurfaceNode) override;
+	virtual UInterchangeBaseNode* Translate(MaterialX::NodePtr SurfaceNode) override;
 };
 #endif
\ No newline at end of file
diff --git a/Engine/Plugins/Interchange/Runtime/Source/Import/Private/MaterialX/MaterialXUtils/MaterialXSurfaceShaderAbstract.cpp b/Engine/Plugins/Interchange/Runtime/Source/Import/Private/MaterialX/MaterialXUtils/MaterialXSurfaceShaderAbstract.cpp
--- a/Engine/Plugins/Interchange/Runtime/Source/Import/Private/MaterialX/MaterialXUtils/MaterialXSurfaceShaderAbstract.cpp
+++ b/Engine/Plugins/Interchange/Runtime/Source/Import/Private/MaterialX/MaterialXUtils/MaterialXSurfaceShaderAbstract.cpp
@@ -20,6 +20,13 @@ namespace mx = MaterialX;
 FString FMaterialXSurfaceShaderAbstract::EmptyString{};
 FString FMaterialXSurfaceShaderAbstract::DefaultOutput{TEXT("out")};
 
+
+/**
+* Get the normal input of a surfaceshader, used to plug it in the displacementshader
+* @param SurfaceShaderNode - The surfaceshader that we return the input from
+* @return the input of the normal
+*/
+
 FMaterialXSurfaceShaderAbstract::FMaterialXSurfaceShaderAbstract(UInterchangeBaseNodeContainer& BaseNodeContainer)
 	: FMaterialXBase{ BaseNodeContainer }
 	, ShaderGraphNode{ nullptr }
@@ -397,7 +404,7 @@ void FMaterialXSurfaceShaderAbstract::ConnectNodeCategoryOutputToInput(const Mat
 		{
 			if(!ConnectMatchingNodeOutputToInput({ UpstreamNode, ParentShaderNode, InputChannelName, OutputChannelName }))
 			{
-				UE_LOG(LogInterchangeImport, Warning, TEXT("<%s> is not supported."), UTF8_TO_TCHAR(UpstreamNode->getCategory().c_str()));
+				UE_LOG(LogInterchangeImport, Warning, TEXT("<%s>: \"%s\" is not supported. [%s]"), UTF8_TO_TCHAR(UpstreamNode->getCategory().c_str()), UTF8_TO_TCHAR(UpstreamNode->getName().c_str()), *SurfaceMaterialName);
 			}
 		}
 	}
@@ -1868,5 +1875,11 @@ UInterchangeShaderNode* FMaterialXSurfaceShaderAbstract::Translate(EInterchangeM
 	return FunctionCallShaderNode;
 }
 
+mx::InputPtr FMaterialXSurfaceShaderAbstract::GetInputNormal(mx::NodePtr Node, const char*& InputNormal) const
+{
+	InputNormal = "normal";
+	return nullptr;
+}
+
 #undef LOCTEXT_NAMESPACE
 #endif
\ No newline at end of file
diff --git a/Engine/Plugins/Interchange/Runtime/Source/Import/Private/MaterialX/MaterialXUtils/MaterialXSurfaceUnlitShader.cpp b/Engine/Plugins/Interchange/Runtime/Source/Import/Private/MaterialX/MaterialXUtils/MaterialXSurfaceUnlitShader.cpp
--- a/Engine/Plugins/Interchange/Runtime/Source/Import/Private/MaterialX/MaterialXUtils/MaterialXSurfaceUnlitShader.cpp
+++ b/Engine/Plugins/Interchange/Runtime/Source/Import/Private/MaterialX/MaterialXUtils/MaterialXSurfaceUnlitShader.cpp
@@ -20,7 +20,7 @@ TSharedRef<FMaterialXBase> FMaterialXSurfaceUnlitShader::MakeInstance(UInterchan
 	return Result;
 }
 
-void FMaterialXSurfaceUnlitShader::Translate(mx::NodePtr SurfaceUnlitNode)
+UInterchangeBaseNode* FMaterialXSurfaceUnlitShader::Translate(mx::NodePtr SurfaceUnlitNode)
 {
 	using namespace UE::Interchange::Materials;
 	this->SurfaceShaderNode = SurfaceUnlitNode;
@@ -56,5 +56,7 @@ void FMaterialXSurfaceUnlitShader::Translate(mx::NodePtr SurfaceUnlitNode)
 			ShaderGraphNode->SetCustomBlendMode(EBlendMode::BLEND_Masked);
 		}
 	}
+
+	return SurfaceUnlitShaderNode;
 }
 #endif
\ No newline at end of file
diff --git a/Engine/Plugins/Interchange/Runtime/Source/Import/Private/MaterialX/MaterialXUtils/MaterialXSurfaceUnlitShader.h b/Engine/Plugins/Interchange/Runtime/Source/Import/Private/MaterialX/MaterialXUtils/MaterialXSurfaceUnlitShader.h
--- a/Engine/Plugins/Interchange/Runtime/Source/Import/Private/MaterialX/MaterialXUtils/MaterialXSurfaceUnlitShader.h
+++ b/Engine/Plugins/Interchange/Runtime/Source/Import/Private/MaterialX/MaterialXUtils/MaterialXSurfaceUnlitShader.h
@@ -18,6 +18,6 @@ public:
 
 	static TSharedRef<FMaterialXBase> MakeInstance(UInterchangeBaseNodeContainer& BaseNodeContainer);
 
-	virtual void Translate(MaterialX::NodePtr SurfaceUnlitNode) override;
+	virtual UInterchangeBaseNode* Translate(MaterialX::NodePtr SurfaceUnlitNode) override;
 };
 #endif
\ No newline at end of file
diff --git a/Engine/Plugins/Interchange/Runtime/Source/Import/Private/MaterialX/MaterialXUtils/MaterialXUsdPreviewSurfaceShader.cpp b/Engine/Plugins/Interchange/Runtime/Source/Import/Private/MaterialX/MaterialXUtils/MaterialXUsdPreviewSurfaceShader.cpp
--- a/Engine/Plugins/Interchange/Runtime/Source/Import/Private/MaterialX/MaterialXUtils/MaterialXUsdPreviewSurfaceShader.cpp
+++ b/Engine/Plugins/Interchange/Runtime/Source/Import/Private/MaterialX/MaterialXUtils/MaterialXUsdPreviewSurfaceShader.cpp
@@ -20,7 +20,7 @@ TSharedRef<FMaterialXBase> FMaterialXUsdPreviewSurfaceShader::MakeInstance(UInte
 	return Result;
 }
 
-void FMaterialXUsdPreviewSurfaceShader::Translate(MaterialX::NodePtr UsdPreviewSurfaceNode)
+UInterchangeBaseNode* FMaterialXUsdPreviewSurfaceShader::Translate(MaterialX::NodePtr UsdPreviewSurfaceNode)
 {
 	this->SurfaceShaderNode = UsdPreviewSurfaceNode;
 
@@ -68,5 +68,13 @@ void FMaterialXUsdPreviewSurfaceShader::Translate(MaterialX::NodePtr UsdPreviewS
 			ShaderGraphNode->SetCustomBlendMode(EBlendMode::BLEND_TranslucentColoredTransmittance);
 		}
 	}
+
+	return UsdPreviewSurfaceShaderNode;
+}
+
+mx::InputPtr FMaterialXUsdPreviewSurfaceShader::GetInputNormal(mx::NodePtr UsdPreviewSurfaceNode, const char*& InputNormal) const
+{
+	InputNormal = mx::UsdPreviewSurface::Input::Normal;
+	return UsdPreviewSurfaceNode->getActiveInput(mx::UsdPreviewSurface::Input::Normal);
 }
 #endif
\ No newline at end of file
diff --git a/Engine/Plugins/Interchange/Runtime/Source/Import/Private/MaterialX/MaterialXUtils/MaterialXUsdPreviewSurfaceShader.h b/Engine/Plugins/Interchange/Runtime/Source/Import/Private/MaterialX/MaterialXUtils/MaterialXUsdPreviewSurfaceShader.h
--- a/Engine/Plugins/Interchange/Runtime/Source/Import/Private/MaterialX/MaterialXUtils/MaterialXUsdPreviewSurfaceShader.h
+++ b/Engine/Plugins/Interchange/Runtime/Source/Import/Private/MaterialX/MaterialXUtils/MaterialXUsdPreviewSurfaceShader.h
@@ -19,6 +19,8 @@ public:
 
 	static TSharedRef<FMaterialXBase> MakeInstance(UInterchangeBaseNodeContainer& BaseNodeContainer);
 
-	virtual void Translate(MaterialX::NodePtr UsdPreviewSurfaceNode) override;
+	virtual UInterchangeBaseNode* Translate(MaterialX::NodePtr UsdPreviewSurfaceNode) override;
+
+	virtual MaterialX::InputPtr GetInputNormal(MaterialX::NodePtr UsdPreviewSurfaceNode, const char*& InputNormal) const override;
 };
 #endif
\ No newline at end of file
diff --git a/Engine/Plugins/Interchange/Runtime/Source/Import/Public/MaterialX/MaterialXUtils/MaterialXBase.h b/Engine/Plugins/Interchange/Runtime/Source/Import/Public/MaterialX/MaterialXUtils/MaterialXBase.h
--- a/Engine/Plugins/Interchange/Runtime/Source/Import/Public/MaterialX/MaterialXUtils/MaterialXBase.h
+++ b/Engine/Plugins/Interchange/Runtime/Source/Import/Public/MaterialX/MaterialXUtils/MaterialXBase.h
@@ -17,7 +17,7 @@ protected:
 
 public:
 
-	virtual void Translate(MaterialX::NodePtr Node) = 0;
+	virtual UInterchangeBaseNode* Translate(MaterialX::NodePtr Node) = 0;
 
 	virtual ~FMaterialXBase() = default;
 	
diff --git a/Engine/Plugins/Interchange/Runtime/Source/Import/Public/MaterialX/MaterialXUtils/MaterialXSurfaceShaderAbstract.h b/Engine/Plugins/Interchange/Runtime/Source/Import/Public/MaterialX/MaterialXUtils/MaterialXSurfaceShaderAbstract.h
--- a/Engine/Plugins/Interchange/Runtime/Source/Import/Public/MaterialX/MaterialXUtils/MaterialXSurfaceShaderAbstract.h
+++ b/Engine/Plugins/Interchange/Runtime/Source/Import/Public/MaterialX/MaterialXUtils/MaterialXSurfaceShaderAbstract.h
@@ -14,6 +14,15 @@ enum EVectorNoiseFunction : int;
 
 class FMaterialXSurfaceShaderAbstract : public FMaterialXBase
 {
+public:
+
+	/**
+	 * Get the normal input of a surfaceshader, used to plug it in the displacementshader
+	 * @param SurfaceShaderNode - The surfaceshader that we return the input from
+	 * @return the input of the normal
+	 */
+	virtual MaterialX::InputPtr GetInputNormal(MaterialX::NodePtr Node, const char*& InputNormal) const;
+
 protected:
 
 	static FString EmptyString;
@@ -518,7 +527,7 @@ protected:
 	 */
 	UInterchangeShaderNode* Translate(EInterchangeMaterialXShaders ShaderType);
 
-	virtual void Translate(MaterialX::NodePtr ShaderNode) override = 0;
+	virtual UInterchangeBaseNode* Translate(MaterialX::NodePtr ShaderNode) override = 0;
 
 private:
 	
@@ -544,6 +553,9 @@ protected:
 	/** Initialized by the material shader (e.g: surfacematerial), the derived class should only set the ShaderType */
 	UInterchangeShaderGraphNode* ShaderGraphNode;
 
+	/** Here for log purpose, it makes the logs easier to read especially when several materials are imported at once*/
+	FString SurfaceMaterialName;
+
 	/** Used for texture compression and transform to tangent space nodes coming from inputs such as coat_normal  */
 	bool bTangentSpaceInput;
 };
diff --git a/Engine/Plugins/Interchange/Runtime/Source/Nodes/Private/InterchangeShaderGraphNode.cpp b/Engine/Plugins/Interchange/Runtime/Source/Nodes/Private/InterchangeShaderGraphNode.cpp
--- a/Engine/Plugins/Interchange/Runtime/Source/Nodes/Private/InterchangeShaderGraphNode.cpp
+++ b/Engine/Plugins/Interchange/Runtime/Source/Nodes/Private/InterchangeShaderGraphNode.cpp
@@ -354,6 +354,16 @@ bool UInterchangeShaderGraphNode::SetCustomBlendMode(int AttributeValue)
 	IMPLEMENT_NODE_ATTRIBUTE_SETTER_NODELEGATE(BlendMode, int);
 }
 
+bool UInterchangeShaderGraphNode::GetCustomDisplacementCenterMode(float& AttributeValue) const
+{
+	IMPLEMENT_NODE_ATTRIBUTE_GETTER(DisplacementCenter, float);
+}
+
+bool UInterchangeShaderGraphNode::SetCustomDisplacementCenterMode(float AttributeValue)
+{
+	IMPLEMENT_NODE_ATTRIBUTE_SETTER_NODELEGATE(DisplacementCenter, float);
+}
+
 FString UInterchangeFunctionCallShaderNode::GetTypeName() const
 {
 	const FString TypeName = TEXT("FunctionCallShaderNode");
diff --git a/Engine/Plugins/Interchange/Runtime/Source/Nodes/Public/InterchangeShaderGraphNode.h b/Engine/Plugins/Interchange/Runtime/Source/Nodes/Public/InterchangeShaderGraphNode.h
--- a/Engine/Plugins/Interchange/Runtime/Source/Nodes/Public/InterchangeShaderGraphNode.h
+++ b/Engine/Plugins/Interchange/Runtime/Source/Nodes/Public/InterchangeShaderGraphNode.h
@@ -300,6 +300,15 @@ public:
 	UFUNCTION(BlueprintCallable, Category = "Interchange | Node | Material")
 	bool SetCustomBlendMode(int AttributeValue);
 
+	/**
+	 * Set the center of the displacement
+	 */
+	UFUNCTION(BlueprintCallable, Category = "Interchange | Node | Material")
+	bool GetCustomDisplacementCenterMode(float& AttributeValue) const;
+
+	UFUNCTION(BlueprintCallable, Category = "Interchange | Node | Material")
+	bool SetCustomDisplacementCenterMode(float AttributeValue);
+
 private:
 	IMPLEMENT_NODE_ATTRIBUTE_KEY(TwoSided)
 	IMPLEMENT_NODE_ATTRIBUTE_KEY(TwoSidedTransmission)
@@ -307,4 +316,5 @@ private:
 	IMPLEMENT_NODE_ATTRIBUTE_KEY(IsAShaderFunction)
 	IMPLEMENT_NODE_ATTRIBUTE_KEY(ScreenSpaceReflections)
 	IMPLEMENT_NODE_ATTRIBUTE_KEY(BlendMode)
+	IMPLEMENT_NODE_ATTRIBUTE_KEY(DisplacementCenter);
 };
\ No newline at end of file
