[Interchagne] Apply custom attributes or reimport strategy on import always.

#rnx

diff --git a/Engine/Plugins/Interchange/Runtime/Source/Import/Private/Scene/InterchangeActorHelper.cpp b/Engine/Plugins/Interchange/Runtime/Source/Import/Private/Scene/InterchangeActorHelper.cpp
--- a/Engine/Plugins/Interchange/Runtime/Source/Import/Private/Scene/InterchangeActorHelper.cpp
+++ b/Engine/Plugins/Interchange/Runtime/Source/Import/Private/Scene/InterchangeActorHelper.cpp
@@ -210,14 +210,11 @@ void UE::Interchange::ActorHelper::ApplySlotMaterialDependencies(const UIntercha
 
 void UE::Interchange::ActorHelper::ApplyAllCustomAttributes(const UInterchangeFactoryBase::FImportSceneObjectsParams& CreateSceneObjectsParams, UObject& ObjectToUpdate)
 {
-	if (CreateSceneObjectsParams.ReimportObject)
+	if (CreateSceneObjectsParams.ReimportObject && (ObjectToUpdate.GetOuter() == CreateSceneObjectsParams.ReimportObject))
 	{
-		if (ObjectToUpdate.GetOuter() == CreateSceneObjectsParams.ReimportObject)
-		{
-			UInterchangeFactoryBaseNode* CurrentNode = UInterchangeFactoryBaseNode::DuplicateWithObject(CreateSceneObjectsParams.FactoryNode, &ObjectToUpdate);
+		UInterchangeFactoryBaseNode* CurrentNode = UInterchangeFactoryBaseNode::DuplicateWithObject(CreateSceneObjectsParams.FactoryNode, &ObjectToUpdate);
 
-			FFactoryCommon::ApplyReimportStrategyToAsset(&ObjectToUpdate, CreateSceneObjectsParams.ReimportFactoryNode, CurrentNode, CreateSceneObjectsParams.FactoryNode);
-		}
+		FFactoryCommon::ApplyReimportStrategyToAsset(&ObjectToUpdate, CreateSceneObjectsParams.ReimportFactoryNode, CurrentNode, CreateSceneObjectsParams.FactoryNode);
 	}
 	else
 	{
