MaterialX:
- Update OpenPBR to 1.1

#virtualized

diff --git a/Engine/Plugins/Interchange/Runtime/Source/Common/Public/MaterialX/InterchangeMaterialXDefinitions.h b/Engine/Plugins/Interchange/Runtime/Source/Common/Public/MaterialX/InterchangeMaterialXDefinitions.h
--- a/Engine/Plugins/Interchange/Runtime/Source/Common/Public/MaterialX/InterchangeMaterialXDefinitions.h
+++ b/Engine/Plugins/Interchange/Runtime/Source/Common/Public/MaterialX/InterchangeMaterialXDefinitions.h
@@ -290,15 +290,13 @@ MATERIALX_NAMESPACE_BEGIN
 		{
 			static constexpr const char* BaseWeight = "base_weight";
 			static constexpr const char* BaseColor = "base_color";
-			static constexpr const char* BaseRoughness = "base_roughness";
+			static constexpr const char* BaseDiffuseRoughness = "base_diffuse_roughness";
 			static constexpr const char* BaseMetalness = "base_metalness";
 			static constexpr const char* SpecularWeight = "specular_weight";
 			static constexpr const char* SpecularColor = "specular_color";
 			static constexpr const char* SpecularRoughness = "specular_roughness";
 			static constexpr const char* SpecularIOR = "specular_ior";
-			static constexpr const char* SpecularIORLevel = "specular_ior_level";
-			static constexpr const char* SpecularAnisotropy = "specular_anisotropy";
-			static constexpr const char* SpecularRotation = "specular_rotation";
+			static constexpr const char* SpecularRoughnessAnisotropy = "specular_roughness_anisotropy";
 			static constexpr const char* TransmissionWeight = "transmission_weight";
 			static constexpr const char* TransmissionColor = "transmission_color";
 			static constexpr const char* TransmissionDepth = "transmission_depth";
@@ -310,17 +308,18 @@ MATERIALX_NAMESPACE_BEGIN
 			static constexpr const char* SubsurfaceColor = "subsurface_color";
 			static constexpr const char* SubsurfaceRadius = "subsurface_radius";
 			static constexpr const char* SubsurfaceRadiusScale = "subsurface_radius_scale";
-			static constexpr const char* SubsurfaceAnisotropy = "subsurface_anisotropy";
+			static constexpr const char* SubsurfaceScatterAnisotropy = "subsurface_scatter_anisotropy";
 			static constexpr const char* FuzzWeight = "fuzz_weight";
 			static constexpr const char* FuzzColor = "fuzz_color";
 			static constexpr const char* FuzzRoughness = "fuzz_roughness";
 			static constexpr const char* CoatWeight = "coat_weight";
 			static constexpr const char* CoatColor = "coat_color";
 			static constexpr const char* CoatRoughness = "coat_roughness";
-			static constexpr const char* CoatAnisotropy = "coat_anisotropy";
+			static constexpr const char* CoatRoughnessAnisotropy = "coat_roughness_anisotropy";
 			static constexpr const char* CoatRotation = "coat_rotation";
 			static constexpr const char* CoatIOR = "coat_ior";
-			static constexpr const char* CoatIORLevel = "coat_ior_level";
+			static constexpr const char* CoatDarkening = "coat_darkening";
+			static constexpr const char* ThinFilmWeight= "thin_film_weight";
 			static constexpr const char* ThinFilmThickness = "thin_film_thickness";
 			static constexpr const char* ThinFilmIOR = "thin_film_ior";
 			static constexpr const char* EmissionLuminance = "emission_luminance";
@@ -336,42 +335,40 @@ MATERIALX_NAMESPACE_BEGIN
 		{
 			static constexpr float BaseWeight = 1.f;
 			static constexpr FLinearColor BaseColor{ 0.8, 0.8, 0.8 };
-			static constexpr float BaseRoughness = 0.f;
+			static constexpr float BaseDiffuseRoughness = 0.f;
 			static constexpr float BaseMetalness = 0.f;
 			static constexpr float SpecularWeight = 1.f;
 			static constexpr FLinearColor SpecularColor{ 1, 1, 1 };
 			static constexpr float SpecularRoughness = 0.3f;
 			static constexpr float SpecularIOR = 1.5f;
-			static constexpr float SpecularIORLevel = 0.5f;
-			static constexpr float SpecularAnisotropy = 0.f;
-			static constexpr float SpecularRotation = 0.f;
+			static constexpr float SpecularRoughnessAnisotropy = 0.f;
 			static constexpr float TransmissionWeight = 0.f;
 			static constexpr FLinearColor TransmissionColor{ 1, 1, 1 };
 			static constexpr float TransmissionDepth = 0.f;
 			static constexpr FLinearColor TransmissionScatter{ 0, 0, 0 };
 			static constexpr float TransmissionScatterAnisotropy = 0.f;
 			static constexpr float TransmissionDispersionScale = 0.f;
-			static constexpr float TransmissionDispersionAbbeNumber = 0.f;
+			static constexpr float TransmissionDispersionAbbeNumber = 20.f;
 			static constexpr float SubsurfaceWeight = 0.f;
 			static constexpr FLinearColor SubsurfaceColor{ 0.8, 0.8, 0.8 };
 			static constexpr float SubsurfaceRadius = 1.f;
 			static constexpr FLinearColor SubsurfaceRadiusScale{1, 0.5, 0.25};
-			static constexpr float SubsurfaceAnisotropy = 0.f;
+			static constexpr float SubsurfaceScatterAnisotropy = 0.f;
 			static constexpr float FuzzWeight = 0.f;
 			static constexpr FLinearColor FuzzColor{ 1, 1, 1 };
 			static constexpr float FuzzRoughness = 0.5f;
 			static constexpr float CoatWeight = 0.f;
 			static constexpr FLinearColor CoatColor{ 1, 1, 1 };
 			static constexpr float CoatRoughness = 0.f;
-			static constexpr float CoatAnisotropy = 0.f;
-			static constexpr float CoatRotation = 0.f;
+			static constexpr float CoatRoughnessAnisotropy = 0.f;
 			static constexpr float CoatIOR = 1.6f;
-			static constexpr float CoatIORLevel = 0.5f;
-			static constexpr float ThinFilmThickness = 0.f;
-			static constexpr float ThinFilmIOR = 1.5f;
+			static constexpr float CoatDarkening = 1.f;
+			static constexpr float ThinFilmWeight = 0.f;
+			static constexpr float ThinFilmThickness = 0.5f;
+			static constexpr float ThinFilmIOR = 1.4f;
 			static constexpr float EmissionLuminance = 0.f;
 			static constexpr FLinearColor EmissionColor{ 1, 1, 1 };
-			static constexpr FLinearColor GeometryOpacity{ 1, 1, 1 };
+			static constexpr float GeometryOpacity = 1.f;
 			static constexpr bool GeometryThinWalled = false;
 			static const FVector GeometryNormal{ 0, 0, 1 };
 			static const FVector GeometryCoatNormal{ 0, 0, 1 };
