[Interchange] Texture payload acqusition failure causing crash fix.

#rnx

diff --git a/Engine/Plugins/Interchange/Runtime/Source/Import/Private/Texture/InterchangeTextureFactory.cpp b/Engine/Plugins/Interchange/Runtime/Source/Import/Private/Texture/InterchangeTextureFactory.cpp
--- a/Engine/Plugins/Interchange/Runtime/Source/Import/Private/Texture/InterchangeTextureFactory.cpp
+++ b/Engine/Plugins/Interchange/Runtime/Source/Import/Private/Texture/InterchangeTextureFactory.cpp
@@ -1656,6 +1656,48 @@ namespace UE::Interchange::Private::InterchangeTextureFactory
 		ErrorMessage->Text = LOCTEXT("InvalidPayload", "Unable to retrieve the payload from the source file.");
 	}
 
+	bool IsTextureOptionalPayloadSet(const FTexturePayloadVariant& TexturePayloadVarient)
+	{
+		if (TexturePayloadVarient.IsType<FEmptyVariantState>())
+		{
+			return false;
+		}
+
+		if (TexturePayloadVarient.IsType<TOptional<FImportImage>>())
+		{
+			if (const TOptional<FImportImage>* Payload = TexturePayloadVarient.TryGet<TOptional<FImportImage>>())
+			{
+				return Payload->IsSet();
+			}
+		}
+
+		if (TexturePayloadVarient.IsType<TOptional<FImportBlockedImage>>())
+		{
+			if (const TOptional<FImportBlockedImage>* Payload = TexturePayloadVarient.TryGet<TOptional<FImportBlockedImage>>())
+			{
+				return Payload->IsSet();
+			}
+		}
+
+		if (TexturePayloadVarient.IsType<TOptional<FImportSlicedImage>>())
+		{
+			if (const TOptional<FImportSlicedImage>* Payload = TexturePayloadVarient.TryGet<TOptional<FImportSlicedImage>>())
+			{
+				return Payload->IsSet();
+			}
+		}
+
+		if (TexturePayloadVarient.IsType<TOptional<FImportLightProfile>>())
+		{
+			if (const TOptional<FImportLightProfile>* Payload = TexturePayloadVarient.TryGet<TOptional<FImportLightProfile>>())
+			{
+				return Payload->IsSet();
+			}
+		}
+
+		return false;
+	}
+
 	bool CanSetupTexture2DSourceData(FTexturePayloadVariant& TexturePayload)
 	{
 		static_assert(TVariantSize_V<FTexturePayloadVariant> == 5, "Please update the code below and this assert to reflect the change to the variant type.");
@@ -2007,7 +2049,7 @@ UInterchangeFactoryBase::FImportAssetResult UInterchangeTextureFactory::ImportAs
 	AlternateTexturePath.Reset();
 	FTexturePayloadVariant TexturePayload = GetTexturePayload(Arguments.SourceData, PayLoadKey.GetValue(), TextureNodeVariant, TextureFactoryNodeVariant, Arguments.Translator, AlternateTexturePath);
 
-	if(TexturePayload.IsType<FEmptyVariantState>())
+	if (!IsTextureOptionalPayloadSet(TexturePayload))
 	{
 		ImportTextureErrorLog(LOCTEXT("TextureFactory_Async_CannotRetrievePayload", "UInterchangeTextureFactory: Invalid translator couldn't retrieve a payload."));
 		return ImportAssetResult;
@@ -2128,6 +2170,8 @@ UInterchangeFactoryBase::FImportAssetResult UInterchangeTextureFactory::EndImpor
 				// Not thread safe. So those should stay on the game thread.
 				Texture->RemoveFromRoot();
 				Texture->MarkAsGarbage();
+
+				Arguments.AssetNode->SetCustomReferenceObject(nullptr);
 			}
 			Arguments.AssetNode->SetEnabled(false);
 			bSkipImport = true;
diff --git a/Engine/Source/Runtime/Interchange/Engine/Private/Tasks/InterchangeTaskCompletion.cpp b/Engine/Source/Runtime/Interchange/Engine/Private/Tasks/InterchangeTaskCompletion.cpp
--- a/Engine/Source/Runtime/Interchange/Engine/Private/Tasks/InterchangeTaskCompletion.cpp
+++ b/Engine/Source/Runtime/Interchange/Engine/Private/Tasks/InterchangeTaskCompletion.cpp
@@ -219,6 +219,11 @@ void UE::Interchange::FTaskCompletion_GameThread::Execute()
 				{
 					if (UObject* Asset = AssetInfo.ImportedObject)
 					{
+						if (Asset->HasAnyFlags(RF_MirroredGarbage))
+						{
+							continue;
+						}
+
 						if (!AsyncHelper->TaskData.ReimportObject)
 						{
 							//Notify the asset registry, only when we have created the asset
