MaterialX:
- Remove UVs flipped as it is something that should be taken care of by formats which import meshes

#virtualized

diff --git a/Engine/Plugins/Interchange/Runtime/Source/Import/Private/MaterialX/MaterialXUtils/MaterialXSurfaceShaderAbstract.cpp b/Engine/Plugins/Interchange/Runtime/Source/Import/Private/MaterialX/MaterialXUtils/MaterialXSurfaceShaderAbstract.cpp
--- a/Engine/Plugins/Interchange/Runtime/Source/Import/Private/MaterialX/MaterialXUtils/MaterialXSurfaceShaderAbstract.cpp
+++ b/Engine/Plugins/Interchange/Runtime/Source/Import/Private/MaterialX/MaterialXUtils/MaterialXSurfaceShaderAbstract.cpp
@@ -1301,21 +1301,7 @@ void FMaterialXSurfaceShaderAbstract::ConnectTexCoordInputToOutput(const FConnec
 		TexCoord->AddInt32Attribute(UInterchangeShaderPortsAPI::MakeInputValueKey(TextureCoordinate::Inputs::Index.ToString()), mx::fromValueString<int>(Input->getValueString()));
 	}
 
-	// Flip UV coordinates to be bottom left instead of top left
-	UInterchangeShaderNode* TexCoordR = CreateMaskShaderNode(0b1000, Connect.UpstreamNode, (Connect.UpstreamNode->getName() + "ComponentMaskR").c_str());
-	UInterchangeShaderPortsAPI::ConnectDefaultOuputToInput(TexCoordR, Mask::Inputs::Input.ToString(), TexCoord->GetUniqueID());
-
-	UInterchangeShaderNode* TexCoordG = CreateMaskShaderNode(0b0100, Connect.UpstreamNode, (Connect.UpstreamNode->getName() + "ComponentMaskG").c_str());
-	UInterchangeShaderPortsAPI::ConnectDefaultOuputToInput(TexCoordG, Mask::Inputs::Input.ToString(), TexCoord->GetUniqueID());
-
-	UInterchangeShaderNode* OneMinusG = CreateShaderNode(Connect.UpstreamNode, (Connect.UpstreamNode->getName() + "OneMinusG").c_str(), OneMinus::Name.ToString());
-	UInterchangeShaderPortsAPI::ConnectDefaultOuputToInput(OneMinusG, OneMinus::Inputs::Input.ToString(), TexCoordG->GetUniqueID());
-
-	UInterchangeShaderNode* TexCoordAppend = CreateShaderNode(Connect.UpstreamNode, (Connect.UpstreamNode->getName() + "Append").c_str(), UE::Expressions::Names::AppendVector);
-	UInterchangeShaderPortsAPI::ConnectDefaultOuputToInput(TexCoordAppend, UE::Expressions::Inputs::A, TexCoordR->GetUniqueID());
-	UInterchangeShaderPortsAPI::ConnectDefaultOuputToInput(TexCoordAppend, UE::Expressions::Inputs::B, OneMinusG->GetUniqueID());
-
-	UInterchangeShaderPortsAPI::ConnectDefaultOuputToInput(Connect.ParentShaderNode, Connect.InputChannelName, TexCoordAppend->GetUniqueID());
+	UInterchangeShaderPortsAPI::ConnectDefaultOuputToInput(Connect.ParentShaderNode, Connect.InputChannelName, TexCoord->GetUniqueID());
 }
 
 void FMaterialXSurfaceShaderAbstract::ConnectSeparateInputToOutput(const FConnectNode& Connect)
