USD: Build the info cache from the Interchange USD Translator.

diff --git a/Engine/Plugins/Interchange/Runtime/Source/Import/Private/Usd/InterchangeUsdTranslator.cpp b/Engine/Plugins/Interchange/Runtime/Source/Import/Private/Usd/InterchangeUsdTranslator.cpp
--- a/Engine/Plugins/Interchange/Runtime/Source/Import/Private/Usd/InterchangeUsdTranslator.cpp
+++ b/Engine/Plugins/Interchange/Runtime/Source/Import/Private/Usd/InterchangeUsdTranslator.cpp
@@ -2,6 +2,8 @@
 
 #include "Usd/InterchangeUsdTranslator.h"
 
+#include "Objects/USDInfoCache.h"
+#include "Objects/USDSchemaTranslator.h"
 #include "UnrealUSDWrapper.h"
 #include "Usd/InterchangeUsdDefinitions.h"
 #include "USDConversionUtils.h"
@@ -192,10 +194,16 @@ namespace UE::InterchangeUsdTranslator::Private
 		/** If we decompressed a USDZ file to a temp folder this will delete everything from that folder */
 		void CleanUpDecompressedUSDZFolder();
 
+		void SetupTranslationContext(const UInterchangeUsdTranslatorSettings& Settings);
+
+		void SetupInfoCache();
+
 	public:
 		// We have to keep a stage reference so that we can parse the payloads after Translate() complete.
 		// ReleaseSource() clears this member, once translation is complete.
 		UE::FUsdStage UsdStage;
+		TSharedPtr<FUsdSchemaTranslationContext> TranslationContext;
+		TSharedPtr<FUsdInfoCache> InfoCache;
 
 		FString USDZFilePath;
 		FString DecompressedUSDZRoot;
@@ -1021,6 +1029,40 @@ namespace UE::InterchangeUsdTranslator::Private
 		DecompressedUSDZRoot.Reset();
 	}
 
+	void UInterchangeUSDTranslatorImpl::SetupTranslationContext(const UInterchangeUsdTranslatorSettings& Settings)
+	{
+		if (!TranslationContext)
+		{
+			TranslationContext = MakeShared<FUsdSchemaTranslationContext>(UsdStage);
+		}
+
+		TranslationContext->bIsImporting = true;
+		TranslationContext->Time = UsdUtils::GetDefaultTimeCode();
+		TranslationContext->bMergeIdenticalMaterialSlots = true;	// Interchange always does this
+		TranslationContext->bAllowInterpretingLODs = false;			// We don't support USD LODs yet
+
+		TranslationContext->PurposesToLoad = EUsdPurpose::Default | EUsdPurpose::Proxy | EUsdPurpose::Render | EUsdPurpose::Guide;
+		TranslationContext->RenderContext = Settings.RenderContext;
+		TranslationContext->MaterialPurpose = Settings.MaterialPurpose;
+	}
+
+	void UInterchangeUSDTranslatorImpl::SetupInfoCache()
+	{
+		if (!ensure(TranslationContext))
+		{
+			return;
+		}
+
+		if (!InfoCache)
+		{
+			InfoCache = MakeShared<FUsdInfoCache>();
+		}
+
+		TranslationContext->UsdInfoCache = InfoCache.Get();
+
+		InfoCache->RebuildCacheForSubtrees({UE::FSdfPath::AbsoluteRootPath()}, *TranslationContext);
+	}
+
 	bool DecompressUSDZFileToTempFolder(const FString& InUSDZFilePath, FString& OutDecompressedUSDZRoot)
 	{
 		const bool bIncludeDot = false;
@@ -1333,7 +1375,7 @@ namespace UE::InterchangeUsdTranslator::Private
 		// We are currently not supporting Skinned Primitive Shapes.
 		// In theory a Skinned Mesh needs Joint influences and weights information provided,
 		//		However, there does not seem to be a ruleset against a Primitive Shape having SkelBindingAPI set.
-		//		Which means, that on a theoretical level there could be such a scenario. 
+		//		Which means, that on a theoretical level there could be such a scenario.
 		ensureMsgf(!(bPrimitiveShape && bIsSkinned), TEXT("Unexpected scenario: Primitive Shape is skinned."));
 
 		if (bIsSkinned && !bPrimitiveShape)
@@ -1356,7 +1398,7 @@ namespace UE::InterchangeUsdTranslator::Private
 		}
 		else
 		{
-			// For Primitive Shapes we add PrimitiveShapePrefix for the PayLoadKey, 
+			// For Primitive Shapes we add PrimitiveShapePrefix for the PayLoadKey,
 			// in order to be able to identify the PrimitiveShape in the PayloadData acquisition phase,
 			// as the Primitive Shapes require a different MeshDescription acquisition path, compared to Static Meshes.
 			MeshNode->SetPayLoadKey((bPrimitiveShape ? PrimitiveShapePrefix : TEXT("")) + PrimPath, EInterchangeMeshPayLoadType::STATIC);
@@ -2064,7 +2106,7 @@ namespace UE::InterchangeUsdTranslator::Private
 		// versions of the materials that are assigned to this mesh, whenever we get a pipeline
 		UsdUtils::FUsdPrimMaterialAssignmentInfo TempMaterialInfo;
 		bool bSuccess = false;
-		
+
 		if (bIsPrimitiveShape)
 		{
 			FMeshDescription TempMeshDescription;
@@ -2084,7 +2126,7 @@ namespace UE::InterchangeUsdTranslator::Private
 		{
 			bSuccess = UsdToUnreal::ConvertGeomMesh(Prim, OutMeshDescription, TempMaterialInfo, Options);
 		}
-		
+
 		if (!bSuccess)
 		{
 			return false;
@@ -2683,6 +2725,9 @@ bool UInterchangeUSDTranslator::Translate(UInterchangeBaseNodeContainer& NodeCon
 		ImplPtr->UsdStage.SetInterpolationType(Settings->InterpolationType);
 	}
 
+	ImplPtr->SetupTranslationContext(*Settings);
+	ImplPtr->SetupInfoCache();
+
 	// Cache these so we don't have to keep converting these tokens over and over during translation
 	FUsdMeshConversionOptions& MeshOptions = ImplPtr->CachedMeshConversionOptions;
 
diff --git a/Engine/Plugins/Runtime/USDCore/Source/USDUtilities/Private/Objects/USDSchemaTranslator.cpp b/Engine/Plugins/Runtime/USDCore/Source/USDUtilities/Private/Objects/USDSchemaTranslator.cpp
--- a/Engine/Plugins/Runtime/USDCore/Source/USDUtilities/Private/Objects/USDSchemaTranslator.cpp
+++ b/Engine/Plugins/Runtime/USDCore/Source/USDUtilities/Private/Objects/USDSchemaTranslator.cpp
@@ -181,6 +181,12 @@ FUsdSchemaTranslationContext::FUsdSchemaTranslationContext(const UE::FUsdStage&
 	RenderContext = UnrealIdentifiers::UniversalRenderContext;
 }
 
+FUsdSchemaTranslationContext::FUsdSchemaTranslationContext(const UE::FUsdStage& InStage)
+	: Stage(InStage)
+{
+	RenderContext = UnrealIdentifiers::UniversalRenderContext;
+}
+
 FUsdSchemaTranslationContext::FUsdSchemaTranslationContext(const UE::FUsdStage& InStage, UUsdAssetCache3& InAssetCache)
 	: Stage(InStage)
 	, UsdAssetCache(&InAssetCache)
diff --git a/Engine/Plugins/Runtime/USDCore/Source/USDUtilities/Public/Objects/USDSchemaTranslator.h b/Engine/Plugins/Runtime/USDCore/Source/USDUtilities/Public/Objects/USDSchemaTranslator.h
--- a/Engine/Plugins/Runtime/USDCore/Source/USDUtilities/Public/Objects/USDSchemaTranslator.h
+++ b/Engine/Plugins/Runtime/USDCore/Source/USDUtilities/Public/Objects/USDSchemaTranslator.h
@@ -166,6 +166,7 @@ struct USDUTILITIES_API FUsdSchemaTranslationContext : public TSharedFromThis<FU
 	UE_DEPRECATED(5.5, "Use the constructor that receives an UUsdAssetCache3 instead")
 	explicit FUsdSchemaTranslationContext(const UE::FUsdStage& InStage, UUsdAssetCache2& InAssetCache);
 
+	explicit FUsdSchemaTranslationContext(const UE::FUsdStage& InStage);
 	explicit FUsdSchemaTranslationContext(const UE::FUsdStage& InStage, UUsdAssetCache3& InAssetCache);
 
 	/** True if we're a context created by the USDStageImporter to fully import to persistent assets and actors */
