USD: Support variant-style LODs via USD Interchange;

Add some basic automated tests for new variant-style LODs and FBX-style LODS.

diff --git a/Engine/Plugins/Interchange/Runtime/Source/Import/Private/Usd/InterchangeUsdTranslator.cpp b/Engine/Plugins/Interchange/Runtime/Source/Import/Private/Usd/InterchangeUsdTranslator.cpp
--- a/Engine/Plugins/Interchange/Runtime/Source/Import/Private/Usd/InterchangeUsdTranslator.cpp
+++ b/Engine/Plugins/Interchange/Runtime/Source/Import/Private/Usd/InterchangeUsdTranslator.cpp
@@ -19,6 +19,7 @@
 #include "USDTypesConversion.h"
 #include "UsdWrappers/SdfPath.h"
 #include "UsdWrappers/UsdAttribute.h"
+#include "UsdWrappers/UsdEditContext.h"
 #include "UsdWrappers/UsdGeomXformable.h"
 #include "UsdWrappers/UsdPrim.h"
 #include "UsdWrappers/UsdRelationship.h"
@@ -32,8 +33,8 @@
 #include "UsdWrappers/UsdSkelSkinningQuery.h"
 #include "UsdWrappers/UsdStage.h"
 #include "UsdWrappers/UsdTyped.h"
+#include "UsdWrappers/UsdVariantSets.h"
 
-#include "Async/Async.h"
 #include "Async/ParallelFor.h"
 #include "HAL/IConsoleManager.h"
 #include "InterchangeCameraNode.h"
@@ -51,12 +52,10 @@
 #include "Rendering/SkeletalMeshLODImporterData.h"
 #include "StaticMeshAttributes.h"
 #include "UDIMUtilities.h"
-
 #include "UObject/GCObjectScopeGuard.h"
 
 #include "MaterialX/MaterialXUtils/MaterialXBase.h"
 #include "MaterialX/MaterialXUtils/MaterialXManager.h"
-#include "MaterialXFormat/XmlIo.h"
 
 #if USE_USD_SDK
 #include "USDIncludesStart.h"
@@ -89,27 +88,132 @@ namespace UE::InterchangeUsdTranslator::Private
 {
 	const static FString AnimationPrefix = TEXT("\\Animation\\");
 	const static FString AnimationTrackPrefix = TEXT("\\AnimationTrack\\");
+	const static FString BonePrefix = TEXT("\\Bone\\");
 	const static FString CameraPrefix = TEXT("\\Camera\\");
 	const static FString LightPrefix = TEXT("\\Light\\");
+	const static FString LODPrefix = TEXT("\\LOD\\");
 	const static FString MaterialPrefix = TEXT("\\Material\\");
 	const static FString MeshPrefix = TEXT("\\Mesh\\");
-	const static FString PrimitiveShapePrefix = TEXT("\\PrimitiveShape\\");
 	const static FString MorphTargetPrefix = TEXT("\\MorphTarget\\");
-	const static FString BonePrefix = TEXT("\\Bone\\");
+	const static FString PrimitiveShapePrefix = TEXT("\\PrimitiveShape\\");
+
+	const static FString LODContainerSuffix = TEXT("LODContainer");
+
+	const static FString LODString = TEXT("LOD");
 
 	// Information intended to be passed down from parent to children (by value) as we traverse the stage
 	struct FTraversalInfo
 	{
 		UInterchangeBaseNode* ParentNode = nullptr;
-		bool bVisible = true;
 
 		TSharedPtr<UE::FUsdSkelCache> FurthestSkelCache;
-		UE::FUsdPrim ClosestParentSkelRoot;
-
-		UE::FUsdSkelSkeletonQuery ActiveSkelQuery;
+		TSharedPtr<FString> FurthestParentSkelRootPath;	   // Used to populate the skel cache
+		TSharedPtr<FString> ClosestParentSkelRootPath;
+		TSharedPtr<FString> BoundSkeletonPrimPath;
 		TSharedPtr<TArray<FString>> SkelJointNames;	   // Needed for skel mesh payloads
+
+		bool bVisible = true;
+		bool bInsideLOD = false;
+		bool bIsLODContainer = false;
+
+	public:
+		void UpdateWithCurrentPrim(const UE::FUsdPrim& CurrentPrim)
+		{
+			bVisible = bVisible && UsdUtils::HasInheritedVisibility(CurrentPrim);
+
+			// Check this first so that we go bInsideLOD if our parent was a LODContainer
+			bInsideLOD = bInsideLOD || bIsLODContainer;
+
+			// We only want this to be true when we're traversing the exact prim that owns
+			// the LOD: Once we step into any of its children it should go back to false
+			bIsLODContainer = CurrentPrim.GetVariantSets().HasVariantSet(LODString);
+
+			if (CurrentPrim.IsA(TEXT("SkelRoot")))
+			{
+				if (!ClosestParentSkelRootPath)
+				{
+					// The root-most skel cache should handle any nested UsdSkel prims as well
+					FurthestSkelCache = MakeShared<UE::FUsdSkelCache>();
+
+					const bool bTraverseInstanceProxies = true;
+					FurthestSkelCache->Populate(CurrentPrim, bTraverseInstanceProxies);
+					FurthestParentSkelRootPath = MakeShared<FString>(CurrentPrim.GetPrimPath().GetString());
+				}
+
+				ClosestParentSkelRootPath = MakeShared<FString>(CurrentPrim.GetPrimPath().GetString());
+			}
+
+			if (ClosestParentSkelRootPath && CurrentPrim.HasAPI(TEXT("SkelBindingAPI")))
+			{
+				UE::FUsdStage Stage = CurrentPrim.GetStage();
+
+				if (UE::FUsdRelationship SkelRel = CurrentPrim.GetRelationship(TEXT("skel:skeleton")))
+				{
+					TArray<UE::FSdfPath> Targets;
+					if (SkelRel.GetTargets(Targets) && Targets.Num() > 0)
+					{
+						UE::FUsdPrim TargetSkeleton = Stage.GetPrimAtPath(Targets[0]);
+						if (TargetSkeleton && TargetSkeleton.IsA(TEXT("Skeleton")))
+						{
+							BoundSkeletonPrimPath = MakeShared<FString>(TargetSkeleton.GetPrimPath().GetString());
+						}
+					}
+				}
+			}
+		}
+
+		UE::FUsdSkelSkeletonQuery ResolveSkelQuery(const UE::FUsdStage& Stage) const
+		{
+			if (!Stage || !BoundSkeletonPrimPath || BoundSkeletonPrimPath->IsEmpty())
+			{
+				return {};
+			}
+
+			UE::FUsdPrim SkeletonPrim = Stage.GetPrimAtPath(UE::FSdfPath{**BoundSkeletonPrimPath});
+			if (!SkeletonPrim)
+			{
+				return {};
+			}
+
+			return FurthestSkelCache->GetSkelQuery(SkeletonPrim);
+		}
+
+		UE::FUsdPrim ResolveClosestParentSkelRoot(const UE::FUsdStage& Stage) const
+		{
+			if (!Stage || !ClosestParentSkelRootPath || ClosestParentSkelRootPath->IsEmpty())
+			{
+				return {};
+			}
+
+			return Stage.GetPrimAtPath(UE::FSdfPath{**ClosestParentSkelRootPath});
+		}
+
+		void RepopulateSkelCache(const UE::FUsdStage& Stage)
+		{
+			if (!FurthestSkelCache.IsValid() || !FurthestParentSkelRootPath.IsValid() || FurthestParentSkelRootPath->IsEmpty())
+			{
+				return;
+			}
+
+			UE::FUsdPrim SkelRootPrim = Stage.GetPrimAtPath(UE::FSdfPath{**FurthestParentSkelRootPath});
+			if (!SkelRootPrim)
+			{
+				return;
+			}
+
+			const bool bTraverseInstanceProxies = true;
+			ensure(FurthestSkelCache->Populate(SkelRootPrim, bTraverseInstanceProxies));
+		}
 	};
 
+	void Traverse(
+		const UE::FUsdPrim& Prim,
+		UInterchangeUSDTranslatorImpl& TranslatorImpl,
+		UInterchangeBaseNodeContainer& NodeContainer,
+		UInterchangeUsdTranslatorSettings* TranslatorSettings,
+		FTraversalInfo Info
+	);
+
 	// clang-format off
 	const static TMap<FName, EInterchangePropertyTracks> PropertyNameToTrackType = {
 		// Common properties
@@ -159,7 +263,20 @@ namespace UE::InterchangeUsdTranslator::Private
 			UInterchangeBaseNodeContainer& NodeContainer
 		);
 
-		void AddMeshNode(const UE::FUsdPrim& Prim, UInterchangeBaseNodeContainer& NodeContainer, const FTraversalInfo& Info, bool bPrimitiveShape = false);
+		void AddMeshNode(
+			const UE::FUsdPrim& Prim,
+			UInterchangeBaseNodeContainer& NodeContainer,
+			const FTraversalInfo& Info,
+			bool bPrimitiveShape = false
+		);
+
+		void AddLODMeshNodes(
+			const UE::FUsdPrim& Prim,
+			UInterchangeBaseNodeContainer& NodeContainer,
+			UInterchangeSceneNode& LODContainerSceneNode,
+			UInterchangeUsdTranslatorSettings* TranslatorSettings,
+			FTraversalInfo Info
+		);
 
 		/**
 		 * Search for a MaterialX file embedded in the file, in case there's none, create a ShaderGraph using MaterialX
@@ -198,6 +315,14 @@ namespace UE::InterchangeUsdTranslator::Private
 
 		void SetupInfoCache();
 
+		UE::FUsdPrim TryGettingInactiveLODPrim(const FString& PrimPath);
+
+		void FinalizeLODContainerTraversal(
+			UInterchangeBaseNodeContainer& NodeContainer,
+			const FTraversalInfo& Info,
+			const UInterchangeSceneNode* SceneNodeWithLODs
+		);
+
 	public:
 		// We have to keep a stage reference so that we can parse the payloads after Translate() complete.
 		// ReleaseSource() clears this member, once translation is complete.
@@ -205,6 +330,11 @@ namespace UE::InterchangeUsdTranslator::Private
 		TSharedPtr<FUsdSchemaTranslationContext> TranslationContext;
 		TSharedPtr<FUsdInfoCache> InfoCache;
 
+		// We store temp stages in here that we open in order to parse stuff inside of inactive variants of our main UsdStage.
+		// We do this because the payload data are retrieved concurrently, and toggling variants mutates the current stage.
+		mutable FRWLock PrimPathToVariantToStageLock;
+		TMap<FString, TMap<FString, UE::FUsdStage>> PrimPathToVariantToStage;
+
 		FString USDZFilePath;
 		FString DecompressedUSDZRoot;
 
@@ -214,6 +344,11 @@ namespace UE::InterchangeUsdTranslator::Private
 		UsdToUnreal::FUsdMeshConversionOptions CachedMeshConversionOptions;
 #endif
 
+		TMap<FString, UsdUtils::FUsdPrimMaterialAssignmentInfo> CachedMaterialAssignments;
+
+		// We fill this in while we're translating a LOD container so that we can do some post-processing inside FinalizeLODContainerTraversal
+		TArray<UInterchangeSceneNode*> CurrentLODSceneNodes;
+
 		// When traversing we'll generate FTraversalInfo objects. If we need to (e.g. for skinned meshes),
 		// we'll store the info for that translated node here, so we don't have to recompute it when returning
 		// the payload data.
@@ -241,6 +376,54 @@ namespace UE::InterchangeUsdTranslator::Private
 		TMap<FString, TArray<FMaterialSlotMesh>> PrimPathToSlotMeshNodes;
 	};
 
+	bool IsValidLODName(const FString& PrimName)
+	{
+		return PrimName.Len() > LODString.Len() && PrimName.StartsWith(LODString) && PrimName.RightChop(LODString.Len()).IsNumeric();
+	}
+
+	int32 GetLODIndexFromName(FString Name)
+	{
+		if (!Name.RemoveFromStart(LODString))
+		{
+			return INDEX_NONE;
+		}
+
+		if (!Name.IsNumeric())
+		{
+			return INDEX_NONE;
+		}
+
+		int32 Index = INDEX_NONE;
+		LexFromString(Index, *Name);
+		return Index;
+	}
+
+	UE::FUsdPrim GetLODMesh(const UE::FUsdPrim& LODContainerPrim, const FString& LODName)
+	{
+		const UE::FSdfPath IdealMeshPrimPath = LODContainerPrim.GetPrimPath().AppendChild(*LODName);
+
+		UE::FUsdPrim Prim = LODContainerPrim.GetStage().GetPrimAtPath(IdealMeshPrimPath);
+		if (Prim.IsA(TEXT("Mesh")))
+		{
+			return Prim;
+		}
+
+		return {};
+	};
+
+	bool CheckAndChopPayloadPrefix(FString& PayloadKey, const FString& Prefix)
+	{
+		if (PayloadKey.StartsWith(Prefix))
+		{
+			// Remove the prefix from the payload key
+			PayloadKey = PayloadKey.RightChop(Prefix.Len());
+
+			return true;
+		}
+
+		return false;
+	}
+
 #if USE_USD_SDK
 	FString HashAnimPayloadQuery(const Interchange::FAnimationPayloadQuery& Query)
 	{
@@ -272,9 +455,53 @@ namespace UE::InterchangeUsdTranslator::Private
 		return FString::Printf(TEXT("%s%s\\%d\\%s"), *MorphTargetPrefix, *MeshPrimPath, MeshBlendShapeIndex, *InbetweenName);
 	}
 
-	FString GetMorphTargetMeshPayloadKey(const FString& MeshPrimPath, int32 MeshBlendShapeIndex, const FString& InbetweenName = FString{})
+	FString GetMorphTargetMeshPayloadKey(
+		bool bIsInsideLOD,
+		const FString& MeshPrimPath,
+		int32 MeshBlendShapeIndex,
+		const FString& InbetweenName = FString{}
+	)
 	{
-		return FString::Printf(TEXT("%s\\%d\\%s"), *MeshPrimPath, MeshBlendShapeIndex, *InbetweenName);
+		return FString::Printf(TEXT("%s%s\\%d\\%s"), bIsInsideLOD ? *LODPrefix : TEXT(""), *MeshPrimPath, MeshBlendShapeIndex, *InbetweenName);
+	}
+
+	// TODO: Cleanup/unify/standardize these payload manipulating functions (don't add/remove prefixes everywhere but have a standard format, etc.)
+	bool ParseMorphTargetMeshPayloadKey(
+		FString InPayloadKey,
+		bool& bOutIsLODMesh,
+		FString& OutMeshPrimPath,
+		int32& OutBlendShapeIndex,
+		FString& OutInbetweenName
+	)
+	{
+		bool bIsLODMesh = CheckAndChopPayloadPrefix(InPayloadKey, LODPrefix);
+
+		// These payload keys are generated by GetMorphTargetMeshPayloadKey(), and so should take the form
+		// "<mesh prim path>\<mesh blend shape index>\<optional inbetween name>"
+		const bool bCullEmpty = false;
+		TArray<FString> PayloadKeyTokens;
+		InPayloadKey.ParseIntoArray(PayloadKeyTokens, TEXT("\\"), bCullEmpty);
+		if (PayloadKeyTokens.Num() != 3)
+		{
+			return false;
+		}
+
+		const FString& MeshPrimPath = PayloadKeyTokens[0];
+		const FString& BlendShapeIndexStr = PayloadKeyTokens[1];
+		const FString& InbetweenName = PayloadKeyTokens[2];
+
+		int32 BlendShapeIndex = INDEX_NONE;
+		bool bLexed = LexTryParseString(BlendShapeIndex, *BlendShapeIndexStr);
+		if (!bLexed)
+		{
+			return false;
+		}
+
+		bOutIsLODMesh = bIsLODMesh;
+		OutMeshPrimPath = MeshPrimPath;
+		OutBlendShapeIndex = BlendShapeIndex;
+		OutInbetweenName = InbetweenName;
+		return true;
 	}
 
 	FString GetMorphTargetCurvePayloadKey(const FString& SkeletonPrimPath, int32 SkelAnimChannelIndex, const FString& BlendShapePath)
@@ -356,41 +583,6 @@ namespace UE::InterchangeUsdTranslator::Private
 		}
 	}
 
-	void UpdateTraversalInfo(FTraversalInfo& Info, const UE::FUsdPrim& CurrentPrim)
-	{
-		if (CurrentPrim.IsA(TEXT("SkelRoot")))
-		{
-			if (!Info.ClosestParentSkelRoot)
-			{
-				// The root-most skel cache should handle any nested UsdSkel prims as well
-				Info.FurthestSkelCache = MakeShared<UE::FUsdSkelCache>();
-
-				const bool bTraverseInstanceProxies = true;
-				Info.FurthestSkelCache->Populate(CurrentPrim, bTraverseInstanceProxies);
-			}
-
-			Info.ClosestParentSkelRoot = CurrentPrim;
-		}
-
-		if (Info.ClosestParentSkelRoot && CurrentPrim.HasAPI(TEXT("SkelBindingAPI")))
-		{
-			UE::FUsdStage Stage = CurrentPrim.GetStage();
-
-			if (UE::FUsdRelationship SkelRel = CurrentPrim.GetRelationship(TEXT("skel:skeleton")))
-			{
-				TArray<UE::FSdfPath> Targets;
-				if (SkelRel.GetTargets(Targets) && Targets.Num() > 0)
-				{
-					UE::FUsdPrim TargetSkeleton = Stage.GetPrimAtPath(Targets[0]);
-					if (TargetSkeleton && TargetSkeleton.IsA(TEXT("Skeleton")))
-					{
-						Info.ActiveSkelQuery = Info.FurthestSkelCache->GetSkelQuery(TargetSkeleton);
-					}
-				}
-			}
-		}
-	}
-
 	bool ReadBools(
 		const UE::FUsdStage& UsdStage,
 		const TArray<double>& UsdTimeSamples,
@@ -693,10 +885,6 @@ namespace UE::InterchangeUsdTranslator::Private
 		UInterchangeBaseNodeContainer& NodeContainer
 	)
 	{
-		FString PrimPath = Prim.GetPrimPath().GetString();
-		FString MaterialUid = MaterialPrefix + PrimPath;
-		FString MaterialPrimName(Prim.GetName().ToString());
-
 		FName RenderContext = TranslatorSettings ? TranslatorSettings->RenderContext : UnrealIdentifiers::UniversalRenderContext;
 
 		// Check for any references of MaterialX
@@ -1088,6 +1276,199 @@ namespace UE::InterchangeUsdTranslator::Private
 		return false;
 	}
 
+	UE::FUsdPrim UInterchangeUSDTranslatorImpl::TryGettingInactiveLODPrim(const FString& PrimPathString)
+	{
+		if (PrimPathString.IsEmpty())
+		{
+			return {};
+		}
+
+		UE::FSdfPath PrimPath{*PrimPathString};
+
+		const FString PrimName = PrimPath.GetName();
+		if (!IsValidLODName(PrimName))
+		{
+			return {};
+		}
+
+		const UE::FSdfPath LODContainerPath = PrimPath.GetParentPath();
+		const FString LODContainerPathString = LODContainerPath.GetString();
+		const FString& VariantName = PrimName;	  // Our convention for LODs is that the prim name matches the variant (e.g. "LOD2")
+
+		UE::FUsdStage TempStage;
+
+		{
+			FWriteScopeLock Lock{PrimPathToVariantToStageLock};
+
+			// Check if we have the stage we want already
+			if (TMap<FString, UE::FUsdStage>* TempStagesForPrim = PrimPathToVariantToStage.Find(LODContainerPathString))
+			{
+				if (UE::FUsdStage* TempStagesForVariant = TempStagesForPrim->Find(VariantName))
+				{
+					TempStage = *TempStagesForVariant;
+				}
+			}
+
+			// If not, let's open a new masked stage.
+			// USD will only compose the LOD container prim (and its subtree), so this should be very small and quick to open.
+			// We won't use the stage cache for these, so our strong reference right here is the only thing holding the stage opened.
+			if (!TempStage)
+			{
+				TempStage = UnrealUSDWrapper::OpenMaskedStage(
+					*UsdStage.GetRootLayer().GetIdentifier(),
+					EUsdInitialLoadSet::LoadAll,
+					{LODContainerPathString}
+				);
+
+				if (!TempStage)
+				{
+					return {};
+				}
+
+				UE::FUsdPrim LODContainerPrim = TempStage.GetPrimAtPath(LODContainerPath);
+				if (!LODContainerPrim)
+				{
+					return {};
+				}
+
+				// We have to edit the session layer here, and not the root layer directly. This because USD only opens a layer
+				// once in memory, so if we have multiple of these temp stages all trying to set the variant to different values
+				// on the root layer itself, they'd be actually trying to overwrite each other and could even lead to threading issues.
+				//
+				// The session layer however is unique to each of these temp stages so we won't have that problem, and it still
+				// should compose the variant switch just the same.
+				UE::FUsdEditContext Context{TempStage, TempStage.GetSessionLayer()};
+
+				UE::FUsdVariantSets VariantSets = LODContainerPrim.GetVariantSets();
+				bool bSwitched = VariantSets.SetSelection(LODString, VariantName);
+				if (!bSwitched)
+				{
+					return {};
+				}
+
+				// We finally have our stage for the particular lod container and with the variant we want: Cache it for later, if needed.
+				PrimPathToVariantToStage.FindOrAdd(LODContainerPathString).Add(VariantName, TempStage);
+			}
+		}
+
+		return TempStage.GetPrimAtPath(PrimPath);
+	}
+
+	void UInterchangeUSDTranslatorImpl::FinalizeLODContainerTraversal(
+		UInterchangeBaseNodeContainer& NodeContainer,
+		const FTraversalInfo& Info,
+		const UInterchangeSceneNode* SceneNodeWithLODs
+	)
+	{
+		if (!SceneNodeWithLODs)
+		{
+			return;
+		}
+
+		// Add a dedicated LODContainer node. We need this because every child of a LodGroup node
+		// will be interpreted as a LOD mesh, and we could have ended up with any additional children so far
+		// since any regular prim can contain the LOD variant (it could have any number of other children like
+		// other transforms, lights, skeleton prims, etc.)
+		const FString LODContainerUid = SceneNodeWithLODs->GetUniqueID() + LODContainerSuffix;
+		UInterchangeSceneNode* LODContainer = NewObject<UInterchangeSceneNode>(&NodeContainer);
+		{
+			LODContainer->InitializeNode(LODContainerUid, SceneNodeWithLODs->GetDisplayLabel(), EInterchangeNodeContainerType::TranslatedScene);
+			LODContainer->AddSpecializedType(UE::Interchange::FSceneNodeStaticData::GetLodGroupSpecializeTypeString());
+
+			NodeContainer.AddNode(LODContainer);
+			NodeContainer.SetNodeParentUid(LODContainerUid, SceneNodeWithLODs->GetUniqueID());
+		}
+
+		struct FLODSortHelper
+		{
+			UInterchangeSceneNode* SceneNode = nullptr;
+			int32 LODIndex;	   // Using actual ints in here instead of sorting on the node DisplayLabel directly so we correctly sort LOD10 > LOD2
+		};
+
+		TArray<FLODSortHelper> SortedLODs;
+		SortedLODs.Reserve(CurrentLODSceneNodes.Num());
+		for (UInterchangeSceneNode* Node : CurrentLODSceneNodes)
+		{
+			SortedLODs.Add({Node, GetLODIndexFromName(Node->GetDisplayLabel())});
+		}
+		SortedLODs.Sort(
+			[](const FLODSortHelper& LHS, const FLODSortHelper& RHS)
+			{
+				return LHS.LODIndex < RHS.LODIndex;
+			}
+		);
+
+		// Parent our LOD nodes to the LODContainer in the right order, because Interchange will try assigning
+		// those meshes to LOD numbers in the order it traverses the children, and there's no guarantee about the order we run
+		// into the variants when traversing the stage (or the order with which they were authored)
+		for (int32 Index = 0; Index < CurrentLODSceneNodes.Num(); ++Index)
+		{
+			const FLODSortHelper& PackedNode = SortedLODs[Index];
+			const FString& NodeUID = PackedNode.SceneNode->GetUniqueID();
+
+			NodeContainer.SetNodeParentUid(NodeUID, LODContainerUid);
+			NodeContainer.SetNodeDesiredChildIndex(NodeUID, Index);
+		}
+
+		// Detect invalid setup of LOD morph targets: Morph targets used by LOD2 must be present in LOD1, and all of those must be present in LOD0.
+		// On legacy USD support we used a hack to work around that, but it's not so trivial to do in Interchange and it's probably the wrong approach
+		// anyway: we should just let the user know how to correct their data for Unreal instead
+		bool bIsFirstLOD = true;
+		TSet<const UInterchangeMeshNode*> AllowedMorphTargets;
+		for (const FLODSortHelper& SortedLODNode : SortedLODs)
+		{
+			FString AssetInstanceUid;
+			if (!SortedLODNode.SceneNode->GetCustomAssetInstanceUid(AssetInstanceUid))
+			{
+				continue;
+			}
+
+			const UInterchangeMeshNode* MeshNode = Cast<const UInterchangeMeshNode>(NodeContainer.GetNode(AssetInstanceUid));
+			if (!MeshNode)
+			{
+				continue;
+			}
+
+			// Only skinned meshes have blend shapes / morph targets in USD
+			if (!MeshNode->IsSkinnedMesh())
+			{
+				continue;
+			}
+
+			TArray<FString> LODMorphTargetUids;
+			MeshNode->GetMorphTargetDependencies(LODMorphTargetUids);
+
+			TSet<const UInterchangeMeshNode*> LODMorphTargets;
+			LODMorphTargets.Reserve(LODMorphTargetUids.Num());
+
+			for (const FString& MorphTargetUid : LODMorphTargetUids)
+			{
+				if (const UInterchangeMeshNode* MorphTargetNode = Cast<const UInterchangeMeshNode>(NodeContainer.GetNode(MorphTargetUid)))
+				{
+					LODMorphTargets.Add(MorphTargetNode);
+				}
+			}
+
+			if (!bIsFirstLOD && !AllowedMorphTargets.Includes(LODMorphTargets))
+			{
+				UE_LOG(
+					LogUsd,
+					Warning,
+					TEXT(
+						"Invalid blend shape configuration for skeletal mesh LOD '%s': The set of blend shapes names used by lower LODs should include all blend shape names used by higher LODs"
+					),
+					*MeshNode->GetUniqueID()
+				);
+				break;
+			}
+
+			Swap(AllowedMorphTargets, LODMorphTargets);
+			bIsFirstLOD = false;
+		}
+
+		CurrentLODSceneNodes.Reset();
+	}
+
 	void AddDisplayColorMaterialInstanceNodeIfNeeded(UInterchangeBaseNodeContainer& NodeContainer, const FString& DisplayColorDesc)
 	{
 		using namespace UsdUnreal::MaterialUtils;
@@ -1304,7 +1685,7 @@ namespace UE::InterchangeUsdTranslator::Private
 		const FString MeshPrimPath = MeshPrim.GetPrimPath().GetString();
 
 		TFunction<void(const FString&, int32, const FString&)> AddMorphTargetNode =
-			[&MeshNode, &MeshPrimPath, &NodeContainer](const FString& MorphTargetName, int32 BlendShapeIndex, const FString& InbetweenName)
+			[&MeshNode, &MeshPrimPath, &NodeContainer, &Info](const FString& MorphTargetName, int32 BlendShapeIndex, const FString& InbetweenName)
 		{
 			// Note: We identify a blend shape by its Mesh prim path and the blend shape index, even though
 			// the blend shape itself is a full standalone prim. This is for two reasons:
@@ -1316,7 +1697,7 @@ namespace UE::InterchangeUsdTranslator::Private
 			// blendshapes across the model to combine into a single morph target. Interchange has an import option to let you control
 			// whether they become separate morph targets or not anyway ("Merge Morph Targets with Same Name")
 			const FString NodeUid = GetMorphTargetMeshNodeUid(MeshPrimPath, BlendShapeIndex, InbetweenName);
-			const FString PayloadKey = GetMorphTargetMeshPayloadKey(MeshPrimPath, BlendShapeIndex, InbetweenName);
+			const FString PayloadKey = GetMorphTargetMeshPayloadKey(Info.bInsideLOD, MeshPrimPath, BlendShapeIndex, InbetweenName);
 
 			UInterchangeMeshNode* MorphTargetMeshNode = NewObject<UInterchangeMeshNode>(&NodeContainer);
 			MorphTargetMeshNode->InitializeNode(NodeUid, MorphTargetName, EInterchangeNodeContainerType::TranslatedAsset);
@@ -1370,21 +1751,34 @@ namespace UE::InterchangeUsdTranslator::Private
 		UInterchangeMeshNode* MeshNode = NewObject<UInterchangeMeshNode>(&NodeContainer);
 		MeshNode->InitializeNode(NodeUid, NodeName, EInterchangeNodeContainerType::TranslatedAsset);
 		MeshNode->SetAssetName(NodeName);
-		const bool bIsSkinned = static_cast<bool>(Info.ClosestParentSkelRoot) && Prim.HasAPI(TEXT("SkelBindingAPI"));
+		const bool bIsSkinned = static_cast<bool>(Info.ClosestParentSkelRootPath) && Prim.HasAPI(TEXT("SkelBindingAPI"));
 
-		// We are currently not supporting Skinned Primitive Shapes.
-		// In theory a Skinned Mesh needs Joint influences and weights information provided,
-		//		However, there does not seem to be a ruleset against a Primitive Shape having SkelBindingAPI set.
-		//		Which means, that on a theoretical level there could be such a scenario.
-		ensureMsgf(!(bPrimitiveShape && bIsSkinned), TEXT("Unexpected scenario: Primitive Shape is skinned."));
+		FString PayloadKey = PrimPath;
+		if (Info.bInsideLOD)
+		{
+			PayloadKey = LODPrefix + PayloadKey;
+		}
+		if (bPrimitiveShape)
+		{
+			// We are currently not supporting Skinned Primitive Shapes.
+			// In theory a Skinned Mesh needs Joint influences and weights information provided,
+			// However, there does not seem to be a ruleset against a Primitive Shape having SkelBindingAPI set.
+			// Which means, that on a theoretical level there could be such a scenario.
+			ensureMsgf(!bIsSkinned, TEXT("Unexpected scenario: Primitive Shape is skinned."));
+
+			// For Primitive Shapes we add PrimitiveShapePrefix for the PayLoadKey,
+			// in order to be able to identify the PrimitiveShape in the PayloadData acquisition phase,
+			// as the Primitive Shapes require a different MeshDescription acquisition path, compared to Static Meshes.
+			PayloadKey = PrimitiveShapePrefix + PayloadKey;
+		}
 
 		if (bIsSkinned && !bPrimitiveShape)
 		{
 			MeshNode->SetSkinnedMesh(true);
-			MeshNode->SetPayLoadKey(PrimPath, EInterchangeMeshPayLoadType::SKELETAL);
-			if (Info.ActiveSkelQuery)
+			MeshNode->SetPayLoadKey(PayloadKey, EInterchangeMeshPayLoadType::SKELETAL);
+			if (Info.BoundSkeletonPrimPath && !Info.BoundSkeletonPrimPath->IsEmpty())
 			{
-				MeshNode->SetSkeletonDependencyUid(Info.ActiveSkelQuery.GetSkeleton().GetPrimPath().GetString());
+				MeshNode->SetSkeletonDependencyUid(*Info.BoundSkeletonPrimPath);
 			}
 
 			AddMorphTargetNodes(Prim, *this, *MeshNode, NodeContainer, Info);
@@ -1398,10 +1792,7 @@ namespace UE::InterchangeUsdTranslator::Private
 		}
 		else
 		{
-			// For Primitive Shapes we add PrimitiveShapePrefix for the PayLoadKey,
-			// in order to be able to identify the PrimitiveShape in the PayloadData acquisition phase,
-			// as the Primitive Shapes require a different MeshDescription acquisition path, compared to Static Meshes.
-			MeshNode->SetPayLoadKey((bPrimitiveShape ? PrimitiveShapePrefix : TEXT("")) + PrimPath, EInterchangeMeshPayLoadType::STATIC);
+			MeshNode->SetPayLoadKey(PayloadKey, EInterchangeMeshPayLoadType::STATIC);
 		}
 
 		// Material assignments
@@ -1416,6 +1807,11 @@ namespace UE::InterchangeUsdTranslator::Private
 				CachedMeshConversionOptions.MaterialPurpose
 			);
 
+			if (Info.bInsideLOD)
+			{
+				CachedMaterialAssignments.Add(PrimPath, Assignments);
+			}
+
 			for (const UsdUtils::FUsdPrimMaterialSlot& Slot : Assignments.Slots)
 			{
 				// Use the material prim path/display color desc as the material slot name, because Interchange
@@ -1479,6 +1875,88 @@ namespace UE::InterchangeUsdTranslator::Private
 		NodeContainer.AddNode(MeshNode);
 	}
 
+	void UInterchangeUSDTranslatorImpl::AddLODMeshNodes(	//
+		const UE::FUsdPrim& Prim,
+		UInterchangeBaseNodeContainer& NodeContainer,
+		UInterchangeSceneNode& ParentSceneNode,
+		UInterchangeUsdTranslatorSettings* TranslatorSettings,
+		FTraversalInfo Info
+	)
+	{
+		UE::FUsdVariantSets VariantSets = Prim.GetVariantSets();
+		const UE::FUsdVariantSet LODVariantSet = VariantSets.GetVariantSet(LODString);
+		if (!LODVariantSet.IsValid())
+		{
+			return;
+		}
+
+		const FString ActiveVariant = LODVariantSet.GetVariantSelection();
+		if (!IsValidLODName(ActiveVariant))
+		{
+			return;
+		}
+
+		Info.ParentNode = &ParentSceneNode;
+
+		bool bSwitchedFromInitialVariant = false;
+		for (const FString& VariantName : LODVariantSet.GetVariantNames())
+		{
+			// The active variant will be parsed via regular traversal
+			if (VariantName == ActiveVariant || !IsValidLODName(VariantName))
+			{
+				continue;
+			}
+
+			{
+				// For creating the scene nodes themselves we'll switch the active variant on the currently opened stage
+				// (still using the session layer to minimize impact to the actual layer). This mainly so that we can retrieve
+				// and fetch and cache the correct material bindings for the LOD meshes. Later on we'll use separate stages
+				// with population masks to read the LODs concurrently, and we won't be able to resolve material bindings
+				UE::FUsdEditContext Context{UsdStage, UsdStage.GetSessionLayer()};
+
+				bool bSwitchedVariant = VariantSets.SetSelection(LODString, VariantName);
+				if (!bSwitchedVariant)
+				{
+					continue;
+				}
+				bSwitchedFromInitialVariant = true;
+			}
+
+			if (UE::FUsdPrim LODMeshPrim = GetLODMesh(Prim, VariantName))
+			{
+				Traverse(LODMeshPrim, *this, NodeContainer, TranslatorSettings, Info);
+			}
+			else
+			{
+				UE_LOG(
+					LogTemp,
+					Warning,
+					TEXT(
+						"Failed to parse a LOD Mesh from variant '%s' of prim '%s'. For automatic parsing of LODs, make sure there is a single Mesh prim within the variant, named exactly as the variant itself (e.g. 'LOD0', 'LOD1', etc.)"
+					),
+					*VariantName,
+					*Prim.GetPrimPath().GetString()
+				);
+			}
+		}
+
+		// Put the active variant back to what it originally was
+		if (bSwitchedFromInitialVariant)
+		{
+			{
+				UE::FUsdEditContext Context{UsdStage, UsdStage.GetSessionLayer()};
+
+				const bool bRestoredSelection = VariantSets.SetSelection(LODString, ActiveVariant);
+				ensure(bRestoredSelection);
+			}
+
+			// Recompute our skel cache here if we have any as ancestor, because switching variants could have
+			// invalidated some of its internal state about its descendant prims, which we'll need to be OK
+			// when handling the payloads
+			Info.RepopulateSkelCache(UsdStage);
+		}
+	}
+
 	void AddTrackSetNode(UInterchangeUSDTranslatorImpl& Impl, UInterchangeBaseNodeContainer& NodeContainer)
 	{
 		// For now we only want a single track set (i.e. LevelSequence) per stage.
@@ -1719,7 +2197,7 @@ namespace UE::InterchangeUsdTranslator::Private
 			UE::FUsdSkelBinding SkelBinding;
 			const bool bTraverseInstanceProxies = true;
 			bool bSuccess = Info.FurthestSkelCache->ComputeSkelBinding(	   //
-				Info.ClosestParentSkelRoot,
+				Info.ResolveClosestParentSkelRoot(Stage),
 				SkeletonPrim,
 				SkelBinding,
 				bTraverseInstanceProxies
@@ -1823,7 +2301,7 @@ namespace UE::InterchangeUsdTranslator::Private
 	)
 	{
 		// If we're not inside of a SkelRoot, the skeleton shouldn't really do anything
-		if (!Info.FurthestSkelCache.IsValid())
+		if (!Info.ClosestParentSkelRootPath.IsValid())
 		{
 			return;
 		}
@@ -1956,7 +2434,7 @@ namespace UE::InterchangeUsdTranslator::Private
 		FString DisplayLabel(Prim.GetName().ToString());
 
 		// Do this before generating other nodes as they may need the updated info
-		UpdateTraversalInfo(Info, Prim);
+		Info.UpdateWithCurrentPrim(Prim);
 
 		// Generate asset node if applicable
 		const FString* Prefix = nullptr;
@@ -1982,8 +2460,8 @@ namespace UE::InterchangeUsdTranslator::Private
 		}
 		else if (Prim.IsA(TEXT("Gprim")) && !Prim.IsA(TEXT("PointBased")) && !Prim.IsA(TEXT("Volume")))
 		{
-			//PointBased prims are currently not supported apart from Meshes (which is taken care of in a previous branch).
-			//Volumes are also not currently supported.
+			// PointBased prims are currently not supported apart from Meshes (which is taken care of in a previous branch).
+			// Volumes are also not currently supported.
 
 			Prefix = &PrimitiveShapePrefix;
 			constexpr bool bPrimitiveShape_True = true;
@@ -1994,9 +2472,11 @@ namespace UE::InterchangeUsdTranslator::Private
 		// This includes Xforms but also Scopes, which are not Xformable.
 		//
 		// We add a scene node for the pseudoroot in order to make tree traversal easier on the pipeline,
-		// but the pipeline will strip the pseudoroot node as it's final step
+		// but the pipeline will strip the pseudoroot node as its final step (if desired)
 		const bool bIsImageable = Prim.IsA(TEXT("Imageable"));
-		const bool bNeedsSceneNode = Prim.IsPseudoRoot() || (bIsImageable && (Info.ParentNode || Prim.GetParent().IsPseudoRoot()));
+		const bool bNeedsSceneNode = Prim.IsPseudoRoot()														  //
+									 || (bIsImageable && (Info.ParentNode || Prim.GetParent().IsPseudoRoot()))	  //
+									 || Info.bIsLODContainer;	 // We allow typeless prims (so not imageable) to be LOD containers
 
 		UInterchangeSceneNode* SceneNode = nullptr;
 		if (bNeedsSceneNode)
@@ -2012,19 +2492,12 @@ namespace UE::InterchangeUsdTranslator::Private
 				SceneNode->AddInt32Attribute(UE::Interchange::USD::GeometryPurposeIdentifier, (int32)PrimPurpose);
 			}
 
-			// If we're an Xformable, get our transform
+			// If we're an Xformable, get our transform.
+			// All SceneNodes should have their LocalTransform set though.
+			// Not setting will cause ensure hits in Skeleton generations for example.
 			FTransform Transform = FTransform::Identity;
 			bool bResetTransformStack = false;
-			UsdToUnreal::ConvertXformable(
-				Prim.GetStage(),
-				UE::FUsdTyped(Prim),
-				Transform,
-				UsdUtils::GetEarliestTimeCode(),
-				&bResetTransformStack
-			);
-
-			// All SceneNodes should have their LocalTransform set.
-			//	Not setting will cause ensure hits in Skeleton generations for example.
+			UsdToUnreal::ConvertXformable(Prim.GetStage(), UE::FUsdTyped(Prim), Transform, UsdUtils::GetEarliestTimeCode(), &bResetTransformStack);
 			SceneNode->SetCustomLocalTransform(&NodeContainer, Transform);
 
 			// Hide our scene node if it is meant to be hidden in USD
@@ -2034,11 +2507,19 @@ namespace UE::InterchangeUsdTranslator::Private
 			// "scene component stuff" ends up as actor tracks), and also due to how it behaves better for cameras: Component
 			// visibility for camera nodes would hide the camera component itself, which has no effect. Actor visibility for
 			// camera actors does hide the entire camera actor however
-			const bool bComputedVisibility = Info.bVisible && UsdUtils::HasInheritedVisibility(Prim);
-			if (!bComputedVisibility)
+			if (!Info.bVisible)
+			{
+				SceneNode->SetCustomActorVisibility(Info.bVisible);
+			}
+
+			if (Info.bIsLODContainer)
+			{
+				TranslatorImpl.AddLODMeshNodes(Prim, NodeContainer, *SceneNode, TranslatorSettings, Info);
+			}
+
+			if (Info.bInsideLOD && Prefix == &MeshPrefix)
 			{
-				SceneNode->SetCustomActorVisibility(bComputedVisibility);
-				Info.bVisible = bComputedVisibility;
+				TranslatorImpl.CurrentLODSceneNodes.Add(SceneNode);
 			}
 
 			// Skeleton joints are separate scene nodes in Interchange, so we need to emit that node hierarchy now
@@ -2073,9 +2554,6 @@ namespace UE::InterchangeUsdTranslator::Private
 
 		// Recurse into child prims
 		{
-			// Note: This has the effect of effectively shutting down the generation of scene nodes
-			// below any prim that is not a least an Imageable, as we check for a valid parent before
-			// generating one
 			Info.ParentNode = SceneNode;
 
 			for (const FUsdPrim& ChildPrim : Prim.GetChildren())
@@ -2083,74 +2561,111 @@ namespace UE::InterchangeUsdTranslator::Private
 				Traverse(ChildPrim, TranslatorImpl, NodeContainer, TranslatorSettings, Info);
 			}
 		}
+
+		// Finalize the LOD container after recursing normally, because we'll rely on the regular traversal to
+		// process the mesh of the active LOD variant. AddLODMeshNodes only handles the inactive variants
+		if (Info.bIsLODContainer)
+		{
+			TranslatorImpl.FinalizeLODContainerTraversal(NodeContainer, Info, SceneNode);
+		}
 	}
 
 	bool GetStaticMeshPayloadData(
-		const FString& PayloadKey,
-		const UInterchangeUSDTranslatorImpl& Impl,
+		FString PayloadKey,
+		UInterchangeUSDTranslatorImpl& Impl,
 		const UsdToUnreal::FUsdMeshConversionOptions& Options,
 		FMeshDescription& OutMeshDescription
 	)
 	{
-		bool bIsPrimitiveShape = PayloadKey.StartsWith(PrimitiveShapePrefix);
-
-		FString PrimPath = bIsPrimitiveShape ? PayloadKey.RightChop(PrimitiveShapePrefix.Len()) : PayloadKey;
+		const bool bIsLODMesh = CheckAndChopPayloadPrefix(PayloadKey, LODPrefix);
+		const bool bIsPrimitiveShape = CheckAndChopPayloadPrefix(PayloadKey, PrimitiveShapePrefix);
 
+		const FString& PrimPath = PayloadKey;
 		UE::FUsdPrim Prim = Impl.UsdStage.GetPrimAtPath(UE::FSdfPath{*PrimPath});
 		if (!Prim)
 		{
-			return false;
+			if (bIsLODMesh)
+			{
+				Prim = Impl.TryGettingInactiveLODPrim(PrimPath);
+			}
+
+			if (!Prim)
+			{
+				return false;
+			}
 		}
 
-		// TODO: We can't do much with these yet: There will be used to generate primvar-compatible
+		FMeshDescription TempMeshDescription;
+		FStaticMeshAttributes StaticMeshAttributes(TempMeshDescription);
+		StaticMeshAttributes.Register();
+
+		// TODO: We can't do much with these yet: They will be used to generate primvar-compatible
 		// versions of the materials that are assigned to this mesh, whenever we get a pipeline
 		UsdUtils::FUsdPrimMaterialAssignmentInfo TempMaterialInfo;
-		bool bSuccess = false;
 
+		bool bSuccess = false;
 		if (bIsPrimitiveShape)
 		{
-			FMeshDescription TempMeshDescription;
-			UsdUtils::FUsdPrimMaterialAssignmentInfo PrimitiveTempMaterialInfo;
-
-			FStaticMeshAttributes StaticMeshAttributes(TempMeshDescription);
-			StaticMeshAttributes.Register();
-
-			bSuccess = UsdToUnreal::ConvertGeomPrimitive(Prim, TempMeshDescription, PrimitiveTempMaterialInfo, Options);
-			if (bSuccess)
-			{
-				OutMeshDescription = { MoveTemp(TempMeshDescription) };
-				TempMaterialInfo = { MoveTemp(PrimitiveTempMaterialInfo) };
-			}
+			bSuccess = UsdToUnreal::ConvertGeomPrimitive(Prim, TempMeshDescription, TempMaterialInfo, Options);
 		}
 		else
 		{
-			bSuccess = UsdToUnreal::ConvertGeomMesh(Prim, OutMeshDescription, TempMaterialInfo, Options);
+			bSuccess = UsdToUnreal::ConvertGeomMesh(Prim, TempMeshDescription, TempMaterialInfo, Options);
 		}
-
 		if (!bSuccess)
 		{
 			return false;
 		}
 
-		FixMaterialSlotNames(OutMeshDescription, TempMaterialInfo.Slots);
+		UsdUtils::FUsdPrimMaterialAssignmentInfo* AssignmentPtr = &TempMaterialInfo;
+		if (bIsLODMesh)
+		{
+			// Use our cached material assignments instead of whatever we pull from ConvertGeomMesh
+			// because if we're in an LOD mesh then we may be reading from a temp stage, that has a
+			// population mask that may not include the material, meaning ConvertGeomMesh may have failed
+			// to resolve all the bindings. The cached assignments come from AddMeshNode step, where
+			// we switch the active variant on the current stage and so get nice material bindings that
+			// resolve normally.
+			//
+			// Note that we can't even use the info cache here, because it wouldn't have cached info
+			// about the inactive LOD variants
+			if (UsdUtils::FUsdPrimMaterialAssignmentInfo* CachedInfo = Impl.CachedMaterialAssignments.Find(PrimPath))
+			{
+				AssignmentPtr = CachedInfo;
+			}
+		}
+
+		OutMeshDescription = MoveTemp(TempMeshDescription);
+
+		FixMaterialSlotNames(OutMeshDescription, AssignmentPtr->Slots);
 
 		return true;
 	}
 
 	bool GetSkeletalMeshPayloadData(
-		const FString& PayloadKey,
-		const UInterchangeUSDTranslatorImpl& Impl,
+		FString PayloadKey,
+		UInterchangeUSDTranslatorImpl& Impl,
 		const UsdToUnreal::FUsdMeshConversionOptions& Options,
 		FMeshDescription& OutMeshDescription,
 		TArray<FString>& OutJointNames
 	)
 	{
 #if WITH_EDITOR
+		const bool bIsLODMesh = CheckAndChopPayloadPrefix(PayloadKey, LODPrefix);
+
 		const FString& PrimPath = PayloadKey;
 		UE::FUsdPrim Prim = Impl.UsdStage.GetPrimAtPath(UE::FSdfPath{*PrimPath});
 		if (!Prim)
 		{
-			return false;
+			if (bIsLODMesh)
+			{
+				Prim = Impl.TryGettingInactiveLODPrim(PrimPath);
+			}
+
+			if (!Prim)
+			{
+				return false;
+			}
 		}
 
 		const FString& MeshNodeUid = MeshPrefix + Prim.GetPrimPath().GetString();
@@ -2166,7 +2681,8 @@ namespace UE::InterchangeUsdTranslator::Private
 			{
 				return false;
 			}
-			SkelQuery = MeshInfo->ActiveSkelQuery;
+
+			SkelQuery = MeshInfo->ResolveSkelQuery(Impl.UsdStage);
 			if (!SkelQuery)
 			{
 				return false;
@@ -2174,7 +2690,7 @@ namespace UE::InterchangeUsdTranslator::Private
 
 			// The above fields are associated to the mesh *asset* node Uid (hence the prefix),
 			// while the joint names are associated to the skeleton *scene* node Uid, so no prefix
-			const FString SkeletonNodeUid = SkelQuery.GetSkeleton().GetPrimPath().GetString();
+			const FString& SkeletonNodeUid = *MeshInfo->BoundSkeletonPrimPath;
 			const FTraversalInfo* SkeletonInfo = Impl.NodeUidToCachedTraversalInfo.Find(SkeletonNodeUid);
 			if (!SkeletonInfo)
 			{
@@ -2208,15 +2724,36 @@ namespace UE::InterchangeUsdTranslator::Private
 		// be done by Interchange later..
 		const USkeletalMesh* UnusedSkelMesh = nullptr;
 		FSkeletalMeshBuildSettings* UnusedBuildSettings = nullptr;
-		bSuccess = SkelMeshImportData.GetMeshDescription(UnusedSkelMesh, UnusedBuildSettings, OutMeshDescription);
+		FMeshDescription TempMeshDescription;
+		bSuccess = SkelMeshImportData.GetMeshDescription(UnusedSkelMesh, UnusedBuildSettings, TempMeshDescription);
 		if (!bSuccess)
 		{
 			return false;
 		}
 
+		UsdUtils::FUsdPrimMaterialAssignmentInfo* AssignmentPtr = &TempMaterialInfo;
+		if (bIsLODMesh)
+		{
+			// Use our cached material assignments instead of whatever we pull from ConvertSkinnedMesh
+			// because if we're in an LOD mesh then we may be reading from a temp stage, that has a
+			// population mask that may not include the material, meaning ConvertSkinnedMesh may have failed
+			// to resolve all the bindings. The cached assignments come from AddMeshNode step, where
+			// we switch the active variant on the current stage and so get nice material bindings that
+			// resolve normally.
+			//
+			// Note that we can't even use the info cache here, because it wouldn't have cached info
+			// about the inactive LOD variants
+			if (UsdUtils::FUsdPrimMaterialAssignmentInfo* CachedInfo = Impl.CachedMaterialAssignments.Find(PrimPath))
+			{
+				AssignmentPtr = CachedInfo;
+			}
+		}
+
+		OutMeshDescription = MoveTemp(TempMeshDescription);
+
 		FixSkeletalMeshDescriptionColors(OutMeshDescription);
 
-		FixMaterialSlotNames(OutMeshDescription, TempMaterialInfo.Slots);
+		FixMaterialSlotNames(OutMeshDescription, AssignmentPtr->Slots);
 
 		OutJointNames = *JointNames;
 
@@ -2227,35 +2764,36 @@ namespace UE::InterchangeUsdTranslator::Private
 	}
 
 	bool GetMorphTargetPayloadData(
-		const FString& PayloadKey,
-		const UInterchangeUSDTranslatorImpl& Impl,
+		FString PayloadKey,
+		UInterchangeUSDTranslatorImpl& Impl,
 		const UsdToUnreal::FUsdMeshConversionOptions& Options,
 		FMeshDescription& OutMeshDescription,
 		FString& OutMorphTargetName
 	)
 	{
-		// These payload keys are generated by GetMorphTargetMeshPayloadKey(), and so should take the form
-		// "<mesh prim path>\<mesh blend shape index>\<optional inbetween name>"
-		const bool bCullEmpty = false;
-		TArray<FString> PayloadKeyTokens;
-		PayloadKey.ParseIntoArray(PayloadKeyTokens, TEXT("\\"), bCullEmpty);
-		if (PayloadKeyTokens.Num() != 3)
+		bool bIsLODMesh = false;
+		FString MeshPrimPath;
+		int32 BlendShapeIndex = INDEX_NONE;
+		FString InbetweenName;
+		if (!ParseMorphTargetMeshPayloadKey(PayloadKey, bIsLODMesh, MeshPrimPath, BlendShapeIndex, InbetweenName))
 		{
 			return false;
 		}
 
-		const FString& MeshPrimPath = PayloadKeyTokens[0];
-		const FString& BlendShapeIndexStr = PayloadKeyTokens[1];
-		const FString& InbetweenName = PayloadKeyTokens[2];
-
-		int32 BlendShapeIndex = INDEX_NONE;
-		bool bLexed = LexTryParseString(BlendShapeIndex, *BlendShapeIndexStr);
-		if (!bLexed)
+		UE::FUsdPrim MeshPrim = Impl.UsdStage.GetPrimAtPath(FSdfPath{*MeshPrimPath});
+		if (!MeshPrim)
 		{
-			return false;
+			if (bIsLODMesh)
+			{
+				MeshPrim = Impl.TryGettingInactiveLODPrim(MeshPrimPath);
+			}
+
+			if (!MeshPrim)
+			{
+				return false;
+			}
 		}
 
-		UE::FUsdPrim MeshPrim = Impl.UsdStage.GetPrimAtPath(FSdfPath{*MeshPrimPath});
 		UE::FUsdSkelBlendShapeQuery Query{MeshPrim};
 		if (!Query)
 		{
@@ -2271,7 +2809,8 @@ namespace UE::InterchangeUsdTranslator::Private
 		// TODO: This is extremely slow, as it will reimport the mesh for every single morph target!
 		// It seems to be what the other translators do, however. We need some form of FMeshDescription caching here
 		TArray<FString> UnusedJointNames;
-		bool bConverted = GetSkeletalMeshPayloadData(MeshPrimPath, Impl, Options, OutMeshDescription, UnusedJointNames);
+		const FString MeshPayloadKey = bIsLODMesh ? LODPrefix + MeshPrimPath : MeshPrimPath;
+		bool bConverted = GetSkeletalMeshPayloadData(MeshPayloadKey, Impl, Options, OutMeshDescription, UnusedJointNames);
 		if (!bConverted || OutMeshDescription.IsEmpty())
 		{
 			return false;
@@ -2373,7 +2912,8 @@ namespace UE::InterchangeUsdTranslator::Private
 			{
 				return false;
 			}
-			SkelQuery = MeshInfo->ActiveSkelQuery;
+
+			SkelQuery = MeshInfo->ResolveSkelQuery(Impl.UsdStage);
 			if (!SkelQuery)
 			{
 				return false;
@@ -2506,7 +3046,8 @@ namespace UE::InterchangeUsdTranslator::Private
 				{
 					return false;
 				}
-				SkelQuery = MeshInfo->ActiveSkelQuery;
+
+				SkelQuery = MeshInfo->ResolveSkelQuery(Impl.UsdStage);
 				if (!SkelQuery)
 				{
 					return false;
@@ -2766,6 +3307,11 @@ void UInterchangeUSDTranslator::ReleaseSource()
 	ImplPtr->UsdStage = UE::FUsdStage{};
 	ImplPtr->CurrentTrackSet = nullptr;
 
+	{
+		FWriteScopeLock Lock{ImplPtr->PrimPathToVariantToStageLock};
+		ImplPtr->PrimPathToVariantToStage.Reset();
+	}
+
 	if (TranslatorSettings)
 	{
 		TranslatorSettings->ClearFlags(RF_Standalone);
@@ -2851,6 +3397,9 @@ TOptional<UE::Interchange::FMeshPayloadData> UInterchangeUSDTranslator::GetMeshP
 
 	UsdToUnreal::FUsdMeshConversionOptions OptionsCopy = ImplPtr->CachedMeshConversionOptions;
 	OptionsCopy.AdditionalTransform = MeshGlobalTransform;
+	OptionsCopy.bMergeIdenticalMaterialSlots = false;	 // This must always be false, because we need the material assignments we read from the
+														 // meshes to match up with whatever we cached from AddMeshNode, in order to fixup LOD
+														 // material slots
 
 	UE::Interchange::FMeshPayloadData MeshPayloadData;
 	switch (PayloadKey.Type)
