USD: Fix bug introduced in cl 38998578 where stages were being left open after Interchange import, causing a bunch of automated tests to fail.

diff --git a/Engine/Plugins/Interchange/Runtime/Source/Import/Private/Usd/InterchangeUsdTranslator.cpp b/Engine/Plugins/Interchange/Runtime/Source/Import/Private/Usd/InterchangeUsdTranslator.cpp
--- a/Engine/Plugins/Interchange/Runtime/Source/Import/Private/Usd/InterchangeUsdTranslator.cpp
+++ b/Engine/Plugins/Interchange/Runtime/Source/Import/Private/Usd/InterchangeUsdTranslator.cpp
@@ -3304,7 +3304,9 @@ void UInterchangeUSDTranslator::ReleaseSource()
 		return;
 	}
 
+	ImplPtr->TranslationContext.Reset();	// This holds a reference to the stage
 	ImplPtr->UsdStage = UE::FUsdStage{};
+	ImplPtr->InfoCache.Reset();
 	ImplPtr->CurrentTrackSet = nullptr;
 
 	{
