MaterialX:
- Fix input normal not connected to the displacement shader material function

diff --git a/Engine/Plugins/Interchange/Runtime/Source/Import/Private/MaterialX/MaterialXUtils/MaterialXSurfaceMaterial.cpp b/Engine/Plugins/Interchange/Runtime/Source/Import/Private/MaterialX/MaterialXUtils/MaterialXSurfaceMaterial.cpp
--- a/Engine/Plugins/Interchange/Runtime/Source/Import/Private/MaterialX/MaterialXUtils/MaterialXSurfaceMaterial.cpp
+++ b/Engine/Plugins/Interchange/Runtime/Source/Import/Private/MaterialX/MaterialXUtils/MaterialXSurfaceMaterial.cpp
@@ -69,7 +69,7 @@ UInterchangeBaseNode* FMaterialXSurfaceMaterial::Translate(MaterialX::NodePtr Su
 					if (ConnectedDisplacementShaderNode && InputNormal)
 					{
 						//add the input normal from the surfaceshader to the displacement
-						ConnectedDisplacementShaderNode->addInput()->copyContentFrom(InputNormal);
+						ConnectedDisplacementShaderNode->addInput("normal", InputNormal->getType())->copyContentFrom(InputNormal);
 						// now remove the normal input from the surfaceshader
 						ConnectedNode->removeInput(InputNormal->getName());
 					}
