Interchange:
- Set Displacement center scaling on the Material Factory Node in the material pipeline
- Set the displacement center on the material in Legacy in the material factory

diff --git a/Engine/Plugins/Interchange/Runtime/Source/Import/Private/Material/InterchangeMaterialFactory.cpp b/Engine/Plugins/Interchange/Runtime/Source/Import/Private/Material/InterchangeMaterialFactory.cpp
--- a/Engine/Plugins/Interchange/Runtime/Source/Import/Private/Material/InterchangeMaterialFactory.cpp
+++ b/Engine/Plugins/Interchange/Runtime/Source/Import/Private/Material/InterchangeMaterialFactory.cpp
@@ -1803,6 +1803,11 @@ void UInterchangeMaterialFactory::SetupMaterial(UMaterial* Material, const FImpo
 					{
 						DisplacementExpression->ConnectExpression(DisplacementInput, GetOutputIndex(*DisplacementExpression, OutputName));
 						Material->bEnableTessellation = true;
+
+						if (float DisplacementCenter; MaterialFactoryNode->GetCustomDisplacementCenter(DisplacementCenter))
+						{
+							Material->DisplacementScaling.Center = DisplacementCenter;
+						}
 					}
 				}
 			}
diff --git a/Engine/Plugins/Interchange/Runtime/Source/Pipelines/Private/InterchangeGenericMaterialPipeline.cpp b/Engine/Plugins/Interchange/Runtime/Source/Pipelines/Private/InterchangeGenericMaterialPipeline.cpp
--- a/Engine/Plugins/Interchange/Runtime/Source/Pipelines/Private/InterchangeGenericMaterialPipeline.cpp
+++ b/Engine/Plugins/Interchange/Runtime/Source/Pipelines/Private/InterchangeGenericMaterialPipeline.cpp
@@ -1805,6 +1805,11 @@ void UInterchangeGenericMaterialPipeline::HandleCommonParameters(const UIntercha
 			{
 				MaterialFactoryNode->ConnectOutputToDisplacement(DisplacementExpressionFactoryNode.Get<0>()->GetUniqueID(), DisplacementExpressionFactoryNode.Get<1>());
 			}
+
+			if (float DisplacementCenter; ShaderGraphNode->GetCustomDisplacementCenterMode(DisplacementCenter))
+			{
+				MaterialFactoryNode->SetCustomDisplacementCenter(DisplacementCenter);
+			}
 		}
 	}
 }
@@ -3744,6 +3749,11 @@ bool UInterchangeGenericMaterialPipeline::HandleSubstrate(const UInterchangeShad
 			{
 				UInterchangeShaderPortsAPI::ConnectOuputToInputByName(MaterialFactoryNode, SubstrateMaterial::Parameters::Displacement.ToString(), DisplacementFactoryNode.Get<0>()->GetUniqueID(), DisplacementFactoryNode.Get<1>());
 			}
+
+			if (float DisplacementCenter; ShaderGraphNode->GetCustomDisplacementCenterMode(DisplacementCenter))
+			{
+				MaterialFactoryNode->SetCustomDisplacementCenter(DisplacementCenter);
+			}
 		}
 
 		if(UInterchangeShaderPortsAPI::HasInput(ShaderGraphNode, SubstrateMaterial::Parameters::Occlusion))
