USD: Add prim kind as a custom translated scene node attribute.

diff --git a/Engine/Plugins/Interchange/Runtime/Source/Import/Private/Usd/InterchangeUsdTranslator.cpp b/Engine/Plugins/Interchange/Runtime/Source/Import/Private/Usd/InterchangeUsdTranslator.cpp
--- a/Engine/Plugins/Interchange/Runtime/Source/Import/Private/Usd/InterchangeUsdTranslator.cpp
+++ b/Engine/Plugins/Interchange/Runtime/Source/Import/Private/Usd/InterchangeUsdTranslator.cpp
@@ -2472,6 +2472,13 @@ namespace UE::InterchangeUsdTranslator::Private
 				SceneNode->AddInt32Attribute(UE::Interchange::USD::GeometryPurposeIdentifier, (int32)PrimPurpose);
 			}
 
+			// Store our prim kind as well, if we have any (this becomes the empty string if the prim has no authored kind)
+			const FString KindString = UsdToUnreal::ConvertToken(IUsdPrim::GetKind(Prim));
+			if (!KindString.IsEmpty())
+			{
+				SceneNode->AddStringAttribute(UE::Interchange::USD::KindIdentifier, KindString);
+			}
+
 			// If we're an Xformable, get our transform.
 			// All SceneNodes should have their LocalTransform set though.
 			// Not setting will cause ensure hits in Skeleton generations for example.
diff --git a/Engine/Plugins/Interchange/Runtime/Source/Import/Public/Usd/InterchangeUsdDefinitions.h b/Engine/Plugins/Interchange/Runtime/Source/Import/Public/Usd/InterchangeUsdDefinitions.h
--- a/Engine/Plugins/Interchange/Runtime/Source/Import/Public/Usd/InterchangeUsdDefinitions.h
+++ b/Engine/Plugins/Interchange/Runtime/Source/Import/Public/Usd/InterchangeUsdDefinitions.h
@@ -9,6 +9,9 @@ namespace UE::Interchange::USD
 	// Name of a custom attribute added to translated nodes to contain their geometry purpose (proxy, render, guide, etc.)
 	const FString GeometryPurposeIdentifier = TEXT("USD_Geometry_Purpose");
 
+	// Prim kind metadata
+	const FString KindIdentifier = TEXT("USD_Kind");
+
 	// Prefix for custom attributes that describe material instance parameters
 	const FString AttributeIdentifierPrefix = TEXT("USD_MI_Attribute_");
 
