Interchange: Add SourceData ContextObjectsByTag;

Add USD context object.

diff --git a/Engine/Plugins/Interchange/Runtime/Source/Import/Private/Usd/InterchangeUsdContext.cpp b/Engine/Plugins/Interchange/Runtime/Source/Import/Private/Usd/InterchangeUsdContext.cpp
new file mode 100644
--- /dev/null
+++ b/Engine/Plugins/Interchange/Runtime/Source/Import/Private/Usd/InterchangeUsdContext.cpp
@@ -0,0 +1,115 @@
+// Copyright Epic Games, Inc. All Rights Reserved.
+
+#include "Usd/InterchangeUsdContext.h"
+
+#include "Objects/USDInfoCache.h"
+#include "USDConversionUtils.h"
+#include "UsdWrappers/UsdStage.h"
+
+#include UE_INLINE_GENERATED_CPP_BY_NAME(InterchangeUsdContext)
+
+#define LOCTEXT_NAMESPACE "InterchangeUSDContext"
+
+UInterchangeUsdContext::UInterchangeUsdContext()
+{
+}
+
+void UInterchangeUsdContext::BeginDestroy()
+{
+#if USE_USD_SDK
+	if (bShouldCleanUpFromStageCache && StageIdInUsdUtilsStageCache != INDEX_NONE)
+	{
+		UsdUtils::RemoveStageFromUsdUtilsStageCache(StageIdInUsdUtilsStageCache);
+	}
+
+	ReleaseInfoCache();
+#endif	  // USE_USD_SDK
+
+	Super::BeginDestroy();
+}
+
+int64 UInterchangeUsdContext::GetStageId() const
+{
+	return StageIdInUsdUtilsStageCache;
+}
+
+void UInterchangeUsdContext::SetStageId(int64 InStageId)
+{
+	StageIdInUsdUtilsStageCache = InStageId;
+}
+
+UE::FUsdStage UInterchangeUsdContext::GetUsdStage() const
+{
+#if USE_USD_SDK
+	return UsdUtils::FindUsdUtilsStageCacheStageId(StageIdInUsdUtilsStageCache);
+#else
+	return {};
+#endif	  // USE_USD_SDK
+}
+
+bool UInterchangeUsdContext::SetUsdStage(const UE::FUsdStage& InStage)
+{
+#if USE_USD_SDK
+	UE::FUsdStage OurStage = GetUsdStage();
+	if (OurStage == InStage)
+	{
+		return true;
+	}
+
+	// This stage is already in the stage cache somehow. Just take its existing id,
+	// but let's remember to not remove it from the stage cache whenever we're done, because
+	// it wasn't us that put it there in the first place
+	int64 ExistingId = UsdUtils::GetUsdUtilsStageCacheStageId(InStage);
+	if (ExistingId != INDEX_NONE)
+	{
+		bShouldCleanUpFromStageCache = false;
+		StageIdInUsdUtilsStageCache = ExistingId;
+		return true;
+	}
+
+	// We're adding this stage to the stage cache, so let's make sure it's cleaned
+	// up whenever we're released or else it will remain there forever
+	bShouldCleanUpFromStageCache = true;
+	StageIdInUsdUtilsStageCache = UsdUtils::InsertStageIntoUsdUtilsStageCache(InStage);
+	return StageIdInUsdUtilsStageCache != INDEX_NONE;
+#else
+	return false;
+#endif	  // USE_USD_SDK
+}
+
+FUsdInfoCache* UInterchangeUsdContext::GetInfoCache() const
+{
+	return InfoCache;
+}
+
+void UInterchangeUsdContext::SetExternalInfoCache(FUsdInfoCache& InInfoCache)
+{
+	if (&InInfoCache == GetInfoCache())
+	{
+		return;
+	}
+
+	ReleaseInfoCache();
+
+	InfoCache = &InInfoCache;
+}
+
+FUsdInfoCache* UInterchangeUsdContext::CreateOwnedInfoCache()
+{
+	if (OwnedInfoCache.IsValid())
+	{
+		return GetInfoCache();
+	}
+
+	OwnedInfoCache = MakeUnique<FUsdInfoCache>();
+	InfoCache = OwnedInfoCache.Get();
+	return GetInfoCache();
+}
+
+void UInterchangeUsdContext::ReleaseInfoCache()
+{
+	InfoCache = nullptr;
+	OwnedInfoCache.Reset();
+}
+
+#undef LOCTEXT_NAMESPACE
diff --git a/Engine/Plugins/Interchange/Runtime/Source/Import/Private/Usd/InterchangeUsdTranslator.cpp b/Engine/Plugins/Interchange/Runtime/Source/Import/Private/Usd/InterchangeUsdTranslator.cpp
--- a/Engine/Plugins/Interchange/Runtime/Source/Import/Private/Usd/InterchangeUsdTranslator.cpp
+++ b/Engine/Plugins/Interchange/Runtime/Source/Import/Private/Usd/InterchangeUsdTranslator.cpp
@@ -5,6 +5,7 @@
 #include "Objects/USDInfoCache.h"
 #include "Objects/USDSchemaTranslator.h"
 #include "UnrealUSDWrapper.h"
+#include "Usd/InterchangeUsdContext.h"
 #include "Usd/InterchangeUsdDefinitions.h"
 #include "USDConversionUtils.h"
 #include "USDGeomMeshConversion.h"
@@ -328,7 +329,7 @@ namespace UE::InterchangeUsdTranslator::Private
 		// ReleaseSource() clears this member, once translation is complete.
 		UE::FUsdStage UsdStage;
 		TSharedPtr<FUsdSchemaTranslationContext> TranslationContext;
-		TSharedPtr<FUsdInfoCache> InfoCache;
+		FUsdInfoCache* InfoCache = nullptr;
 
 		// We store temp stages in here that we open in order to parse stuff inside of inactive variants of our main UsdStage.
 		// We do this because the payload data are retrieved concurrently, and toggling variants mutates the current stage.
@@ -1214,23 +1215,6 @@ namespace UE::InterchangeUsdTranslator::Private
 		TranslationContext->MaterialPurpose = Settings.MaterialPurpose;
 	}
 
-	void UInterchangeUSDTranslatorImpl::SetupInfoCache()
-	{
-		if (!ensure(TranslationContext))
-		{
-			return;
-		}
-
-		if (!InfoCache)
-		{
-			InfoCache = MakeShared<FUsdInfoCache>();
-		}
-
-		TranslationContext->UsdInfoCache = InfoCache.Get();
-
-		InfoCache->RebuildCacheForSubtrees({UE::FSdfPath::AbsoluteRootPath()}, *TranslationContext);
-	}
-
 	bool DecompressUSDZFileToTempFolder(const FString& InUSDZFilePath, FString& OutDecompressedUSDZRoot)
 	{
 		const bool bIncludeDot = false;
@@ -3193,7 +3177,6 @@ bool UInterchangeUSDTranslator::Translate(UInterchangeBaseNodeContainer& NodeCon
 	{
 		return false;
 	}
-	ImplPtr->CurrentTrackSet = nullptr;
 
 	UInterchangeUsdTranslatorSettings* Settings = Cast<UInterchangeUsdTranslatorSettings>(GetSettings());
 	if (!Settings)
@@ -3201,70 +3184,121 @@ bool UInterchangeUSDTranslator::Translate(UInterchangeBaseNodeContainer& NodeCon
 		return false;
 	}
 
-	FString FilePath = GetSourceData()->GetFilename();
-	if (!FPaths::FileExists(FilePath))
+	if (!SourceData)
 	{
 		return false;
 	}
 
-	// If we're provided a USDZ path, for now we will decompress it into a temp dir and redirect our paths.
-	//
-	// This is mainly because the texture factories must receive a simple file path in order to produce the
-	// their payloads. It's not practical to make them handle USDZ files, and it's not yet possible
-	// to provide them with raw binary buffers directly either
-	if (FString DecompressedUSDZRoot; DecompressUSDZFileToTempFolder(FilePath, DecompressedUSDZRoot))
+	// Setup context
+	UObject* ContextObject = SourceData->GetContextObjectByTag(UE::Interchange::USD::USDContextTag);
+	UInterchangeUsdContext* Context = Cast<UInterchangeUsdContext>(ContextObject);
+	if (!Context)
 	{
-		ImplPtr->USDZFilePath = FilePath;
-		ImplPtr->DecompressedUSDZRoot = DecompressedUSDZRoot;
+		Context = NewObject<UInterchangeUsdContext>();
 
-		FilePath = DecompressedUSDZRoot;
+		ensureMsgf(
+			!ContextObject,
+			TEXT("Invalid ContextObject with tag '%s' will be removed and replaced with an UInterchangeUsdContext object"),
+			*UE::Interchange::USD::USDContextTag
+		);
+		SourceData->SetContextObjectByTag(UE::Interchange::USD::USDContextTag, Context);
 	}
 
-	// Import should always feel like it's directly from disk, so we ignore already loaded layers and stage cache
-	const bool bUseStageCache = false;
-	const bool bForceReloadLayersFromDisk = true;
-	ImplPtr->UsdStage = UnrealUSDWrapper::OpenStage(*FilePath, EUsdInitialLoadSet::LoadAll, bUseStageCache, bForceReloadLayersFromDisk);
-	if (!ImplPtr->UsdStage)
+	// Setup stage
+	FString USDZFilePath;
+	FString DecompressedUSDZRoot;
+	UE::FUsdStage StageToImport = Context->GetUsdStage();
 	{
-		return false;
-	}
+		// Context didn't provide a stage: Try loading one from the provided file path
+		if (!StageToImport)
+		{
+			FString FilePath = SourceData->GetFilename();
+			if (!FPaths::FileExists(FilePath))
+			{
+				return false;
+			}
 
-	// Apply stage settings
-	if (Settings)
-	{
-		// Apply coordinate system conversion to the stage if we have one
-		if (Settings->bOverrideStageOptions)
+			// If we're provided a USDZ path, for now we will decompress it into a temp dir and redirect our paths.
+			//
+			// This is mainly because the texture factories must receive a simple file path in order to produce the
+			// their payloads. It's not practical to make them handle USDZ files, and it's not yet possible
+			// to provide them with raw binary buffers directly either
+			if (DecompressUSDZFileToTempFolder(FilePath, DecompressedUSDZRoot))
+			{
+				USDZFilePath = FilePath;
+				FilePath = DecompressedUSDZRoot;
+			}
+
+			// Import should always feel like it's directly from disk, so we ignore already loaded layers and stage cache
+			const bool bUseStageCache = false;
+			const bool bForceReloadLayersFromDisk = true;
+			StageToImport = UnrealUSDWrapper::OpenStage(*FilePath, EUsdInitialLoadSet::LoadAll, bUseStageCache, bForceReloadLayersFromDisk);
+		}
+
+		if (!StageToImport)
 		{
-			UsdUtils::SetUsdStageMetersPerUnit(ImplPtr->UsdStage, Settings->StageOptions.MetersPerUnit);
-			UsdUtils::SetUsdStageUpAxis(ImplPtr->UsdStage, Settings->StageOptions.UpAxis);
+			return false;
+		}
+
+		if (Settings)
+		{
+			// Apply coordinate system conversion to the stage if we have one
+			if (Settings->bOverrideStageOptions)
+			{
+				UsdUtils::SetUsdStageMetersPerUnit(StageToImport, Settings->StageOptions.MetersPerUnit);
+				UsdUtils::SetUsdStageUpAxis(StageToImport, Settings->StageOptions.UpAxis);
+			}
+
+			StageToImport.SetInterpolationType(Settings->InterpolationType);
 		}
+	}
 
-		ImplPtr->UsdStage.SetInterpolationType(Settings->InterpolationType);
+	// Setup info cache
+	FUsdInfoCache* InfoCache = Context->GetInfoCache();
+	if (!InfoCache)
+	{
+		InfoCache = Context->CreateOwnedInfoCache();
 	}
 
+	// Fill in our context with our potentially-created-on-demand stage and info cache
+	ensure(Context->SetUsdStage(StageToImport));
+	Context->SetExternalInfoCache(*InfoCache);
+
+	// Setup impl
+	ImplPtr->UsdStage = StageToImport;
+	ImplPtr->InfoCache = InfoCache;
 	ImplPtr->SetupTranslationContext(*Settings);
-	ImplPtr->SetupInfoCache();
+	ImplPtr->TranslationContext->UsdInfoCache = ImplPtr->InfoCache;
+	ImplPtr->CurrentTrackSet = nullptr;
+	ImplPtr->USDZFilePath = USDZFilePath;
+	ImplPtr->DecompressedUSDZRoot = DecompressedUSDZRoot;
 
 	// Cache these so we don't have to keep converting these tokens over and over during translation
-	FUsdMeshConversionOptions& MeshOptions = ImplPtr->CachedMeshConversionOptions;
-
-	// We filter for this on the pipeline now
-	MeshOptions.PurposesToLoad = EUsdPurpose::Default | EUsdPurpose::Proxy | EUsdPurpose::Render | EUsdPurpose::Guide;
-
-	// TODO: Change FUsdMeshConversionOptions to not hold USD types directly, so we don't have to the conversion below everywhere.
-	// We can't use UsdToUnreal::ConvertToken() here because it returns a TUsdStore, and the template instantiation created in this module doesn't
-	// really do anything anyway as the module doesn't use IMPLEMENT_MODULE_USD!
-	// Luckily we can get around this here because pxr::TfToken doesn't allocate on its own: At most USD makes a copy of the string, which it should
-	// allocate/deallocate on its own allocator.
-	MeshOptions.RenderContext = Settings->RenderContext == UnrealIdentifiers::UniversalRenderContext
-									? pxr::UsdShadeTokens->universalRenderContext
-									: pxr::TfToken{TCHAR_TO_UTF8(*Settings->RenderContext.ToString())};
-	MeshOptions.MaterialPurpose = Settings->MaterialPurpose.IsNone() ? pxr::UsdShadeTokens->allPurpose
-																	 : pxr::TfToken{TCHAR_TO_UTF8(*Settings->MaterialPurpose.ToString())};
+	{
+		FUsdMeshConversionOptions& MeshOptions = ImplPtr->CachedMeshConversionOptions;
+
+		// We filter for this on the pipeline now
+		MeshOptions.PurposesToLoad = EUsdPurpose::Default | EUsdPurpose::Proxy | EUsdPurpose::Render | EUsdPurpose::Guide;
+
+		// TODO: Change FUsdMeshConversionOptions to not hold USD types directly, so we don't have to the conversion below everywhere.
+		// We can't use UsdToUnreal::ConvertToken() here because it returns a TUsdStore, and the template instantiation created in this module doesn't
+		// really do anything anyway as the module doesn't use IMPLEMENT_MODULE_USD!
+		// Luckily we can get around this here because pxr::TfToken doesn't allocate on its own: At most USD makes a copy of the string, which it
+		// should allocate/deallocate on its own allocator.
+		MeshOptions.RenderContext = Settings->RenderContext == UnrealIdentifiers::UniversalRenderContext
+										? pxr::UsdShadeTokens->universalRenderContext
+										: pxr::TfToken{TCHAR_TO_UTF8(*Settings->RenderContext.ToString())};
+		MeshOptions.MaterialPurpose = Settings->MaterialPurpose.IsNone() ? pxr::UsdShadeTokens->allPurpose
+																		 : pxr::TfToken{TCHAR_TO_UTF8(*Settings->MaterialPurpose.ToString())};
+	}
 
 	// Traverse stage and emit translated nodes
-	FTraversalInfo Info;
-	Traverse(ImplPtr->UsdStage.GetPseudoRoot(), *ImplPtr, NodeContainer, TranslatorSettings, Info);
+	{
+		ImplPtr->InfoCache->RebuildCacheForSubtrees({UE::FSdfPath::AbsoluteRootPath()}, *ImplPtr->TranslationContext);
+
+		FTraversalInfo Info;
+		Traverse(ImplPtr->UsdStage.GetPseudoRoot(), *ImplPtr, NodeContainer, TranslatorSettings, Info);
+	}
 
 	return true;
 #else
@@ -3282,7 +3316,7 @@ void UInterchangeUSDTranslator::ReleaseSource()
 
 	ImplPtr->TranslationContext.Reset();	// This holds a reference to the stage
 	ImplPtr->UsdStage = UE::FUsdStage{};
-	ImplPtr->InfoCache.Reset();
+	ImplPtr->InfoCache = nullptr;
 	ImplPtr->CurrentTrackSet = nullptr;
 
 	{
diff --git a/Engine/Plugins/Interchange/Runtime/Source/Import/Public/Usd/InterchangeUsdContext.h b/Engine/Plugins/Interchange/Runtime/Source/Import/Public/Usd/InterchangeUsdContext.h
new file mode 100644
--- /dev/null
+++ b/Engine/Plugins/Interchange/Runtime/Source/Import/Public/Usd/InterchangeUsdContext.h
@@ -0,0 +1,57 @@
+// Copyright Epic Games, Inc. All Rights Reserved.
+
+#pragma once
+
+#include "UsdWrappers/UsdStage.h"
+
+#include "CoreMinimal.h"
+
+#include "InterchangeUsdContext.generated.h"
+
+class FUsdInfoCache;
+
+UCLASS(BlueprintType)
+class INTERCHANGEIMPORT_API UInterchangeUsdContext : public UObject
+{
+	GENERATED_BODY()
+
+public:
+	UInterchangeUsdContext();
+	virtual void BeginDestroy() override;
+
+	// Returns the ID the provided stage has within the USDUtils' singleton StageCache.
+	UFUNCTION(BlueprintCallable, Category = "Interchange | USD")
+	int64 GetStageId() const;
+
+	// Sets the ID of a particular stage from the UsdUtils' singleton StageCache.
+	// If this corresponds to a valid USD Stage, that stage will be used for the Interchange import.
+	UFUNCTION(BlueprintCallable, Category = "Interchange | USD")
+	void SetStageId(int64 InStageId);
+
+	// Convenience functions to get/set the stage directly, although it will internally just
+	// get/set the stage into the UsdUtils' singleton StageCache and track it's Id instead.
+	UE::FUsdStage GetUsdStage() const;
+	bool SetUsdStage(const UE::FUsdStage& InStage);
+
+	// Returns the current info cache, whether it is one we fully own, or just an external reference we're tracking
+	FUsdInfoCache* GetInfoCache() const;
+
+	// Receive a reference to an info cache that is external to this object and set it as the current info cache.
+	// Note that this will discard our owned info cache if we had one before.
+	void SetExternalInfoCache(FUsdInfoCache& InInfoCache);
+
+	// Create an info cache that is fully owned by this UInterchangeUsdContext object, and set it as the current info cache.
+	FUsdInfoCache* CreateOwnedInfoCache();
+
+	// Properly releases the reference to the current stage cache, whether owned or not (deleting it if owned)
+	void ReleaseInfoCache();
+
+private:
+	// We never store the stage itself, but only it's Id within the UsdUtils' singleton StageCache.
+	// The intent here is to allow Python stages to be passed in and manipulated via Python in case of Python imports or pipelines.
+	int64 StageIdInUsdUtilsStageCache = INDEX_NONE;
+	bool bShouldCleanUpFromStageCache = false;
+
+	FUsdInfoCache* InfoCache = nullptr;
+	TUniquePtr<FUsdInfoCache> OwnedInfoCache;
+};
diff --git a/Engine/Plugins/Interchange/Runtime/Source/Import/Public/Usd/InterchangeUsdDefinitions.h b/Engine/Plugins/Interchange/Runtime/Source/Import/Public/Usd/InterchangeUsdDefinitions.h
--- a/Engine/Plugins/Interchange/Runtime/Source/Import/Public/Usd/InterchangeUsdDefinitions.h
+++ b/Engine/Plugins/Interchange/Runtime/Source/Import/Public/Usd/InterchangeUsdDefinitions.h
@@ -6,6 +6,8 @@
 
 namespace UE::Interchange::USD
 {
+	const FString USDContextTag = TEXT("USD");
+
 	// Name of a custom attribute added to translated nodes to contain their geometry purpose (proxy, render, guide, etc.)
 	const FString GeometryPurposeIdentifier = TEXT("USD_Geometry_Purpose");
 
diff --git a/Engine/Plugins/Runtime/USDCore/Source/USDUtilities/Private/USDConversionUtils.cpp b/Engine/Plugins/Runtime/USDCore/Source/USDUtilities/Private/USDConversionUtils.cpp
--- a/Engine/Plugins/Runtime/USDCore/Source/USDUtilities/Private/USDConversionUtils.cpp
+++ b/Engine/Plugins/Runtime/USDCore/Source/USDUtilities/Private/USDConversionUtils.cpp
@@ -69,6 +69,7 @@
 #include "pxr/usd/usd/primCompositionQuery.h"
 #include "pxr/usd/usd/primRange.h"
 #include "pxr/usd/usd/stage.h"
+#include "pxr/usd/usd/stageCache.h"
 #include "pxr/usd/usd/variantSets.h"
 #include "pxr/usd/usdGeom/camera.h"
 #include "pxr/usd/usdGeom/gprim.h"
@@ -89,6 +90,7 @@
 #include "pxr/usd/usdSkel/cache.h"
 #include "pxr/usd/usdSkel/root.h"
 #include "pxr/usd/usdSkel/skeletonQuery.h"
+#include "pxr/usd/usdUtils/stageCache.h"
 #include "pxr/usd/usdVol/openVDBAsset.h"
 #include "pxr/usd/usdVol/volume.h"
 #include "USDIncludesEnd.h"
@@ -414,6 +416,32 @@ int32 UsdUtils::GetUsdStageNumFrames(const pxr::UsdStageRefPtr& Stage)
 	return Stage ? FMath::Abs(FMath::CeilToInt32(Stage->GetEndTimeCode()) - FMath::FloorToInt32(Stage->GetStartTimeCode()) + 1) : 0;
 }
 
+int64 UsdUtils::GetUsdUtilsStageCacheStageId(const pxr::UsdStageRefPtr& Stage)
+{
+	pxr::UsdStageCache& StageCache = pxr::UsdUtilsStageCache::Get();
+	pxr::UsdStageCache::Id Id = StageCache.GetId(Stage);
+	return Id.IsValid() ? Id.ToLongInt() : INDEX_NONE;
+}
+
+UE::FUsdStage UsdUtils::FindUsdUtilsStageCacheStageId(int64 Id)
+{
+	pxr::UsdStageCache& StageCache = pxr::UsdUtilsStageCache::Get();
+	pxr::UsdStageRefPtr Stage = StageCache.Find(pxr::UsdStageCache::Id::FromLongInt(Id));
+	return UE::FUsdStage{Stage};
+}
+
+int64 UsdUtils::InsertStageIntoUsdUtilsStageCache(const pxr::UsdStageRefPtr& Stage)
+{
+	pxr::UsdStageCache& StageCache = pxr::UsdUtilsStageCache::Get();
+	return StageCache.Insert(Stage).ToLongInt();
+}
+
+bool UsdUtils::RemoveStageFromUsdUtilsStageCache(int64 StageId)
+{
+	pxr::UsdStageCache& StageCache = pxr::UsdUtilsStageCache::Get();
+	return StageCache.Erase(pxr::UsdStageCache::Id::FromLongInt(StageId));
+}
+
 bool UsdUtils::HasCompositionArcs(const pxr::UsdPrim& Prim)
 {
 	if (!Prim || !Prim.IsActive())
diff --git a/Engine/Plugins/Runtime/USDCore/Source/USDUtilities/Public/USDConversionUtils.h b/Engine/Plugins/Runtime/USDCore/Source/USDUtilities/Public/USDConversionUtils.h
--- a/Engine/Plugins/Runtime/USDCore/Source/USDUtilities/Public/USDConversionUtils.h
+++ b/Engine/Plugins/Runtime/USDCore/Source/USDUtilities/Public/USDConversionUtils.h
@@ -120,6 +120,20 @@ namespace UsdUtils
 
 	USDUTILITIES_API int32 GetUsdStageNumFrames(const pxr::UsdStageRefPtr& Stage);
 
+	/**
+	 * Simple API to interact with the UsdUtils' singleton StageCache from modules not setup to interact with the USD SDK directly
+	 * like Interchange's Import module.
+	 *
+	 * Note: We do this instead of having a UsdStageCache wrapper type because when dealing with stage caches we are realistically going
+	 * to be using the UsdUtils' singleton instance. The wrapper types all *own* the USD version of their objects, and so having the
+	 * FUsdStageCache wrapper *refer* to an external, existing instance of the USD version of itself instead (the singleton instance)
+	 * would likely be very misleading and confusing.
+	 */
+	USDUTILITIES_API int64 GetUsdUtilsStageCacheStageId(const pxr::UsdStageRefPtr& Stage);
+	USDUTILITIES_API UE::FUsdStage FindUsdUtilsStageCacheStageId(int64 Id);
+	USDUTILITIES_API int64 InsertStageIntoUsdUtilsStageCache(const pxr::UsdStageRefPtr& Stage);
+	USDUTILITIES_API bool RemoveStageFromUsdUtilsStageCache(int64 StageId);
+
 	USDUTILITIES_API bool HasCompositionArcs(const pxr::UsdPrim& Prim);
 	USDUTILITIES_API bool HasCompositionArcs(const pxr::SdfPrimSpecHandle& PrimSpec);
 
diff --git a/Engine/Source/Runtime/Interchange/Core/Private/InterchangeSourceData.cpp b/Engine/Source/Runtime/Interchange/Core/Private/InterchangeSourceData.cpp
new file mode 100644
--- /dev/null
+++ b/Engine/Source/Runtime/Interchange/Core/Private/InterchangeSourceData.cpp
@@ -0,0 +1,49 @@
+// Copyright Epic Games, Inc. All Rights Reserved.
+
+#include "InterchangeSourceData.h"
+
+#include UE_INLINE_GENERATED_CPP_BY_NAME(InterchangeSourceData)
+
+UInterchangeSourceData::UInterchangeSourceData()
+{
+	Filename = TEXT("");
+	FileContentHashCache.Reset();
+}
+
+UInterchangeSourceData::UInterchangeSourceData(FString InFilename)
+{
+	SetFilename(InFilename);
+}
+
+UObject* UInterchangeSourceData::GetContextObjectByTag(const FString& Tag) const
+{
+	return ContextObjectsByTag.FindRef(Tag);
+}
+
+void UInterchangeSourceData::SetContextObjectByTag(const FString& Tag, UObject* Object) const
+{
+	ContextObjectsByTag.Add(Tag, Object);
+}
+
+TArray<FString> UInterchangeSourceData::GetAllContextObjectTags() const
+{
+	TArray<FString> Result;
+	ContextObjectsByTag.GetKeys(Result);
+	return Result;
+}
+
+void UInterchangeSourceData::RemoveAllContextObjects() const
+{
+	ContextObjectsByTag.Reset();
+}
+
+void UInterchangeSourceData::ComputeFileContentHashCache() const
+{
+	FileContentHashCache.Reset();
+	if (!Filename.IsEmpty())
+	{
+		// @todo : this is slow and not multi-threaded.  Use FXxHash64::HashBufferChunked instead?
+
+		FileContentHashCache = FMD5Hash::HashFile(*Filename);
+	}
+}
diff --git a/Engine/Source/Runtime/Interchange/Core/Public/InterchangeSourceData.h b/Engine/Source/Runtime/Interchange/Core/Public/InterchangeSourceData.h
--- a/Engine/Source/Runtime/Interchange/Core/Public/InterchangeSourceData.h
+++ b/Engine/Source/Runtime/Interchange/Core/Public/InterchangeSourceData.h
@@ -22,18 +22,12 @@ UCLASS(BlueprintType, Blueprintable, MinimalAPI)
 class UInterchangeSourceData : public UObject
 {
 	 GENERATED_BODY()
+
 public:
-	UInterchangeSourceData()
-	{
-		Filename = TEXT("");
-		FileContentHashCache.Reset();
-	}
+	INTERCHANGECORE_API UInterchangeSourceData();
+
+	INTERCHANGECORE_API UInterchangeSourceData(FString InFilename);
 
-	UInterchangeSourceData(FString InFilename)
-	{
-		SetFilename(InFilename);
-	}
-	
 	UFUNCTION(BlueprintCallable, Category = "Interchange | Import Manager")
 	FORCEINLINE FString GetFilename() const
 	{
@@ -81,17 +75,20 @@ public:
 		return FString();
 	}
 
-private:
-	void ComputeFileContentHashCache() const
-	{
-		FileContentHashCache.Reset();
-		if (!Filename.IsEmpty())
-		{
-			// @todo : this is slow and not multi-threaded.  Use FXxHash64::HashBufferChunked instead?
+	UFUNCTION(BlueprintCallable, Category = "Interchange | Import Manager")
+	INTERCHANGECORE_API UObject* GetContextObjectByTag(const FString& Tag) const;
 
-			FileContentHashCache = FMD5Hash::HashFile(*Filename);
-		}
-	}
+	UFUNCTION(BlueprintCallable, Category = "Interchange | Import Manager")
+	INTERCHANGECORE_API void SetContextObjectByTag(const FString& Tag, UObject* Object) const;
+
+	UFUNCTION(BlueprintCallable, Category = "Interchange | Import Manager")
+	INTERCHANGECORE_API TArray<FString> GetAllContextObjectTags() const;
+
+	UFUNCTION(BlueprintCallable, Category = "Interchange | Import Manager")
+	INTERCHANGECORE_API void RemoveAllContextObjects() const;
+
+private:
+	void INTERCHANGECORE_API ComputeFileContentHashCache() const;
 
 	UPROPERTY()
 	FString Filename;
@@ -103,4 +100,13 @@ private:
 	 * It also makes sure we do not waste CPU computing a cache in case no client uses GetFileContentHash().
 	 */
 	mutable TOptional<FMD5Hash> FileContentHashCache;
+
+	/**
+	 * UObjects that are accessible by the translators, pipelines, and the caller of the Interchange import.
+	 *
+	 * Use this to transmit additional information that is beneficial to reuse but can't be serialized, such as
+	 * external SDK memory objects, external assets, or large cached data.
+	 */
+	UPROPERTY()
+	mutable TMap<FString, TObjectPtr<UObject>> ContextObjectsByTag;
 };
