Interchange: add packed actor pipeline property to allow users to create a packed level actor.

#rnx

diff --git a/Engine/Plugins/Interchange/Editor/Source/Utilities/Private/InterchangeEditorUtilities.cpp b/Engine/Plugins/Interchange/Editor/Source/Utilities/Private/InterchangeEditorUtilities.cpp
--- a/Engine/Plugins/Interchange/Editor/Source/Utilities/Private/InterchangeEditorUtilities.cpp
+++ b/Engine/Plugins/Interchange/Editor/Source/Utilities/Private/InterchangeEditorUtilities.cpp
@@ -13,7 +13,7 @@ extern UNREALED_API UEditorEngine* GEditor;
 
 #include UE_INLINE_GENERATED_CPP_BY_NAME(InterchangeEditorUtilities)
 
-bool UInterchangeEditorUtilities::SaveAsset(UObject* Asset)
+bool UInterchangeEditorUtilities::SaveAsset(UObject* Asset) const
 {
 	TArray<UPackage*> PackagesToSave;
 	PackagesToSave.Add(Asset->GetPackage());
@@ -21,11 +21,20 @@ bool UInterchangeEditorUtilities::SaveAsset(UObject* Asset)
 	return (ReturnCode == FEditorFileUtils::PR_Success);
 }
 
-bool UInterchangeEditorUtilities::IsRuntimeOrPIE()
+bool UInterchangeEditorUtilities::IsRuntimeOrPIE() const
 {
 #if WITH_EDITOR
 	return (GEditor && GEditor->PlayWorld) || GIsPlayInEditorWorld || IsRunningGame();
 #else
 	return true;
 #endif // WITH_EDITOR
+}
+
+bool UInterchangeEditorUtilities::ClearEditorSelection() const
+{
+	if(GEditor)
+	{
+		GEditor->SelectNone(true, true, false);
+	}
+	return true;
 }
\ No newline at end of file
diff --git a/Engine/Plugins/Interchange/Editor/Source/Utilities/Public/InterchangeEditorUtilities.h b/Engine/Plugins/Interchange/Editor/Source/Utilities/Public/InterchangeEditorUtilities.h
--- a/Engine/Plugins/Interchange/Editor/Source/Utilities/Public/InterchangeEditorUtilities.h
+++ b/Engine/Plugins/Interchange/Editor/Source/Utilities/Public/InterchangeEditorUtilities.h
@@ -15,7 +15,9 @@ public:
 
 protected:
 
-	virtual bool SaveAsset(UObject* Asset) override;
+	virtual bool SaveAsset(UObject* Asset) const override;
 
-	virtual bool IsRuntimeOrPIE() override;
+	virtual bool IsRuntimeOrPIE() const override;
+
+	virtual bool ClearEditorSelection() const override;
 };
diff --git a/Engine/Plugins/Interchange/Runtime/Source/Import/Private/Scene/InterchangeLevelInstanceActorFactory.cpp b/Engine/Plugins/Interchange/Runtime/Source/Import/Private/Scene/InterchangeLevelInstanceActorFactory.cpp
--- a/Engine/Plugins/Interchange/Runtime/Source/Import/Private/Scene/InterchangeLevelInstanceActorFactory.cpp
+++ b/Engine/Plugins/Interchange/Runtime/Source/Import/Private/Scene/InterchangeLevelInstanceActorFactory.cpp
@@ -1,6 +1,8 @@
 // Copyright Epic Games, Inc. All Rights Reserved. 
 #include "Scene/InterchangeLevelInstanceActorFactory.h"
 
+#include "Engine/Blueprint.h"
+#include "Engine/Level.h"
 #include "Engine/World.h"
 #include "InterchangeEditorUtilitiesBase.h"
 #include "InterchangeImportLog.h"
@@ -8,7 +10,10 @@
 #include "InterchangeLevelInstanceActorFactoryNode.h"
 #include "InterchangeManager.h"
 #include "LevelInstance/LevelInstanceActor.h"
+#include "LevelInstance/LevelInstanceSubsystem.h"
 #include "Nodes/InterchangeBaseNodeContainer.h"
+#include "PackedLevelActor/PackedLevelActor.h"
+#include "PackedLevelActor/PackedLevelActorBuilder.h"
 #include "Scene/InterchangeActorHelper.h"
 
 #include UE_INLINE_GENERATED_CPP_BY_NAME(InterchangeLevelInstanceActorFactory)
@@ -18,57 +23,217 @@ UClass* UInterchangeLevelInstanceActorFactory::GetFactoryClass() const
 	return ALevelInstance::StaticClass();
 }
 
-UObject* UInterchangeLevelInstanceActorFactory::ProcessActor(AActor& SpawnedActor, const UInterchangeActorFactoryNode& FactoryNode, const UInterchangeBaseNodeContainer& NodeContainer, const FImportSceneObjectsParams& Params)
+UObject* UInterchangeLevelInstanceActorFactory::ImportSceneObject_GameThread(const UInterchangeFactoryBase::FImportSceneObjectsParams& CreateSceneObjectsParams)
 {
-	auto SuperResult = [&]()
-		{
-			return Super::ProcessActor(SpawnedActor, FactoryNode, NodeContainer, Params);
-		};
+#if !WITH_EDITOR
+	return nullptr;
+#else
 
-	const UInterchangeLevelInstanceActorFactoryNode* LevelInstanceActorFactoryNode = Cast<UInterchangeLevelInstanceActorFactoryNode>(&FactoryNode);
-	if (!LevelInstanceActorFactoryNode)
-	{
-		UE_LOG(LogInterchangeImport, Warning, TEXT("UInterchangeLevelInstanceActorFactory::ProcessActor: Level instance actor factory node is null"));
-		return SuperResult();
-	}
+	using namespace UE::Interchange;
+	
+	const bool bIsReimport = CreateSceneObjectsParams.bIsReimport || CreateSceneObjectsParams.ReimportObject || CreateSceneObjectsParams.ReimportFactoryNode;
 
-	ALevelInstance* LevelInstanceActor = Cast<ALevelInstance>(&SpawnedActor);
-	if (!LevelInstanceActor)
+	UInterchangeLevelInstanceActorFactoryNode* FactoryNode = Cast<UInterchangeLevelInstanceActorFactoryNode>(CreateSceneObjectsParams.FactoryNode);
+	if (!ensure(FactoryNode) || !CreateSceneObjectsParams.NodeContainer)
 	{
-		FString AssetName = LevelInstanceActorFactoryNode->GetDisplayLabel();
-		UE_LOG(LogInterchangeImport, Warning, TEXT("UInterchangeLevelInstanceActorFactory::ProcessActor: Actor level instance was not created %s."), *AssetName);
-		return SuperResult();
+		return nullptr;
 	}
-	//Get the world reference by this level instance actor
+	
 	FString ReferenceLevelFactoryNodeUid;
-	if (!LevelInstanceActorFactoryNode->GetCustomLevelReference(ReferenceLevelFactoryNodeUid))
+	if (!FactoryNode->GetCustomLevelReference(ReferenceLevelFactoryNodeUid))
 	{
-		UE_LOG(LogInterchangeImport, Warning, TEXT("UInterchangeLevelInstanceActorFactory::ProcessActor: Actor level instance do not reference any level factory node."));
-		return SuperResult();
+		UE_LOG(LogInterchangeImport, Warning, TEXT("UInterchangeLevelInstanceActorFactory::ImportSceneObject_GameThread: Actor level instance do not reference any level factory node."));
+		return nullptr;
 	}
 
-	UInterchangeLevelFactoryNode* ReferenceLevelFactoryNode = Cast<UInterchangeLevelFactoryNode>(NodeContainer.GetFactoryNode(ReferenceLevelFactoryNodeUid));
+	UInterchangeLevelFactoryNode* ReferenceLevelFactoryNode = Cast<UInterchangeLevelFactoryNode>(CreateSceneObjectsParams.NodeContainer->GetFactoryNode(ReferenceLevelFactoryNodeUid));
 	if (!ReferenceLevelFactoryNode)
 	{
-		UE_LOG(LogInterchangeImport, Warning, TEXT("UInterchangeLevelInstanceActorFactory::ProcessActor: Level factory node is invalid."));
-		return SuperResult();
+		UE_LOG(LogInterchangeImport, Warning, TEXT("UInterchangeLevelInstanceActorFactory::ImportSceneObject_GameThread: Reference Level factory node is invalid."));
+		return nullptr;
 	}
-	FSoftObjectPath ReferenceLevelSoftObjectPath;
-	if (!ReferenceLevelFactoryNode->GetCustomReferenceObject(ReferenceLevelSoftObjectPath))
+
+	FSoftObjectPath ReferenceWorldPath;
+	if (!ReferenceLevelFactoryNode->GetCustomReferenceObject(ReferenceWorldPath))
 	{
-		UE_LOG(LogInterchangeImport, Warning, TEXT("UInterchangeLevelInstanceActorFactory::ProcessActor: Level factory node do not reference any UWorld soft object path."));
-		return SuperResult();
+		UE_LOG(LogInterchangeImport, Warning, TEXT("UInterchangeLevelInstanceActorFactory::ImportSceneObject_GameThread: Reference Level factory node do not point on a valid UWorld."));
+		return nullptr;
 	}
 
-	UWorld* ReferenceWorld = Cast<UWorld>(ReferenceLevelSoftObjectPath.TryLoad());
+	UWorld* ReferenceWorld = Cast<UWorld>(ReferenceWorldPath.TryLoad());
 	if (!ReferenceWorld)
 	{
-		UE_LOG(LogInterchangeImport, Warning, TEXT("UInterchangeLevelInstanceActorFactory::ProcessActor: UWorld soft object path point on an invalid UWorld."));
-		return SuperResult();
+		UE_LOG(LogInterchangeImport, Warning, TEXT("UInterchangeLevelInstanceActorFactory::ImportSceneObject_GameThread: Cannot load a valid Reference UWorld."));
+		return nullptr;
+	}
+
+	FString TargetLevelFactoryNodeUid;
+	if (!FactoryNode->GetCustomLevelUid(TargetLevelFactoryNodeUid))
+	{
+		UE_LOG(LogInterchangeImport, Warning, TEXT("UInterchangeLevelInstanceActorFactory::ImportSceneObject_GameThread: Actor level instance do not have a valid target world to be created."));
+		return nullptr;
+	}
+
+	UInterchangeLevelFactoryNode* TargetLevelFactoryNode = Cast<UInterchangeLevelFactoryNode>(CreateSceneObjectsParams.NodeContainer->GetFactoryNode(TargetLevelFactoryNodeUid));
+	if (!TargetLevelFactoryNode)
+	{
+		UE_LOG(LogInterchangeImport, Warning, TEXT("UInterchangeLevelInstanceActorFactory::ImportSceneObject_GameThread: Target Level factory node is invalid."));
+		return nullptr;
+	}
+
+	FSoftObjectPath TargetWorldPath;
+	if (!TargetLevelFactoryNode->GetCustomReferenceObject(TargetWorldPath))
+	{
+		UE_LOG(LogInterchangeImport, Warning, TEXT("UInterchangeLevelInstanceActorFactory::ImportSceneObject_GameThread: Target Level factory node do not point on a valid UWorld."));
+		return nullptr;
+	}
+
+	UWorld* TargetWorld = Cast<UWorld>(TargetWorldPath.TryLoad());
+	if (!TargetWorld)
+	{
+		UE_LOG(LogInterchangeImport, Warning, TEXT("UInterchangeLevelInstanceActorFactory::ImportSceneObject_GameThread: Cannot load a valid target UWorld."));
+		return nullptr;
+	}
+	//ALevelInstanceActor is always created in the target world persistent level
+	ensure(TargetWorld->PersistentLevel == CreateSceneObjectsParams.Level);
+
+	const bool bIsPackedActor = FactoryNode->GetObjectClass()->IsChildOf<APackedLevelActor>();
+
+	FString LevelInstanceActorName = FactoryNode->GetDisplayLabel();
+
+	FActorSpawnParameters SpawnParameters;
+	SpawnParameters.Name = FName(CreateSceneObjectsParams.ObjectName);
+	SpawnParameters.NameMode = FActorSpawnParameters::ESpawnActorNameMode::Requested;
+	SpawnParameters.OverrideLevel = CreateSceneObjectsParams.Level;
+	SpawnParameters.SpawnCollisionHandlingOverride = ESpawnActorCollisionHandlingMethod::AlwaysSpawn;
+
+	UClass* LevelInstanceClass = bIsPackedActor ? APackedLevelActor::StaticClass() : ALevelInstance::StaticClass();
+
+	//Reference world must be saved if we want to update the level instance actor
+	if (UInterchangeEditorUtilitiesBase* EditorUtilities = UInterchangeManager::GetInterchangeManager().GetEditorUtilities())
+	{
+		if (!EditorUtilities->SaveAsset(ReferenceWorld))
+		{
+			UE_LOG(LogInterchangeImport, Warning, TEXT("UInterchangeLevelInstanceActorFactory::ImportSceneObject_GameThread: Cannot save the level instance actor (%s) referenced world (%s)"), *FactoryNode->GetDisplayLabel(), *ReferenceWorld->GetName());
+		}
+	}
+
+	auto FindLevelInstance = [TargetWorld, LevelInstanceClass](FName LevelInstanceName)->ALevelInstance*
+		{
+			ALevelInstance* LevelInstanceActor = nullptr;
+			for (ULevel* Level : TargetWorld->GetLevels())
+			{
+				if (LevelInstanceActor)
+				{
+					break;
+				}
+				for (TObjectPtr<AActor> Actor : Level->Actors)
+				{
+					if (!Actor)
+					{
+						continue;
+					}
+					if (FName(*Actor->GetName()) == LevelInstanceName)
+					{
+						if (!Actor->GetClass()->IsChildOf(LevelInstanceClass))
+						{
+							UE_LOG(LogInterchangeImport, Warning, TEXT("UInterchangeLevelInstanceActorFactory::ImportSceneObject_GameThread: There is an existing level instance actor with the same name we want to create/search but with a different class type."));
+							return nullptr;
+						}
+						LevelInstanceActor = Cast<ALevelInstance>(Actor);
+						break;
+					}
+				}
+			}
+			return LevelInstanceActor;
+		};
+
+	ALevelInstance* LevelInstanceActor = nullptr;
+	//Packed level actor need a blueprint between the reference level and the actor
+	if (bIsPackedActor)
+	{
+		auto FindExisitingAsset = [](const FString& AssetPackagePath, const FString& AssetName)->UObject*
+			{
+				UPackage* Pkg = FindPackage(nullptr, *AssetPackagePath);
+				if (!Pkg)
+				{
+					//Try to load the package from disk
+					Pkg = LoadPackage(nullptr, *AssetPackagePath, LOAD_NoWarn | LOAD_Quiet);
+				}
+				if (Pkg)
+				{
+					return StaticFindObject(nullptr, Pkg, *AssetName);
+				}
+				return nullptr;
+			};
+		const FString ReferenceWorldPackageDir = FPaths::GetPath(ReferenceWorldPath.GetLongPackageName());
+		const FString BPAssetName = FPackedLevelActorBuilder::GetPackedBPPrefix() + ReferenceWorldPath.GetAssetName();
+		const FString BPPackageName = FString::Format(TEXT("{0}/{1}"), { ReferenceWorldPackageDir , BPAssetName });
+	
+		UBlueprint* NewBP = Cast<UBlueprint>(FindExisitingAsset(BPPackageName, BPAssetName));
+		TSharedPtr<FPackedLevelActorBuilder> PackedLevelActorBuilder = FPackedLevelActorBuilder::CreateDefaultBuilder();
+
+		if (!NewBP)
+		{
+			constexpr bool bCompileTrue = true;
+			const FString BPAssetPath = FString::Format(TEXT("{0}.{1}"), { BPPackageName , BPAssetName });
+			NewBP = PackedLevelActorBuilder->CreatePackedLevelActorBlueprint(TSoftObjectPtr<UBlueprint>(FSoftObjectPath(BPAssetPath)), ReferenceWorld, bCompileTrue);
+		}
+
+		if (ensure(NewBP))
+		{
+			TSharedPtr<UE::Interchange::FImportAsyncHelper, ESPMode::ThreadSafe> AsyncHelper = UInterchangeManager::GetInterchangeManager().GetAsyncHelper(CreateSceneObjectsParams.AsyncHelperUid);
+			if (AsyncHelper.IsValid())
+			{
+				constexpr int32 SourceIndexZero = 0;
+				NewBP->SetInternalFlags(EInternalObjectFlags::Async);
+				FImportAsyncHelper::FImportedObjectInfo& AssetInfo = AsyncHelper->AddDefaultImportedAssetGetRef(SourceIndexZero);
+
+				//We need a valid factory node to have a valid DefaultImportedAsset entry. Its mandatory to register all created asset
+				//Currently we do not have a blueprint factory node we can simply use the base factory node class.
+				UInterchangeBaseNodeContainer* MutableNodeContainer = const_cast<UInterchangeBaseNodeContainer*>(CreateSceneObjectsParams.NodeContainer);
+				UInterchangeFactoryBaseNode* BpFactoryNode = NewObject<UInterchangeFactoryBaseNode>(MutableNodeContainer);
+				FString UniqueId = FGuid::NewGuid().ToString(EGuidFormats::Base36Encoded);
+				BpFactoryNode->InitializeNode(UniqueId, NewBP->GetName(), EInterchangeNodeContainerType::FactoryData);
+				BpFactoryNode->SetCustomReferenceObject(NewBP);
+				MutableNodeContainer->AddNode(BpFactoryNode);
+
+				AssetInfo.ImportedObject = NewBP;
+				AssetInfo.Factory = this;
+				AssetInfo.FactoryNode = BpFactoryNode;
+				AssetInfo.bIsReimport = false;
+			}
+			LevelInstanceActor = FindLevelInstance(SpawnParameters.Name);
+			if(!LevelInstanceActor)
+			{
+				LevelInstanceActor = TargetWorld->SpawnActor<APackedLevelActor>(NewBP->GeneratedClass, SpawnParameters);
+			}
+		}
+
+		if (!LevelInstanceActor)
+		{
+			UE_LOG(LogInterchangeImport, Warning, TEXT("UInterchangeLevelInstanceActorFactory::ImportSceneObject_GameThread: Failed to create packed level blueprint."));
+			return nullptr;
+		}
+	}
+	else
+	{
+		LevelInstanceActor = FindLevelInstance(SpawnParameters.Name);
+		if(!LevelInstanceActor)
+		{
+			//simple Level Instance actor simply need to spawn the actor, the reference level is already setup
+			LevelInstanceActor = TargetWorld->SpawnActor<ALevelInstance>(ALevelInstance::StaticClass(), SpawnParameters);
+		}
+	}
+
+	if (!LevelInstanceActor)
+	{
+		UE_LOG(LogInterchangeImport, Warning, TEXT("UInterchangeLevelInstanceActorFactory::ImportSceneObject_GameThread: Cannot create/Find a level instance actor in the target world."));
+		return nullptr;
 	}
-#if WITH_EDITOR
-	LevelInstanceActor->SetActorLabel(ReferenceLevelFactoryNode->GetDisplayLabel());
-#endif
 
-	return SuperResult();
+	FactoryNode->SetCustomReferenceObject(LevelInstanceActor);
+
+	return LevelInstanceActor;
+#endif //!WITH_EDITOR
 }
diff --git a/Engine/Plugins/Interchange/Runtime/Source/Import/Public/Scene/InterchangeLevelInstanceActorFactory.h b/Engine/Plugins/Interchange/Runtime/Source/Import/Public/Scene/InterchangeLevelInstanceActorFactory.h
--- a/Engine/Plugins/Interchange/Runtime/Source/Import/Public/Scene/InterchangeLevelInstanceActorFactory.h
+++ b/Engine/Plugins/Interchange/Runtime/Source/Import/Public/Scene/InterchangeLevelInstanceActorFactory.h
@@ -24,17 +24,11 @@ public:
 
 	virtual UClass* GetFactoryClass() const override;
 
-	// Interchange factory base interface end
-	//////////////////////////////////////////////////////////////////////////
+	virtual UObject* ImportSceneObject_GameThread(const UInterchangeFactoryBase::FImportSceneObjectsParams& CreateSceneObjectsParams) override;
 
-protected:
+	// Interchange factory base interface end
 	//////////////////////////////////////////////////////////////////////////
-	// Interchange actor factory interface begin
 
-	virtual UObject* ProcessActor(class AActor& SpawnedActor, const UInterchangeActorFactoryNode& FactoryNode, const UInterchangeBaseNodeContainer& NodeContainer, const FImportSceneObjectsParams& Params) override;
-
-	// Interchange actor factory interface end
-	//////////////////////////////////////////////////////////////////////////
 };
 
 
diff --git a/Engine/Plugins/Interchange/Runtime/Source/Pipelines/Private/InterchangeGenericScenesPipeline.cpp b/Engine/Plugins/Interchange/Runtime/Source/Pipelines/Private/InterchangeGenericScenesPipeline.cpp
--- a/Engine/Plugins/Interchange/Runtime/Source/Pipelines/Private/InterchangeGenericScenesPipeline.cpp
+++ b/Engine/Plugins/Interchange/Runtime/Source/Pipelines/Private/InterchangeGenericScenesPipeline.cpp
@@ -36,6 +36,7 @@
 
 #include "Animation/SkeletalMeshActor.h"
 #include "CineCameraActor.h"
+#include "Engine/Blueprint.h"
 #include "Engine/DirectionalLight.h"
 #include "Engine/Level.h"
 #include "Engine/PointLight.h"
@@ -46,6 +47,10 @@
 #include "GameFramework/WorldSettings.h"
 #include "LevelInstance/LevelInstanceActor.h"
 #include "Misc/PackageName.h"
+#include "PackedLevelActor/PackedLevelActor.h"
+#include "PackedLevelActor/PackedLevelActorBuilder.h"
+#include "PreviewScene.h"
+#include "UObject/UObjectIterator.h"
 #include "WorldPartition/WorldPartition.h"
 
 #if WITH_EDITOR
@@ -205,6 +210,15 @@ namespace UE::Interchange::Private
 	}
 }
 
+UInterchangeGenericLevelPipeline::~UInterchangeGenericLevelPipeline()
+{
+	if (PreviewScene)
+	{
+		delete PreviewScene;
+		PreviewScene = nullptr;
+	}
+}
+
 void UInterchangeGenericLevelPipeline::AdjustSettingsForContext(const FInterchangePipelineContextParams& ContextParams)
 {
 	Super::AdjustSettingsForContext(ContextParams);
@@ -220,6 +234,12 @@ void UInterchangeGenericLevelPipeline::ExecutePipeline(UInterchangeBaseNodeConta
 		return;
 	}
 
+	//Clear any editor selection to avoid crash if selected actor change
+	if (UInterchangeEditorUtilitiesBase* EditorUtilities = UInterchangeManager::GetInterchangeManager().GetEditorUtilities())
+	{
+		EditorUtilities->ClearEditorSelection();
+	}
+
 	BaseNodeContainer = InBaseNodeContainer;
 
 	// Make sure all factory nodes created for assets have the chosen policy strategy
@@ -258,7 +278,20 @@ void UInterchangeGenericLevelPipeline::ExecutePipeline(UInterchangeBaseNodeConta
 #if WITH_EDITORONLY_DATA
 	const FString FilePath = FPaths::ConvertRelativePathToFull(InSourceDatas[0]->GetFilename());
 
-	UWorld* TargetWorld = GWorld->GetCurrentLevel()->GetWorld();
+	UWorld* TargetWorld = nullptr;
+	//Create the preview scene if its a reimport and we have to reimport a ALevelInstance or a APackedLevelActor
+	//In this case we dont want to add a new actor since those import type generate a kind of prefab
+	//that can be instanced multiple time in one or more world.
+	if (bIsReimportContext && SceneHierarchyType != EInterchangeSceneHierarchyType::CreateLevelActors)
+	{
+		PreviewScene = new FPreviewScene(FPreviewScene::ConstructionValues().SetTransactional(false));
+		TargetWorld = PreviewScene->GetWorld();
+	}
+
+	if(!TargetWorld)
+	{
+		TargetWorld = GWorld->GetCurrentLevel()->GetWorld();
+	}
 	
 	UInterchangeLevelFactoryNode* ParentLevelFactoryNode = nullptr;
 	LevelInstanceActorFactoryNode = nullptr;
@@ -282,9 +315,10 @@ void UInterchangeGenericLevelPipeline::ExecutePipeline(UInterchangeBaseNodeConta
 
 		const bool bReimportReferenceWorld = ReimportReferenceWorld && ReimportReferenceWorld == TargetWorld;
 
-		if (!bReimportReferenceWorld && SceneHierarchyType == EInterchangeSceneHierarchyType::CreateLevelInstanceActor)
+		if (!bReimportReferenceWorld && SceneHierarchyType != EInterchangeSceneHierarchyType::CreateLevelActors)
 		{
 			ensure(!LevelFactoryNode);
+			const bool bCreatePackedActor = SceneHierarchyType == EInterchangeSceneHierarchyType::CreatePackedActor;
 			const FString DisplayLabel = ReimportReferenceWorld ? ReimportReferenceWorld->GetName() : ("Level_") + FPaths::GetBaseFilename(FilePath);
 			const FString NodeUid = TEXT("Level_") + FilePath;
 			ensure(!BaseNodeContainer->IsNodeUidValid(NodeUid));
@@ -311,7 +345,8 @@ void UInterchangeGenericLevelPipeline::ExecutePipeline(UInterchangeBaseNodeConta
 
 				if (ensure(LevelInstanceActorFactoryNode))
 				{
-					LevelInstanceActorFactoryNode->SetCustomActorClassName(ALevelInstance::StaticClass()->GetPathName());
+					UClass* LevelInstanceClass = bCreatePackedActor ? APackedLevelActor::StaticClass() : ALevelInstance::StaticClass();
+					LevelInstanceActorFactoryNode->SetCustomActorClassName(LevelInstanceClass->GetPathName());
 					//This actor has the same name has the level it reference
 					const FString ActorDisplayLabel = DisplayLabel;
 					FString ActorNodeUid = TEXT("LevelInstance_") + FilePath;
@@ -1206,7 +1241,7 @@ void UInterchangeGenericLevelPipeline::ExecutePostImportPipeline(const UIntercha
 					}
 				);
 			}
-			else if (SceneHierarchyType == EInterchangeSceneHierarchyType::CreateLevelInstanceActor)
+			else
 			{
 				InBaseNodeContainer->IterateNodesOfType<UInterchangeLevelFactoryNode>(
 					[&SceneImportAssetPath](const FString& NodeUid, UInterchangeLevelFactoryNode* FactoryNode)
@@ -1255,6 +1290,7 @@ void UInterchangeGenericLevelPipeline::ExecutePostImportPipeline(const UIntercha
 		SceneImportAsset->UpdateSceneObjects();
 		AddAssetUserDataToObjects();
 
+		const bool bIsCreateLevelActor = (SceneHierarchyType == EInterchangeSceneHierarchyType::CreateLevelActors);
 		// Nothing to take care of
 		if (PrevSoftObjectPaths.IsEmpty())
 		{
@@ -1281,7 +1317,11 @@ void UInterchangeGenericLevelPipeline::ExecutePostImportPipeline(const UIntercha
 					{
 						if (Object->GetClass()->IsChildOf<AActor>())
 						{
-							ActorsToDelete.Add(Cast<AActor>(Object));
+							//Do not delete ALevelInstance or APackedLevelActor
+							if (bIsCreateLevelActor || !(Object->GetClass()->IsChildOf<ALevelInstance>() || Object->GetClass()->IsChildOf<APackedLevelActor>()))
+							{
+								ActorsToDelete.Add(Cast<AActor>(Object));
+							}
 						}
 						else if(!Object->IsA<UWorld>()) //Avoid deleting UWorld asset
 						{
@@ -1368,9 +1408,35 @@ void UInterchangeGenericLevelPipeline::ExecutePostBroadcastPipeline(const UInter
 				}
 			}
 
-			//Commit an edit to properly finish the setup for the new Level instance actor and new level.
-			LevelInstanceActor->EnterEdit();
-			LevelInstanceActor->ExitEdit();
+			//Update all levelinstance actor that are referencing this reference world
+			for (FThreadSafeObjectIterator It; It; ++It)
+			{
+				if (ALevelInstance* CurrentLevelInstanceActor = Cast<ALevelInstance>(*It))
+				{
+					if (CurrentLevelInstanceActor->GetWorldAsset().Get() == WorldAsset)
+					{
+						if (APackedLevelActor* PackedLevelActor = Cast<APackedLevelActor>(CurrentLevelInstanceActor))
+						{
+							
+							PackedLevelActor->UpdateLevelInstanceFromWorldAsset();
+						}
+						else
+						{
+							CurrentLevelInstanceActor->UpdateLevelInstanceFromWorldAsset();
+							CurrentLevelInstanceActor->LoadLevelInstance();
+						}
+					}
+				}
+			}
+
+			//Reference world must be cleanup since they are not the main world. It was loaded by the Update of the level instance
+			//This remove all the world managers and prevent GC issue when unloading the main world referencing this world.
+			WorldAsset->ClearFlags(RF_Standalone);
+			WorldAsset->ClearInternalFlags(EInternalObjectFlags::Async);
+			if (WorldAsset->bIsWorldInitialized)
+			{
+				WorldAsset->CleanupWorld();
+			}
 		}
 	}
 #endif
@@ -1437,9 +1503,25 @@ bool UInterchangeGenericLevelPipeline::FPostPipelineImportData::UpdateLevelInsta
 
 		ParentWorld->PreEditChange(nullptr);
 
+		bool bIsPackedLevelActor = false;
 		LevelInstanceActor->SetWorldAsset(ReferencedWorld);
-		LevelInstanceActor->UpdateLevelInstanceFromWorldAsset();
-		LevelInstanceActor->LoadLevelInstance();
+		if (APackedLevelActor* PackedLevelActor = Cast<APackedLevelActor>(LevelInstanceActor))
+		{
+			bIsPackedLevelActor = true;
+			constexpr bool bCheckoutAndSaveFalse = false;
+			constexpr bool bPromptForSaveFalse = false;
+			//Get the blueprint
+			UBlueprint* Blueprint = PackedLevelActor->GetRootBlueprint();
+			ULevel::ScanLevelAssets(Blueprint->GetPackage()->GetName());
+			//Update the blueprint with the reference level content
+			FPackedLevelActorBuilder::CreateDefaultBuilder()->CreateOrUpdateBlueprint(ReferencedWorld, Blueprint, bCheckoutAndSaveFalse, bPromptForSaveFalse);
+		}
+
+		if(!bIsPackedLevelActor)
+		{
+			LevelInstanceActor->UpdateLevelInstanceFromWorldAsset();
+			LevelInstanceActor->LoadLevelInstance();
+		}
 
 		//Reference world must be cleanup since they are not the main world.
 		//This remove all the world managers and prevent GC issue when unloading the main world referencing this world.
diff --git a/Engine/Plugins/Interchange/Runtime/Source/Pipelines/Public/InterchangeGenericScenesPipeline.h b/Engine/Plugins/Interchange/Runtime/Source/Pipelines/Public/InterchangeGenericScenesPipeline.h
--- a/Engine/Plugins/Interchange/Runtime/Source/Pipelines/Public/InterchangeGenericScenesPipeline.h
+++ b/Engine/Plugins/Interchange/Runtime/Source/Pipelines/Public/InterchangeGenericScenesPipeline.h
@@ -12,6 +12,7 @@
 #include "InterchangeGenericScenesPipeline.generated.h"
 
 class ALevelInstance;
+class FPreviewScene;
 class UInterchangeActorFactoryNode;
 class UInterchangeLevelFactoryNode;
 class UInterchangeLevelInstanceActorFactoryNode;
@@ -24,14 +25,18 @@ UENUM(BlueprintType)
 enum class EInterchangeSceneHierarchyType : uint8
 {
 	CreateLevelActors UMETA(DisplayName = "Create level actors", ToolTip = "Create actors in the current editor world for all scene nodes in the source hierarchy."),
-	CreateLevelInstanceActor UMETA(DisplayName = "Create a level instance actor", ToolTip = "Create a world containing all scene nodes in the source hierarchy and reference the world has a level instance actor in the current editor world."),
+	CreateLevelInstanceActor UMETA(DisplayName = "Create a level instance actor", ToolTip = "Create a level instance actor referencing a new/existing world containing all scene nodes in the source hierarchy."),
+	CreatePackedActor UMETA(DisplayName = "Create a packed level actor", ToolTip = "Create a packed level actor blueprint which packed all meshes from a new/existing world containing all scene nodes in the source hierarchy.")
 };
 
 UCLASS(BlueprintType, editinlinenew)
 class INTERCHANGEPIPELINES_API UInterchangeGenericLevelPipeline : public UInterchangePipelineBase
 {
 	GENERATED_BODY()
+
 public:
+	virtual ~UInterchangeGenericLevelPipeline();
+
 	/** The name of the pipeline that will be display in the import dialog. */
 	UPROPERTY(EditAnywhere, BlueprintReadWrite, Category = "Scene", meta = (StandAlonePipelineProperty = "True", PipelineInternalEditionData = "True"))
 	FString PipelineDisplayName;
@@ -130,6 +135,8 @@ protected:
 
 	void CacheActiveJointUids();
 	TArray<FString> Cached_ActiveJointUids;
+
+	FPreviewScene* PreviewScene = nullptr;
 };
 
 
diff --git a/Engine/Source/Runtime/Interchange/Core/Public/InterchangeFactoryBase.h b/Engine/Source/Runtime/Interchange/Core/Public/InterchangeFactoryBase.h
--- a/Engine/Source/Runtime/Interchange/Core/Public/InterchangeFactoryBase.h
+++ b/Engine/Source/Runtime/Interchange/Core/Public/InterchangeFactoryBase.h
@@ -232,6 +232,19 @@ public:
 		 * Get all assets we are importing the actor could point on
 		 */
 		TArray<TObjectPtr<UObject>> ImportAssets;
+
+		/**
+		 * true if we are re-importing
+		 */
+		bool bIsReimport = false;
+
+		/*
+		 * This represent the async helper unique id in case a factory need to retrieve the asynchelper associate to the import task
+		 * 
+		 * TODO create a base class for FAsyncHelper in Interchange/Core so we can reference it here instead of using the unique identifier.
+		 * The asynchelper is currently declare in Interchange/Engine so we can't refer to it here.
+		 */
+		int32 AsyncHelperUid = INDEX_NONE;
 	};
 
 	/**
diff --git a/Engine/Source/Runtime/Interchange/Engine/Private/InterchangeManager.cpp b/Engine/Source/Runtime/Interchange/Engine/Private/InterchangeManager.cpp
--- a/Engine/Source/Runtime/Interchange/Engine/Private/InterchangeManager.cpp
+++ b/Engine/Source/Runtime/Interchange/Engine/Private/InterchangeManager.cpp
@@ -1054,8 +1054,8 @@ UInterchangeManager& UInterchangeManager::GetInterchangeManager()
 				InterchangeManager->WaitUntilAllTasksDone(bCancel);
 			}
 
-			//Task should have been cancel in the Engine pre exit callback
-			ensure(InterchangeManager->ImportTasks.Num() == 0);
+			//Task should have been cancel in the Engine pre exit callback.
+			ensure(InterchangeManager->GetImportTaskCount() == 0);
 			InterchangeManager->OnPreDestroyInterchangeManager.Broadcast();
 
 			if (InterchangeManager->QueuedPostImportTasksTickerHandle.IsValid())
@@ -1481,7 +1481,7 @@ void UInterchangeManager::StartQueuedTasks(bool bCancelAllTasks /*= false*/)
 
 		if (Notification.IsValid())
 		{
-			int32 ImportTaskNumber = ImportTasks.Num() + QueueTaskCount;
+			int32 ImportTaskNumber = GetImportTaskCount() + QueueTaskCount;
 			FString ImportTaskNumberStr = TEXT(" (") + FString::FromInt(ImportTaskNumber) + TEXT(")");
 			Notification->SetProgressText(FText::FromString(ImportTaskNumberStr));
 		}
@@ -1489,10 +1489,10 @@ void UInterchangeManager::StartQueuedTasks(bool bCancelAllTasks /*= false*/)
 		{
 			bool bCanShowNotification = false;
 			{
-				int32 ImportTaskCount = ImportTasks.Num();
+				int32 ImportTaskCount = GetImportTaskCount();
 				for (int32 TaskIndex = 0; TaskIndex < ImportTaskCount; ++TaskIndex)
 				{
-					TSharedPtr<UE::Interchange::FImportAsyncHelper, ESPMode::ThreadSafe> AsyncHelper = ImportTasks[TaskIndex];
+					TSharedPtr<UE::Interchange::FImportAsyncHelper, ESPMode::ThreadSafe> AsyncHelper = GetImportTaskForIndex(TaskIndex);
 					if (AsyncHelper.IsValid())
 					{
 						//Allow notification if at least one task is not automated
@@ -1558,7 +1558,7 @@ void UInterchangeManager::StartQueuedTasks(bool bCancelAllTasks /*= false*/)
 		}
 	}
 
-	while (!QueuedTasks.IsEmpty() && (ImportTasks.Num() < MaxNumWorker || bCancelAllTasks))
+	while (!QueuedTasks.IsEmpty() && (GetImportTaskCount() < MaxNumWorker || bCancelAllTasks))
 	{
 		FQueuedTaskData QueuedTaskData;
 		if (QueuedTasks.Dequeue(QueuedTaskData))
@@ -1566,7 +1566,7 @@ void UInterchangeManager::StartQueuedTasks(bool bCancelAllTasks /*= false*/)
 			QueueTaskCount = FMath::Clamp(QueueTaskCount-1, 0, MAX_int32);
 			check(QueuedTaskData.AsyncHelper.IsValid());
 
-			int32 AsyncHelperIndex = ImportTasks.Add(QueuedTaskData.AsyncHelper);
+			int32 AsyncHelperIndex = AddImportTask(QueuedTaskData.AsyncHelper);
 			SetActiveMode(true);
 			//Update the asynchronous notification
 			UpdateNotification();
@@ -2452,7 +2452,8 @@ bool UInterchangeManager::IsObjectBeingImported(UObject* Object) const
 	{
 		return false;
 	}
-
+	FScopeLock Lock(&ImportTasksLock);
+	PRAGMA_DISABLE_DEPRECATION_WARNINGS
 	for (const TSharedPtr < UE::Interchange::FImportAsyncHelper, ESPMode::ThreadSafe>& AsyncHelper : ImportTasks)
 	{
 		if (AsyncHelper->IsImportingObject(Object))
@@ -2460,6 +2461,7 @@ bool UInterchangeManager::IsObjectBeingImported(UObject* Object) const
 			return true;
 		}
 	}
+	PRAGMA_ENABLE_DEPRECATION_WARNINGS
 
 	return false;
 }
@@ -2482,7 +2484,7 @@ bool UInterchangeManager::EnqueuePostImportTask(TSharedPtr<FInterchangePostImpor
 				while (!QueuedPostImportTasks.IsEmpty())
 				{
 					//Wait next frame if we are importing assets or scenes
-					if (!QueuedTasks.IsEmpty() || ImportTasks.Num() > 0)
+					if (!QueuedTasks.IsEmpty() || GetImportTaskCount() > 0)
 					{
 						break;
 					}
@@ -2688,7 +2690,7 @@ void UInterchangeManager::ReleaseAsyncHelper(TWeakPtr<UE::Interchange::FImportAs
 		ForEachResult(AsyncHelperPtr->SceneImportResult->GetResults()->GetResults());
 	}
 
-	ImportTasks.RemoveSingle(AsyncHelper.Pin());
+	RemoveImportTask(AsyncHelper.Pin());
 	//Make sure the async helper is destroy, if not destroy its because we are canceling the import and we still have a shared ptr on it
 	{
 		TSharedPtr<UE::Interchange::FImportAsyncHelper, ESPMode::ThreadSafe> AsyncHelperSharedPtr = AsyncHelper.Pin();
@@ -2698,7 +2700,7 @@ void UInterchangeManager::ReleaseAsyncHelper(TWeakPtr<UE::Interchange::FImportAs
 		}
 	}
 
-	int32 ImportTaskNumber = ImportTasks.Num() + QueueTaskCount;
+	int32 ImportTaskNumber = GetImportTaskCount() + QueueTaskCount;
 	FString ImportTaskNumberStr = TEXT(" (") + FString::FromInt(ImportTaskNumber) + TEXT(")");
 	if (ImportTaskNumber == 0)
 	{
@@ -2752,6 +2754,21 @@ void UInterchangeManager::ReleaseAsyncHelper(TWeakPtr<UE::Interchange::FImportAs
 	StartQueuedTasks();
 }
 
+TSharedPtr<UE::Interchange::FImportAsyncHelper, ESPMode::ThreadSafe> UInterchangeManager::GetAsyncHelper(int32 UniqueId)
+{
+	FScopeLock Lock(&ImportTasksLock);
+	PRAGMA_DISABLE_DEPRECATION_WARNINGS
+	for (TSharedPtr<UE::Interchange::FImportAsyncHelper, ESPMode::ThreadSafe> AsyncHelper : ImportTasks)
+	{
+		if (AsyncHelper.IsValid() && AsyncHelper->UniqueId == UniqueId)
+		{
+			return AsyncHelper;
+		}
+	}
+	PRAGMA_ENABLE_DEPRECATION_WARNINGS
+	return nullptr;
+}
+
 UInterchangeTranslatorBase* UInterchangeManager::GetTranslatorForSourceData(const UInterchangeSourceData* SourceData) const
 {
 	// Find the translator
@@ -2918,13 +2935,16 @@ void UInterchangeManager::CancelAllTasks()
 	StartQueuedTasks(bCancelAllTasks);
 
 	//Set the cancel state on all running tasks
-	int32 ImportTaskCount = ImportTasks.Num();
-	for (int32 TaskIndex = 0; TaskIndex < ImportTaskCount; ++TaskIndex)
 	{
-		TSharedPtr<UE::Interchange::FImportAsyncHelper, ESPMode::ThreadSafe> AsyncHelper = ImportTasks[TaskIndex];
-		if (AsyncHelper.IsValid())
+		FScopeLock Lock(&ImportTasksLock);
+		int32 ImportTaskCount = GetImportTaskCount();
+		for (int32 TaskIndex = 0; TaskIndex < ImportTaskCount; ++TaskIndex)
 		{
-			AsyncHelper->InitCancel();
+			TSharedPtr<UE::Interchange::FImportAsyncHelper, ESPMode::ThreadSafe> AsyncHelper = GetImportTaskForIndex(TaskIndex);
+			if (AsyncHelper.IsValid())
+			{
+				AsyncHelper->InitCancel();
+			}
 		}
 	}
 	for (TPair<UClass*, TArray<FQueuedTaskData>>& ClassAndTasks : NonParallelTranslatorQueueTasks)
@@ -2954,9 +2974,9 @@ void UInterchangeManager::WaitUntilAllTasksDone(bool bCancel)
 		CancelAllTasks();
 	}
 
-	while (ImportTasks.Num() > 0)
+	while (GetImportTaskCount())
 	{
-		TSharedPtr<UE::Interchange::FImportAsyncHelper, ESPMode::ThreadSafe> AsyncHelper = ImportTasks[0];
+		TSharedPtr<UE::Interchange::FImportAsyncHelper, ESPMode::ThreadSafe> AsyncHelper = GetImportTaskForIndex(0);
 		if (AsyncHelper.IsValid())
 		{
 			TWeakPtr<UE::Interchange::FImportAsyncHelper, ESPMode::ThreadSafe> WeakAsyncHelper = AsyncHelper;
@@ -3015,4 +3035,82 @@ TOptional<bool> UInterchangeManager::GetReplaceExistingAlldialogAnswer()
 		Result = (UE::Interchange::Private::StaticHelpers::bReplaceExistingAllDialogAnswer == 1);
 	}
 	return Result;
+}
+
+int32 UInterchangeManager::GetImportTaskCount() const
+{
+	FScopeLock Lock(&ImportTasksLock);
+	PRAGMA_DISABLE_DEPRECATION_WARNINGS
+	return ImportTasks.Num();
+	PRAGMA_ENABLE_DEPRECATION_WARNINGS
+}
+
+TSharedPtr<UE::Interchange::FImportAsyncHelper, ESPMode::ThreadSafe> UInterchangeManager::GetImportTaskForIndex(int32 Index) const
+{
+	FScopeLock Lock(&ImportTasksLock);
+	PRAGMA_DISABLE_DEPRECATION_WARNINGS
+	if (ImportTasks.IsValidIndex(Index))
+	{
+		return ImportTasks[Index];
+	}
+	PRAGMA_ENABLE_DEPRECATION_WARNINGS
+	return nullptr;
+}
+
+TSharedPtr<UE::Interchange::FImportAsyncHelper, ESPMode::ThreadSafe> UInterchangeManager::GetImportTaskForAsyncHelperUniqueId(int32 AsyncHelperUniqueId) const
+{
+	FScopeLock Lock(&ImportTasksLock);
+	PRAGMA_DISABLE_DEPRECATION_WARNINGS
+	for(TSharedPtr<UE::Interchange::FImportAsyncHelper, ESPMode::ThreadSafe> AsyncHelper : ImportTasks)
+	{
+		if (AsyncHelper.IsValid() && AsyncHelper->UniqueId == AsyncHelperUniqueId)
+		{
+			return AsyncHelper;
+		}
+	}
+	PRAGMA_ENABLE_DEPRECATION_WARNINGS
+	return nullptr;
+}
+
+int32 UInterchangeManager::AddImportTask(TSharedPtr<UE::Interchange::FImportAsyncHelper, ESPMode::ThreadSafe> ImportTask)
+{
+	if (!ImportTask.IsValid())
+	{
+		return INDEX_NONE;
+	}
+	FScopeLock Lock(&ImportTasksLock);
+	PRAGMA_DISABLE_DEPRECATION_WARNINGS
+	return ImportTasks.Add(ImportTask);
+	PRAGMA_ENABLE_DEPRECATION_WARNINGS
+}
+
+void UInterchangeManager::RemoveImportTask(TSharedPtr<UE::Interchange::FImportAsyncHelper, ESPMode::ThreadSafe> ImportTask)
+{
+	FScopeLock Lock(&ImportTasksLock);
+	PRAGMA_DISABLE_DEPRECATION_WARNINGS
+	if (ImportTask.IsValid())
+	{
+		ImportTasks.RemoveSingle(ImportTask);
+	}
+	PRAGMA_ENABLE_DEPRECATION_WARNINGS
+}
+
+void UInterchangeManager::RemoveImportTask(int32 AsyncHelperUniqueId)
+{
+	FScopeLock Lock(&ImportTasksLock);
+	TSharedPtr<UE::Interchange::FImportAsyncHelper, ESPMode::ThreadSafe> ToRemoveAsyncHelper;
+	PRAGMA_DISABLE_DEPRECATION_WARNINGS
+	for(TSharedPtr<UE::Interchange::FImportAsyncHelper, ESPMode::ThreadSafe> AsyncHelper : ImportTasks)
+	{
+		if (AsyncHelper.IsValid() && AsyncHelper->UniqueId == AsyncHelperUniqueId)
+		{
+			ToRemoveAsyncHelper = AsyncHelper;
+			break;
+		}
+	}
+	if (ToRemoveAsyncHelper.IsValid())
+	{
+		ImportTasks.Remove(ToRemoveAsyncHelper);
+	}
+	PRAGMA_ENABLE_DEPRECATION_WARNINGS
 }
\ No newline at end of file
diff --git a/Engine/Source/Runtime/Interchange/Engine/Private/Tasks/InterchangeTaskCompletion.cpp b/Engine/Source/Runtime/Interchange/Engine/Private/Tasks/InterchangeTaskCompletion.cpp
--- a/Engine/Source/Runtime/Interchange/Engine/Private/Tasks/InterchangeTaskCompletion.cpp
+++ b/Engine/Source/Runtime/Interchange/Engine/Private/Tasks/InterchangeTaskCompletion.cpp
@@ -84,7 +84,7 @@ void UE::Interchange::FTaskPreCompletion_GameThread::Execute()
 		//First iteration to call SetupObject_GameThread and pipeline ExecutePostFactoryPipeline
 		for (const FImportAsyncHelper::FImportedObjectInfo& ObjectInfo : ImportedObjects)
 		{
-			UObject* ImportedObject = ObjectInfo.ImportedObject;
+			UObject* ImportedObject = ObjectInfo.ImportedObject.TryLoad();
 			//In case Some factory code cannot run outside of the main thread we offer this callback to finish the work before calling post edit change (building the asset)
 			if (bCallPostImportGameThreadCallback && ObjectInfo.Factory)
 			{
@@ -148,7 +148,7 @@ void UE::Interchange::FTaskPreCompletion_GameThread::Execute()
 		//Second iteration to call BuildObject
 		for (const FImportAsyncHelper::FImportedObjectInfo& ObjectInfo : ImportedObjects)
 		{
-			UObject* ImportedObject = ObjectInfo.ImportedObject;
+			UObject* ImportedObject = ObjectInfo.ImportedObject.TryLoad();
 			if (ImportedObject == nullptr || !IsValid(ImportedObject))
 			{
 				continue;
@@ -166,7 +166,7 @@ void UE::Interchange::FTaskPreCompletion_GameThread::Execute()
 		//Third iteration to register the assets
 		for (const FImportAsyncHelper::FImportedObjectInfo& ObjectInfo : ImportedObjects)
 		{
-			UObject* ImportedObject = ObjectInfo.ImportedObject;
+			UObject* ImportedObject = ObjectInfo.ImportedObject.TryLoad();
 			if (ImportedObject == nullptr || !IsValid(ImportedObject))
 			{
 				continue;
@@ -217,7 +217,7 @@ void UE::Interchange::FTaskCompletion_GameThread::Execute()
 			{
 				for (const FImportAsyncHelper::FImportedObjectInfo& AssetInfo : AssetInfos)
 				{
-					if (UObject* Asset = AssetInfo.ImportedObject)
+					if (UObject* Asset = AssetInfo.ImportedObject.TryLoad())
 					{
 						if (Asset->HasAnyFlags(RF_MirroredGarbage))
 						{
@@ -241,7 +241,7 @@ void UE::Interchange::FTaskCompletion_GameThread::Execute()
 				//Do a second pass for the post broadcast pipeline call
 				for (const FImportAsyncHelper::FImportedObjectInfo& AssetInfo : AssetInfos)
 				{
-					if (UObject* Asset = AssetInfo.ImportedObject)
+					if (UObject* Asset = AssetInfo.ImportedObject.TryLoad())
 					{
 						for (UInterchangePipelineBase* PipelineBase : AsyncHelper->Pipelines)
 						{
@@ -261,7 +261,7 @@ void UE::Interchange::FTaskCompletion_GameThread::Execute()
 			{
 				for (const FImportAsyncHelper::FImportedObjectInfo& SceneObjectInfo : AssetInfos)
 				{
-					if (AActor* Actor = Cast<AActor>(SceneObjectInfo.ImportedObject))
+					if (AActor* Actor = Cast<AActor>(SceneObjectInfo.ImportedObject.TryLoad()))
 					{
 						for (UInterchangePipelineBase* PipelineBase : AsyncHelper->Pipelines)
 						{
@@ -287,7 +287,8 @@ void UE::Interchange::FTaskCompletion_GameThread::Execute()
 					using namespace UE::Interchange::Private::ObjectDeletionUtils;
 
 					//Make any created asset go away
-					PurgeObject(AssetInfo.ImportedObject);
+					UObject* ObjectToPurge = AssetInfo.ImportedObject.TryLoad();
+					PurgeObject(ObjectToPurge);
 				}
 			});
 
@@ -296,7 +297,7 @@ void UE::Interchange::FTaskCompletion_GameThread::Execute()
 			{
 				for (const FImportAsyncHelper::FImportedObjectInfo& SceneObjectInfo : AssetInfos)
 				{
-					if (AActor* Actor = Cast<AActor>(SceneObjectInfo.ImportedObject))
+					if (AActor* Actor = Cast<AActor>(SceneObjectInfo.ImportedObject.TryLoad()))
 					{
 						if (UWorld* ActorWorld = Actor->GetWorld())
 						{
diff --git a/Engine/Source/Runtime/Interchange/Engine/Private/Tasks/InterchangeTaskCreateSceneObjects.cpp b/Engine/Source/Runtime/Interchange/Engine/Private/Tasks/InterchangeTaskCreateSceneObjects.cpp
--- a/Engine/Source/Runtime/Interchange/Engine/Private/Tasks/InterchangeTaskCreateSceneObjects.cpp
+++ b/Engine/Source/Runtime/Interchange/Engine/Private/Tasks/InterchangeTaskCreateSceneObjects.cpp
@@ -56,7 +56,11 @@ void UE::Interchange::FTaskCreateSceneObjects_GameThread::Execute()
 		{
 			for (const UE::Interchange::FImportAsyncHelper::FImportedObjectInfo& Info : ImportedObjectInfos)
 			{
-				ImportObjects.Add(Info.ImportedObject);
+				UObject* ImportObject = Info.ImportedObject.TryLoad();
+				if(ImportObject)
+				{
+					ImportObjects.Add(Info.ImportedObject.TryLoad());
+				}
 			}
 		});
 
@@ -112,6 +116,8 @@ void UE::Interchange::FTaskCreateSceneObjects_GameThread::Execute()
 		CreateSceneObjectsParams.ImportAssets = ImportObjects;
 		CreateSceneObjectsParams.ReimportObject = FFactoryCommon::GetObjectToReimport(Factory, ReimportObject, *FactoryNode, FactoryNodeWorldPath, FactoryNodeWorldName, FactoryNodeNodePrefix + SceneNodeName);
 		CreateSceneObjectsParams.ReimportFactoryNode = FFactoryCommon::GetFactoryNode(ReimportObject, FactoryNodeWorldPath, FactoryNodeWorldName, FactoryNodeNodePrefix + SceneNodeName);
+		CreateSceneObjectsParams.bIsReimport = ReimportObject != nullptr;
+		CreateSceneObjectsParams.AsyncHelperUid = AsyncHelper->UniqueId;
 
 		if (AsyncHelper->BaseNodeContainers.IsValidIndex(SourceIndex))
 		{
diff --git a/Engine/Source/Runtime/Interchange/Engine/Private/Tasks/InterchangeTaskPipeline.cpp b/Engine/Source/Runtime/Interchange/Engine/Private/Tasks/InterchangeTaskPipeline.cpp
--- a/Engine/Source/Runtime/Interchange/Engine/Private/Tasks/InterchangeTaskPipeline.cpp
+++ b/Engine/Source/Runtime/Interchange/Engine/Private/Tasks/InterchangeTaskPipeline.cpp
@@ -85,7 +85,7 @@ void UE::Interchange::FTaskWaitAssetCompilation_GameThread::Execute()
 			ImportedObjects.Reserve(ImportedObjects.Num() + ImportedInfos.Num());
 			for (const UE::Interchange::FImportAsyncHelper::FImportedObjectInfo& ImportedInfo : ImportedInfos)
 			{
-				ImportedObjects.Add(ImportedInfo.ImportedObject);
+				ImportedObjects.Add(ImportedInfo.ImportedObject.TryLoad());
 			}
 		};
 
@@ -155,7 +155,7 @@ void UE::Interchange::FTaskPostImport_GameThread::Execute()
 		{
 			for (const UE::Interchange::FImportAsyncHelper::FImportedObjectInfo& ObjectInfo : ImportedInfos)
 			{
-				UObject* ImportedObject = ObjectInfo.ImportedObject;
+				UObject* ImportedObject = ObjectInfo.ImportedObject.TryLoad();
 				if (!ObjectInfo.bPostEditChangeCalled)
 				{
 					ImportedObject->PostEditChange();
@@ -191,7 +191,7 @@ void UE::Interchange::FTaskPostImport_GameThread::Execute()
 
 			for (const UE::Interchange::FImportAsyncHelper::FImportedObjectInfo& ObjectInfo : ImportedInfos)
 			{
-				UObject* ImportedObject = ObjectInfo.ImportedObject;
+				UObject* ImportedObject = ObjectInfo.ImportedObject.TryLoad();
 
 				//In case Some factory code cannot run outside of the main thread we offer this callback to finish the work after asset build is finish.
 				if (bCallPostImportGameThreadCallback && ObjectInfo.Factory)
@@ -209,7 +209,7 @@ void UE::Interchange::FTaskPostImport_GameThread::Execute()
 				}
 
 				NodeUniqueIDs.Add(ObjectInfo.FactoryNode->GetUniqueID());
-				ImportedObjects.Add(ObjectInfo.ImportedObject);
+				ImportedObjects.Add(ObjectInfo.ImportedObject.TryLoad());
 				IsAssetsReimported.Add(ObjectInfo.bIsReimport);
 			}
 		};
diff --git a/Engine/Source/Runtime/Interchange/Engine/Public/InterchangeEditorUtilitiesBase.h b/Engine/Source/Runtime/Interchange/Engine/Public/InterchangeEditorUtilitiesBase.h
--- a/Engine/Source/Runtime/Interchange/Engine/Public/InterchangeEditorUtilitiesBase.h
+++ b/Engine/Source/Runtime/Interchange/Engine/Public/InterchangeEditorUtilitiesBase.h
@@ -13,12 +13,12 @@ class UInterchangeEditorUtilitiesBase : public UObject
 	GENERATED_BODY()
 
 public:
-	INTERCHANGEENGINE_API virtual bool SaveAsset(UObject* Asset)
+	INTERCHANGEENGINE_API virtual bool SaveAsset(UObject* Asset) const
 	{
 		return false;
 	}
 
-	INTERCHANGEENGINE_API virtual bool IsRuntimeOrPIE()
+	INTERCHANGEENGINE_API virtual bool IsRuntimeOrPIE() const
 	{
 #if WITH_EDITOR
 		return false;
@@ -26,4 +26,9 @@ public:
 		return true;
 #endif
 	}
+
+	INTERCHANGEENGINE_API virtual bool ClearEditorSelection() const
+	{
+		return false;
+	}
 };
diff --git a/Engine/Source/Runtime/Interchange/Engine/Public/InterchangeManager.h b/Engine/Source/Runtime/Interchange/Engine/Public/InterchangeManager.h
--- a/Engine/Source/Runtime/Interchange/Engine/Public/InterchangeManager.h
+++ b/Engine/Source/Runtime/Interchange/Engine/Public/InterchangeManager.h
@@ -268,14 +268,14 @@ namespace UE
 
 			struct FImportedObjectInfo
 			{
-				UObject* ImportedObject = nullptr; // The object that was imported.
+				FSoftObjectPath ImportedObject; // The object that was imported.
 				UInterchangeFactoryBase* Factory = nullptr; //The factory that created the imported object.
 				UInterchangeFactoryBaseNode* FactoryNode = nullptr; //The node that describes the object.
 				bool bIsReimport;
 				mutable bool bPostEditChangeCalled = false; //This field is set by the PreCompletionTask and need to be mutable
 			};
 
-			FImportedObjectInfo& AddDefaultImportedAssetGetRef(int32 SourceIndex);
+			INTERCHANGEENGINE_API FImportedObjectInfo& AddDefaultImportedAssetGetRef(int32 SourceIndex);
 			const FImportedObjectInfo* FindImportedAssets(int32 SourceIndex, TFunction< bool(const FImportedObjectInfo& ImportedObjects) > Predicate) const;
 			void IterateImportedAssets(int32 SourceIndex, TFunction< void(const TArray<FImportedObjectInfo>& ImportedObjects) > Callback) const;
 			void IterateImportedAssetsPerSourceIndex(TFunction< void(int32 SourceIndex, const TArray<FImportedObjectInfo>& ImportedObjects) > Callback) const;
@@ -768,6 +768,12 @@ public:
 	/** Delete the specified AsyncHelper and remove it from the array that was holding it. */
 	INTERCHANGEENGINE_API void ReleaseAsyncHelper(TWeakPtr<UE::Interchange::FImportAsyncHelper, ESPMode::ThreadSafe> AsyncHelper);
 
+	/**
+	 * Return a pointer to an FImportAsyncHelper that is the specified UniqueId. Return an empty shared pointer if the unique id is not matching any async helper
+	 * The return asynchelper should be release before the end of the import process to make sure all resource are release when the import is completed.
+	 */
+	INTERCHANGEENGINE_API TSharedPtr<UE::Interchange::FImportAsyncHelper, ESPMode::ThreadSafe> GetAsyncHelper(int32 UniqueId);
+
 	/*
 	 * Return the first translator that can translate the source data.
 	 * @Param SourceData - The source data that you want a translator for.
@@ -892,6 +898,15 @@ protected:
 	ImportInternal(const FString& ContentPath, const UInterchangeSourceData* SourceData, const FImportAssetParameters& ImportAssetParameters, const UE::Interchange::EImportType ImportType);
 
 	TOptional<FString> ValidateReimportParameter(UObject* ObjectToReimport, const FImportAssetParameters& ImportAssetParameters, bool bRunSynchronous);
+
+	int32 GetImportTaskCount() const;
+	TSharedPtr<UE::Interchange::FImportAsyncHelper, ESPMode::ThreadSafe> GetImportTaskForIndex(int32 Index) const;
+	TSharedPtr<UE::Interchange::FImportAsyncHelper, ESPMode::ThreadSafe> GetImportTaskForAsyncHelperUniqueId(int32 AsyncHelperUniqueId) const;
+	/* Return the index in the ImportTasks array */
+	int32 AddImportTask(TSharedPtr<UE::Interchange::FImportAsyncHelper, ESPMode::ThreadSafe> ImportTask);
+	void RemoveImportTask(TSharedPtr<UE::Interchange::FImportAsyncHelper, ESPMode::ThreadSafe> ImportTask);
+	void RemoveImportTask(int32 AsyncHelperUniqueId);
+
 private:
 
 	static bool bIsCreatingSingleton;
@@ -917,7 +932,10 @@ private:
 	//Ticker which is on only if we have some QueuedPostImportTasks
 	FTSTicker::FDelegateHandle	QueuedPostImportTasksTickerHandle;
 
+	/** Critical section to make the import tasks array thread safe */
+	mutable FCriticalSection ImportTasksLock;
 	//By using pointer, there is no issue if the array get resized.
+	UE_DEPRECATED(5.6, "Use the provided API to manipulate the Import tasks array. This array need to be thread safe")
 	TArray<TSharedPtr<UE::Interchange::FImportAsyncHelper, ESPMode::ThreadSafe> > ImportTasks;
 
 	TSharedPtr<FAsyncTaskNotification> Notification = nullptr;
