Interchange:
- Move InterchangeUsdDefinitions to Common module

diff --git a/Engine/Plugins/Interchange/Runtime/Source/Common/Public/Usd/InterchangeUsdDefinitions.h b/Engine/Plugins/Interchange/Runtime/Source/Common/Public/Usd/InterchangeUsdDefinitions.h
new file mode 100644
--- /dev/null
+++ b/Engine/Plugins/Interchange/Runtime/Source/Common/Public/Usd/InterchangeUsdDefinitions.h
@@ -0,0 +1,37 @@
+// Copyright Epic Games, Inc. All Rights Reserved.
+
+#pragma once
+
+#include "CoreMinimal.h"
+
+namespace UE::Interchange::USD
+{
+	const FString USDContextTag = TEXT("USD");
+
+	// Name of a custom attribute added to translated nodes to contain their geometry purpose (proxy, render, guide, etc.)
+	const FString GeometryPurposeIdentifier = TEXT("USD_Geometry_Purpose");
+
+	// Prim kind metadata
+	const FString KindIdentifier = TEXT("USD_Kind");
+
+	// Prefix for custom attributes that describe material instance parameters
+	const FString AttributeIdentifierPrefix = TEXT("USD_MI_Attribute_");
+
+	// Flag indicating whether we should parse a UInterchangeShaderGraphNode on the InterchangeUSDPipeline
+	const FString ParseMaterialIdentifier = TEXT("USD_MI_ParseMaterial");
+
+	// Means we should pick this material to use as parent material, regardless of what else is provided
+	// (NeedsVTParent, DisplayColor, etc.)
+	const FString ParentMaterialIdentifier = TEXT("USD_MI_Parent");
+
+	// Means that we should use a display color material parent for this material, according to the
+	// value provided in this attribute
+	const FString DisplayColorDescIdentifier = TEXT("USD_MI_DisplayColorDescription");
+
+	// Means that if we were to pick one of the UsdPreviewSurface parent materials for this material
+	// translated node, we should pick the "VT" variant of that material.
+	//
+	// Note: To indicate TwoSided/Translucent we use the CustomTwoSided/BlendMode standard UInterchangeShaderGraphNode
+	// attributes.
+	const FString NeedsVTParentIdentifier = TEXT("USD_MI_NeedsVTParent");
+}
diff --git a/Engine/Plugins/Interchange/Runtime/Source/Import/Public/Usd/InterchangeUsdDefinitions.h b/Engine/Plugins/Interchange/Runtime/Source/Import/Public/Usd/InterchangeUsdDefinitions.h
deleted file mode 100644
