[Interchange][FBX] Fail to import attributes animation curves with key values as integer.

#rnx

diff --git a/Engine/Plugins/Interchange/Runtime/Source/Import/Private/Animation/InterchangeAnimSequenceFactory.cpp b/Engine/Plugins/Interchange/Runtime/Source/Import/Private/Animation/InterchangeAnimSequenceFactory.cpp
--- a/Engine/Plugins/Interchange/Runtime/Source/Import/Private/Animation/InterchangeAnimSequenceFactory.cpp
+++ b/Engine/Plugins/Interchange/Runtime/Source/Import/Private/Animation/InterchangeAnimSequenceFactory.cpp
@@ -59,12 +59,128 @@
 		return true;
 	}
 
+	template <class T>
+	bool ConvertToRichCurve(const TOptional<TArray<T>>& OptionalValues, const TArray<float>& Keys, FRichCurve& OutRichCurve)
+	{
+		if (!OptionalValues.IsSet())
+		{
+			return false;
+		}
+
+		const TArray<T>& Values = OptionalValues.GetValue();
+		if (Values.Num() != Keys.Num())
+		{
+			return false;
+		}
+
+		OutRichCurve.Keys.Reserve(Keys.Num());
+
+		for (size_t CurveEntryIndex = 0; CurveEntryIndex < Keys.Num(); CurveEntryIndex++)
+		{
+			FKeyHandle RichCurveKeyHandle = OutRichCurve.AddKey(Keys[CurveEntryIndex], Values[CurveEntryIndex]);
+			FRichCurveKey& RichCurveKey = OutRichCurve.GetKey(RichCurveKeyHandle);
+			RichCurveKey.InterpMode = ERichCurveInterpMode::RCIM_Constant;
+		}
+
+		OutRichCurve.AutoSetTangents();
+
+		return true;
+	}
+
+	bool InternalCreateCurve(UAnimSequence* TargetSequence
+		, TArray<FRichCurve>& Curves
+		, const TArray<FString>& CurveNames
+		, const int32 CurveFlags
+		, const bool bDoNotImportCurveWithZero
+		, const bool bAddCurveMetadataToSkeleton
+		, const bool bMorphTargetCurve
+		, const bool bMaterialCurve
+		, const bool bShouldTransact)
+	{
+		bool bResult = false;
+		if (!TargetSequence || CurveNames.Num() <= 0 || CurveNames.Num() != Curves.Num())
+		{
+			return bResult;
+		}
+
+
+		for (int32 CurveIndex = 0; CurveIndex < Curves.Num(); ++CurveIndex)
+		{
+			FRichCurve& Curve = Curves[CurveIndex];
+			FName Name = *CurveNames[CurveIndex];
+
+			if (bDoNotImportCurveWithZero)
+			{
+				bool bAllCurveValueAreZero = true;
+				FKeyHandle KeyHandle = Curve.GetFirstKeyHandle();
+				while (KeyHandle != FKeyHandle::Invalid())
+				{
+					if (!FMath::IsNearlyZero(Curve.GetKeyValue(KeyHandle)))
+					{
+						bAllCurveValueAreZero = false;
+						break;
+					}
+					KeyHandle = Curve.GetNextKey(KeyHandle);
+				}
+				if (bAllCurveValueAreZero)
+				{
+					continue;
+				}
+			}
+			
+			FAnimationCurveIdentifier FloatCurveId(Name, ERawCurveTrackTypes::RCT_Float);
+
+			IAnimationDataModel* DataModel = TargetSequence->GetDataModel();
+			IAnimationDataController& Controller = TargetSequence->GetController();
+
+			const FFloatCurve* TargetCurve = DataModel->FindFloatCurve(FloatCurveId);
+			if (TargetCurve == nullptr)
+			{
+				// Need to add the curve first
+				Controller.AddCurve(FloatCurveId, AACF_DefaultCurve | CurveFlags, bShouldTransact);
+				TargetCurve = DataModel->FindFloatCurve(FloatCurveId);
+			}
+			else
+			{
+				// Need to update any of the flags
+				Controller.SetCurveFlags(FloatCurveId, CurveFlags | TargetCurve->GetCurveTypeFlags(), bShouldTransact);
+			}
+
+			// Should be valid at this point
+			ensure(TargetCurve);
+
+			//MorphTarget curves are shift to 0 second
+			const float CurveStartTime = Curve.GetFirstKey().Time;
+			if (CurveStartTime > UE_SMALL_NUMBER)
+			{
+				Curve.ShiftCurve(-CurveStartTime);
+			}
+			// Set actual keys on curve within the model
+			Controller.SetCurveKeys(FloatCurveId, Curve.GetConstRefOfKeys(), bShouldTransact);
+
+			if (bMaterialCurve || bMorphTargetCurve)
+			{
+				if (bAddCurveMetadataToSkeleton)
+				{
+					USkeleton* Skeleton = TargetSequence->GetSkeleton();
+					Skeleton->AccumulateCurveMetaData(Name, bMaterialCurve, bMorphTargetCurve);
+				}
+			}
+			bResult = true;
+		}
+		return bResult;
+	}
+
 	bool CreateAttributeStepCurve(UAnimSequence* TargetSequence
 		, TArray<FInterchangeStepCurve>& StepCurves
 		, const FString& CurveName
 		, const FString& BoneName
 		, const int32 CurveFlags
-		, const bool bDoNotImportCurveWithZero)
+		, const bool bDoNotImportCurveWithZero
+		, const bool bAddCurveMetadataToSkeleton
+		, const bool bIsMorphTargetCurve
+		, const bool bIsMaterialCurve
+		, const bool bShouldTransact)
 	{
 		//For bone attribute we support only single curve type (structured type like vector are not allowed)
 		if (!TargetSequence || StepCurves.Num() != 1 || CurveName.IsEmpty())
@@ -117,6 +233,10 @@
 			}
 		}
 
+		TArray<FRichCurve> RichCurves;
+		TArray<FString> RichCurveNames;
+		RichCurves.Reserve(StepCurves.Num());
+
 		bool bResult = false;
 		for (const FInterchangeStepCurve& StepCurve : StepCurves)
 		{
@@ -144,91 +264,34 @@
 					, MakeArrayView(StepCurve.KeyTimes)
 					, MakeArrayView(StepCurve.StringKeyValues.GetValue()));
 			}
+
+			FRichCurve RichCurve;
+			bool bRichCurveConverted = ConvertToRichCurve(StepCurve.FloatKeyValues, StepCurve.KeyTimes, RichCurve)
+				|| ConvertToRichCurve(StepCurve.IntegerKeyValues, StepCurve.KeyTimes, RichCurve)
+				|| ConvertToRichCurve(StepCurve.ByteKeyValues, StepCurve.KeyTimes, RichCurve)
+				|| ConvertToRichCurve(StepCurve.BooleanKeyValues, StepCurve.KeyTimes, RichCurve);
+
+			if (bRichCurveConverted)
+			{
+				RichCurves.Add(RichCurve);
+				if (RichCurveNames.Num() > 0)
+				{
+					//As CurveNames are used as uids:
+					FString CurveNamePerCurve = CurveName + TEXT("_") + FString::FromInt(RichCurveNames.Num());
+					RichCurveNames.Add(CurveNamePerCurve);
+				}
+				else
+				{
+					RichCurveNames.Add(CurveName);
+				}
+			}
 		}
+
+		if (RichCurves.Num() > 0)
+		{
+			bResult = InternalCreateCurve(TargetSequence, RichCurves, RichCurveNames, CurveFlags, bDoNotImportCurveWithZero, bAddCurveMetadataToSkeleton, bIsMorphTargetCurve, bIsMaterialCurve, bShouldTransact);
+		}
+
 		return bResult;
 	}
 
-	bool InternalCreateCurve(UAnimSequence* TargetSequence
-		, TArray<FRichCurve>& Curves
-		, const TArray<FString>& CurveNames
-		, const int32 CurveFlags
-		, const bool bDoNotImportCurveWithZero
-		, const bool bAddCurveMetadataToSkeleton
-		, const bool bMorphTargetCurve
-		, const bool bMaterialCurve
-		, const bool bShouldTransact)
-	{
-		bool bResult = false;
-		if (!TargetSequence || CurveNames.Num() <= 0 || CurveNames.Num() != Curves.Num())
-		{
-			return bResult;
-		}
-
-
-		for (int32 CurveIndex = 0; CurveIndex < Curves.Num(); ++CurveIndex)
-		{
-			FRichCurve& Curve = Curves[CurveIndex];
-			FName Name = *CurveNames[CurveIndex];
-
-			if (bDoNotImportCurveWithZero)
-			{
-				bool bAllCurveValueAreZero = true;
-				FKeyHandle KeyHandle = Curve.GetFirstKeyHandle();
-				while (KeyHandle != FKeyHandle::Invalid())
-				{
-					if (!FMath::IsNearlyZero(Curve.GetKeyValue(KeyHandle)))
-					{
-						bAllCurveValueAreZero = false;
-						break;
-					}
-					KeyHandle = Curve.GetNextKey(KeyHandle);
-				}
-				if (bAllCurveValueAreZero)
-				{
-					continue;
-				}
-			}
-			
-			FAnimationCurveIdentifier FloatCurveId(Name, ERawCurveTrackTypes::RCT_Float);
-
-			IAnimationDataModel* DataModel = TargetSequence->GetDataModel();
-			IAnimationDataController& Controller = TargetSequence->GetController();
-
-			const FFloatCurve* TargetCurve = DataModel->FindFloatCurve(FloatCurveId);
-			if (TargetCurve == nullptr)
-			{
-				// Need to add the curve first
-				Controller.AddCurve(FloatCurveId, AACF_DefaultCurve | CurveFlags, bShouldTransact);
-				TargetCurve = DataModel->FindFloatCurve(FloatCurveId);
-			}
-			else
-			{
-				// Need to update any of the flags
-				Controller.SetCurveFlags(FloatCurveId, CurveFlags | TargetCurve->GetCurveTypeFlags(), bShouldTransact);
-			}
-
-			// Should be valid at this point
-			ensure(TargetCurve);
-
-			//MorphTarget curves are shift to 0 second
-			const float CurveStartTime = Curve.GetFirstKey().Time;
-			if (CurveStartTime > UE_SMALL_NUMBER)
-			{
-				Curve.ShiftCurve(-CurveStartTime);
-			}
-			// Set actual keys on curve within the model
-			Controller.SetCurveKeys(FloatCurveId, Curve.GetConstRefOfKeys(), bShouldTransact);
-
-			if (bMaterialCurve || bMorphTargetCurve)
-			{
-				if (bAddCurveMetadataToSkeleton)
-				{
-					USkeleton* Skeleton = TargetSequence->GetSkeleton();
-					Skeleton->AccumulateCurveMetaData(Name, bMaterialCurve, bMorphTargetCurve);
-				}
-			}
-			bResult = true;
-		}
-		return bResult;
-	}
-
@@ -753,6 +816,14 @@
 					}
 					return false;
 				};
+
+				//Get if the SkeletonSceneNode has a _AnimationPayloadType attribute set, and if so, does it equal to input
+				auto DoesSourceAnimationAllowCurve = [](const FString& AttibuteName, const UInterchangeSceneNode* SkeletonSceneNode, EInterchangeAnimationPayLoadType AnimationPayloadTypeToCheck) -> bool
+					{
+						EInterchangeAnimationPayLoadType AnimationPayloadType = EInterchangeAnimationPayLoadType::NONE;
+						return SkeletonSceneNode->GetAnimationCurveTypeForCurveName(AttibuteName, AnimationPayloadType) && (AnimationPayloadType == AnimationPayloadTypeToCheck);
+					};
+
 				//Import Attribute curves (FRichCurve)
 				{
 					TArray<FString> AttributeCurveNames;
@@ -769,7 +840,7 @@
 							for (const FInterchangeUserDefinedAttributeInfo& AttributeInfo : AttributeInfos)
 							{
 								//Material curve must be convertible to float since we need a FRichCurve
-								if (AttributeInfo.PayloadKey.IsSet() && IsDecimalType(AttributeInfo.Type) && AttributeCurveNames.Contains(AttributeInfo.Name))
+								if (AttributeInfo.PayloadKey.IsSet() && (IsDecimalType(AttributeInfo.Type) || DoesSourceAnimationAllowCurve(AttributeInfo.Name, SkeletonSceneNode, EInterchangeAnimationPayLoadType::CURVE)) && AttributeCurveNames.Contains(AttributeInfo.Name))
 								{
 									CurveNamePayloads.Add(AttributeInfo.Name, AttributeInfo.PayloadKey.GetValue());
 								}
@@ -836,7 +907,7 @@
 								//Material curve must be convertible to float since we need a FRichCurve
 								if (AttributeInfo.PayloadKey.IsSet() && AttributeStepCurveNames.Contains(AttributeInfo.Name))
 								{
-									if (IsDecimalType(AttributeInfo.Type))
+									if (IsDecimalType(AttributeInfo.Type) && !(DoesSourceAnimationAllowCurve(AttributeInfo.Name, SkeletonSceneNode, EInterchangeAnimationPayLoadType::STEPCURVE)))
 									{
 										CurveNameFloatPayloads.Add(AttributeInfo.Name, AttributeInfo.PayloadKey.GetValue());
 									}
@@ -867,7 +938,7 @@
 						}
 					}
 
-					auto AddAttributeCurveToAnimSequence = [bRemoveCurveRedundantKeys, bDoNotImportCurveWithZero, &AnimSequence](FAnimationPayloadData& StepCurvePayload, const FString& CurveName, const FString& BoneName)
+					auto AddAttributeCurveToAnimSequence = [bRemoveCurveRedundantKeys, bDoNotImportCurveWithZero, bAddCurveMetadataToSkeleton, bShouldTransact, &AnimSequence](FAnimationPayloadData& StepCurvePayload, const FString& CurveName, const FString& BoneName)
 					{
 						if (bRemoveCurveRedundantKeys)
 						{
@@ -877,8 +948,11 @@
 							}
 						}
 
+						constexpr bool bIsMorphTargetCurve = false;
+						constexpr bool bIsMaterialCurve = false;
+
 						constexpr int32 CurveFlags = 0;
-						CreateAttributeStepCurve(AnimSequence, StepCurvePayload.StepCurves, CurveName, BoneName, CurveFlags, bDoNotImportCurveWithZero);
+						CreateAttributeStepCurve(AnimSequence, StepCurvePayload.StepCurves, CurveName, BoneName, CurveFlags, bDoNotImportCurveWithZero, bAddCurveMetadataToSkeleton, bIsMorphTargetCurve, bIsMaterialCurve, bShouldTransact);
 					};
 
 					TArray<FAnimationPayloadData> CurvePayloads = AnimSequenceTranslatorPayloadInterface->GetAnimationPayloadData(CurvePayloadQueries);
diff --git a/Engine/Plugins/Interchange/Runtime/Source/Nodes/Private/InterchangeSceneNode.cpp b/Engine/Plugins/Interchange/Runtime/Source/Nodes/Private/InterchangeSceneNode.cpp
--- a/Engine/Plugins/Interchange/Runtime/Source/Nodes/Private/InterchangeSceneNode.cpp
+++ b/Engine/Plugins/Interchange/Runtime/Source/Nodes/Private/InterchangeSceneNode.cpp
@@ -70,6 +70,12 @@ namespace UE
 			static FString TagsKey(TEXT("__Tags__Key"));
 			return TagsKey;
 		}
+
+		const FString& FSceneNodeStaticData::GetCurveAnimationTypesKey()
+		{
+			static FString CurveAnimationTypesKey(TEXT("__CurveAnimationTypes__Key"));
+			return CurveAnimationTypesKey;
+		}
 	}//ns Interchange
 }//ns UE
 
@@ -81,6 +87,7 @@ UInterchangeSceneNode::UInterchangeSceneNode()
 	MorphTargetCurveWeights.Initialize(Attributes.ToSharedRef(), UE::Interchange::FSceneNodeStaticData::GetMorphTargetCurveWeightsKey());
 	LayerNames.Initialize(Attributes, UE::Interchange::FSceneNodeStaticData::GetLayerNamesKey());
 	Tags.Initialize(Attributes, UE::Interchange::FSceneNodeStaticData::GetTagsKey());
+	CurveAnimationTypes.Initialize(Attributes.ToSharedRef(), UE::Interchange::FSceneNodeStaticData::GetCurveAnimationTypesKey());
 }
 
 /**
@@ -511,3 +518,13 @@ bool UInterchangeSceneNode::RemoveTag(const FString& Tag)
 {
 	return Tags.RemoveItem(Tag);
 }
+
+bool UInterchangeSceneNode::SetAnimationCurveTypeForCurveName(const FString& CurveName, const EInterchangeAnimationPayLoadType& AnimationCurveType)
+{
+	return CurveAnimationTypes.SetKeyValue(CurveName, AnimationCurveType);
+}
+
+bool UInterchangeSceneNode::GetAnimationCurveTypeForCurveName(const FString& CurveName, EInterchangeAnimationPayLoadType& OutCurveAnimationType) const
+{
+	return CurveAnimationTypes.GetValue(CurveName, OutCurveAnimationType);
+}
\ No newline at end of file
diff --git a/Engine/Plugins/Interchange/Runtime/Source/Nodes/Public/InterchangeSceneNode.h b/Engine/Plugins/Interchange/Runtime/Source/Nodes/Public/InterchangeSceneNode.h
--- a/Engine/Plugins/Interchange/Runtime/Source/Nodes/Public/InterchangeSceneNode.h
+++ b/Engine/Plugins/Interchange/Runtime/Source/Nodes/Public/InterchangeSceneNode.h
@@ -5,6 +5,7 @@
 #include "CoreMinimal.h"
 #include "Misc/Optional.h"
 #include "Nodes/InterchangeBaseNode.h"
+#include "InterchangeAnimationTrackSetNode.h"
 
 #include "InterchangeSceneNode.generated.h"
 
@@ -27,6 +28,7 @@ namespace UE
 			static const FString& GetMorphTargetCurveWeightsKey();
 			static const FString& GetLayerNamesKey();
 			static const FString& GetTagsKey();
+			static const FString& GetCurveAnimationTypesKey();
 		};
 
 	}//ns Interchange
@@ -275,6 +277,15 @@ public:
 	/** Remove Tag that this SceneNode (Actor) is supposed to have. */
 	bool RemoveTag(const FString& Tag);
 
+	
+	/** Sets the Animation Curve Type for the given CurveName (StepCurve or Curve). (Mostly used for tracking Custom Attributes' Animation Types) */
+	UFUNCTION(BlueprintCallable, Category = "Interchange | Node | Joint")
+	bool SetAnimationCurveTypeForCurveName(const FString& CurveName, const EInterchangeAnimationPayLoadType& AnimationCurveType);
+
+	/** Gets the Animation Curve Type for the given CurveName. */
+	UFUNCTION(BlueprintCallable, Category = "Interchange | Node | Joint")
+	bool GetAnimationCurveTypeForCurveName(const FString& CurveName, EInterchangeAnimationPayLoadType& OutCurveAnimationType) const;
+
 private:
 
 	bool GetGlobalTransformInternal(const UE::Interchange::FAttributeKey LocalTransformKey
@@ -327,6 +338,9 @@ private:
 	//A scene node can have multiple Tags.
 	UE::Interchange::TArrayAttributeHelper<FString> Tags;
 
+	//A scene node can have Attributes defining Curves. (Mostly used for tracking Custom Attributes' Animation Types).
+	UE::Interchange::TMapAttributeHelper<FString, EInterchangeAnimationPayLoadType> CurveAnimationTypes;
+
 	//A scene node can reference an animation asset on top of base asset:
 	const UE::Interchange::FAttributeKey Macro_CustomAnimationAssetUidToPlayKey = UE::Interchange::FAttributeKey(TEXT("AnimationAssetUidToPlay"));
 
diff --git a/Engine/Plugins/Interchange/Runtime/Source/Parsers/Fbx/Private/FbxAnimation.cpp b/Engine/Plugins/Interchange/Runtime/Source/Parsers/Fbx/Private/FbxAnimation.cpp
--- a/Engine/Plugins/Interchange/Runtime/Source/Parsers/Fbx/Private/FbxAnimation.cpp
+++ b/Engine/Plugins/Interchange/Runtime/Source/Parsers/Fbx/Private/FbxAnimation.cpp
@@ -27,6 +27,35 @@
 
 namespace UE::Interchange::Private
 {
+	bool IsStepCurve(FbxAnimCurveNode* AnimCurveNode)
+	{
+		const uint32 ChannelCount = AnimCurveNode->GetChannelsCount();
+		for (uint32 ChannelIndex = 0; ChannelIndex < ChannelCount; ++ChannelIndex)
+		{
+			const uint32 ChannelCurveCount = AnimCurveNode->GetCurveCount(ChannelIndex);
+			for (uint32 CurveIndex = 0; CurveIndex < ChannelCurveCount; ++CurveIndex)
+			{
+				if (FbxAnimCurve* CurrentAnimCurve = AnimCurveNode->GetCurve(ChannelIndex, CurveIndex))
+				{
+					int32 KeyCount = CurrentAnimCurve->KeyGetCount();
+					for (int32 KeyIndex = 0; KeyIndex < KeyCount; ++KeyIndex)
+					{
+						FbxAnimCurveKey Key = CurrentAnimCurve->KeyGet(KeyIndex);
+
+						FbxAnimCurveDef::EInterpolationType KeyInterpMode = Key.GetInterpolation();
+
+						if (KeyInterpMode != FbxAnimCurveDef::EInterpolationType::eInterpolationConstant)
+						{
+							return false;
+						}
+					}
+				}
+			}
+		}
+
+		return true;
+	}
+
 	bool ImportCurve(const FbxAnimCurve* SourceFloatCurves, const float ScaleValue, TArray<FInterchangeCurveKey>& DestinationFloatCurve)
 	{
 		if (!SourceFloatCurves)
@@ -833,7 +862,8 @@ namespace UE::Interchange::Private
 		, UInterchangeSceneNode* SceneNode
 		, TMap<FString, TSharedPtr<FPayloadContextBase>>& PayloadContexts
 		, EFbxType PropertyType
-		, TOptional<FString>& OutPayloadKey)
+		, TOptional<FString>& OutPayloadKey
+		, TOptional<bool>& OutIsStepCurve)
 	{
 		const FString PropertyName = Parser.GetFbxHelper()->GetFbxPropertyName(Property);
 		const FString PayLoadKey = Parser.GetFbxHelper()->GetFbxNodeHierarchyName(Node) + PropertyName + TEXT("_AnimationPayloadKey");
@@ -844,8 +874,9 @@ namespace UE::Interchange::Private
 			FetchPayloadData.Node = Node;
 			FetchPayloadData.AnimCurves = AnimCurveNode;
 			
-			//Any property that is not decimal should import with step curve
-			FetchPayloadData.bAttributeTypeIsStepCurveAnimation = !IsFbxPropertyTypeDecimal(PropertyType);
+			//Only curves with Constant interpolation on all keys are deemed as StepCurves.
+			OutIsStepCurve = !IsFbxPropertyTypeDecimal(PropertyType) && IsStepCurve(AnimCurveNode);
+			FetchPayloadData.bAttributeTypeIsStepCurveAnimation = OutIsStepCurve.GetValue();
 			FetchPayloadData.PropertyType = PropertyType;
 			FetchPayloadData.Property = Property;
 			AnimPayload->AttributeFetchPayloadData = FetchPayloadData;
diff --git a/Engine/Plugins/Interchange/Runtime/Source/Parsers/Fbx/Private/FbxAnimation.h b/Engine/Plugins/Interchange/Runtime/Source/Parsers/Fbx/Private/FbxAnimation.h
--- a/Engine/Plugins/Interchange/Runtime/Source/Parsers/Fbx/Private/FbxAnimation.h
+++ b/Engine/Plugins/Interchange/Runtime/Source/Parsers/Fbx/Private/FbxAnimation.h
@@ -116,7 +116,8 @@ namespace UE::Interchange::Private
 			, UInterchangeSceneNode* SceneNode
 			, TMap<FString, TSharedPtr<FPayloadContextBase>>& PayloadContexts
 			, EFbxType PropertyType
-			, TOptional<FString>& OutPayloadKey);
+			, TOptional<FString>& OutPayloadKey
+			, TOptional<bool>& OutIsStepCurve);
 
 		static void AddRigidTransformAnimation(FFbxParser& Parser
 			, FbxNode* Node
diff --git a/Engine/Plugins/Interchange/Runtime/Source/Parsers/Fbx/Private/FbxScene.cpp b/Engine/Plugins/Interchange/Runtime/Source/Parsers/Fbx/Private/FbxScene.cpp
--- a/Engine/Plugins/Interchange/Runtime/Source/Parsers/Fbx/Private/FbxScene.cpp
+++ b/Engine/Plugins/Interchange/Runtime/Source/Parsers/Fbx/Private/FbxScene.cpp
@@ -396,30 +396,6 @@ namespace UE
 					NodeContainer.AddNode(AnimTrackNode);
 				};
 
-				//Bool/Enum/Integer curves are treated as StepCurves, otherwise it's a floating point curve
-				auto GetPropertyPayloadType = [](EFbxType FbxType)
-				{
-					switch(FbxType)
-					{
-					case EFbxType::eFbxBool:
-					case EFbxType::eFbxChar:
-					case EFbxType::eFbxUChar:
-					case EFbxType::eFbxShort:
-					case EFbxType::eFbxUShort:
-					case EFbxType::eFbxInt:
-					case EFbxType::eFbxUInt:
-					case EFbxType::eFbxLongLong:
-					case EFbxType::eFbxULongLong:
-					case EFbxType::eFbxEnum:
-					case EFbxType::eFbxEnumM:
-					case EFbxType::eFbxString:
-						return EInterchangeAnimationPayLoadType::STEPCURVE;
-
-					default:
-						return EInterchangeAnimationPayLoadType::CURVE;
-					}
-				};
-
 				const UInterchangeFbxSettings* InterchangeFbxSettings = GetDefault<UInterchangeFbxSettings>();
 
 				//Add all Node Attributes for the node
@@ -435,16 +411,22 @@ namespace UE
 						if(CurveNode && CurveNode->IsAnimated() && FFbxAnimation::IsFbxPropertyTypeSupported(PropertyType))
 						{
 							TOptional<FString> PayloadKey;
+							TOptional<bool> bIsStepCurve;
 							//Attribute is animated, add the curves payload key that represent the attribute animation
-							FFbxAnimation::AddNodeAttributeCurvesAnimation(Parser, Node, Property, CurveNode, UnrealNode, PayloadContexts, PropertyType, PayloadKey);
+							FFbxAnimation::AddNodeAttributeCurvesAnimation(Parser, Node, Property, CurveNode, UnrealNode, PayloadContexts, PropertyType, PayloadKey, bIsStepCurve);
 
-							if(PayloadKey.IsSet())
+							if(PayloadKey.IsSet() && bIsStepCurve.IsSet())
 							{
 								const char* CurveNodeName = CurveNode->GetName();								
-								if(EInterchangePropertyTracks PropertyTrack = InterchangeFbxSettings->GetPropertyTrack(CurveNodeName); PropertyTrack != EInterchangePropertyTracks::None)
+								
+								EInterchangeAnimationPayLoadType PayloadType = bIsStepCurve.GetValue() ? EInterchangeAnimationPayLoadType::STEPCURVE : EInterchangeAnimationPayLoadType::CURVE;
+
+								if (EInterchangePropertyTracks PropertyTrack = InterchangeFbxSettings->GetPropertyTrack(CurveNodeName); PropertyTrack != EInterchangePropertyTracks::None)
 								{
-									AddAnimationTrackNode(PropertyTrack, CurveNodeName, *PayloadKey, GetPropertyPayloadType(PropertyType));
+									AddAnimationTrackNode(PropertyTrack, CurveNodeName, *PayloadKey, PayloadType);
 								}
+
+								UnrealNode->SetAnimationCurveTypeForCurveName(CurveNodeName, PayloadType);
 							}
 						}
 
@@ -462,18 +444,24 @@ namespace UE
 					{
 						FbxAnimCurveNode* CurveNode = Property.GetCurveNode();
 						TOptional<FString> PayloadKey;
+						TOptional<bool> bIsStepCurve;
 						if (CurveNode && CurveNode->IsAnimated())
 						{
 							//Attribute is animated, add the curves payload key that represent the attribute animation
-							FFbxAnimation::AddNodeAttributeCurvesAnimation(Parser, Node, Property, CurveNode, UnrealNode, PayloadContexts, PropertyType, PayloadKey);
+							FFbxAnimation::AddNodeAttributeCurvesAnimation(Parser, Node, Property, CurveNode, UnrealNode, PayloadContexts, PropertyType, PayloadKey, bIsStepCurve);
 
-							if(PayloadKey.IsSet())
+							if(PayloadKey.IsSet() && bIsStepCurve.IsSet())
 							{
 								const char* CurveNodeName = CurveNode->GetName();
-								if(EInterchangePropertyTracks PropertyTrack = InterchangeFbxSettings->GetPropertyTrack(CurveNodeName); PropertyTrack != EInterchangePropertyTracks::None)
+								
+								EInterchangeAnimationPayLoadType PayloadType = bIsStepCurve.GetValue() ? EInterchangeAnimationPayLoadType::STEPCURVE : EInterchangeAnimationPayLoadType::CURVE;
+
+								if (EInterchangePropertyTracks PropertyTrack = InterchangeFbxSettings->GetPropertyTrack(CurveNodeName); PropertyTrack != EInterchangePropertyTracks::None)
 								{
-									AddAnimationTrackNode(PropertyTrack, CurveNodeName, *PayloadKey, GetPropertyPayloadType(PropertyType));
+									AddAnimationTrackNode(PropertyTrack, CurveNodeName, *PayloadKey, PayloadType);
 								}
+
+								UnrealNode->SetAnimationCurveTypeForCurveName(CurveNodeName, PayloadType);
 							}
 						}
 
diff --git a/Engine/Plugins/Interchange/Runtime/Source/Pipelines/Private/InterchangeGenericAnimationPipeline.cpp b/Engine/Plugins/Interchange/Runtime/Source/Pipelines/Private/InterchangeGenericAnimationPipeline.cpp
--- a/Engine/Plugins/Interchange/Runtime/Source/Pipelines/Private/InterchangeGenericAnimationPipeline.cpp
+++ b/Engine/Plugins/Interchange/Runtime/Source/Pipelines/Private/InterchangeGenericAnimationPipeline.cpp
@@ -821,15 +821,38 @@ void UInterchangeGenericAnimationPipeline::CreateAnimSequenceFactoryNode(UInterc
 								AttributeInfo.Type == EAttributeTypes::Double;
 
 							const bool bForceImportBoneCustomAttribute = CustomAttributeNamesToImport.Contains(AttributeInfo.Name);
-							//Material attribute curve
-							if (!bImportAllAttributesOnBone && bDecimalType && !bForceImportBoneCustomAttribute)
+							
+							bool bAttributeProcessed = false;
+							EInterchangeAnimationPayLoadType AnimationPayloadType = EInterchangeAnimationPayLoadType::NONE;
+							if (!bForceImportBoneCustomAttribute && !bImportAllAttributesOnBone
+								&& SceneNode->GetAnimationCurveTypeForCurveName(AttributeInfo.Name, AnimationPayloadType))
 							{
-								AnimSequenceFactoryNode->SetAnimatedAttributeCurveName(AttributeInfo.Name);
+								if (AnimationPayloadType == EInterchangeAnimationPayLoadType::CURVE)
+								{
+									AnimSequenceFactoryNode->SetAnimatedAttributeCurveName(AttributeInfo.Name);
+									bAttributeProcessed = true;
+								}
+								else if (AnimationPayloadType == EInterchangeAnimationPayLoadType::STEPCURVE)
+								{
+									AnimSequenceFactoryNode->SetAnimatedAttributeStepCurveName(AttributeInfo.Name);
+									bAttributeProcessed = true;
+								}
 							}
-							else if (bForceImportBoneCustomAttribute || bImportAllAttributesOnBone)
+
+							if (!bAttributeProcessed)
 							{
-								AnimSequenceFactoryNode->SetAnimatedAttributeStepCurveName(AttributeInfo.Name);
+								//Material attribute curve
+								if (!bImportAllAttributesOnBone && bDecimalType && !bForceImportBoneCustomAttribute)
+								{
+									AnimSequenceFactoryNode->SetAnimatedAttributeCurveName(AttributeInfo.Name);
+								}
+								else if (bForceImportBoneCustomAttribute || bImportAllAttributesOnBone)
+								{
+									AnimSequenceFactoryNode->SetAnimatedAttributeStepCurveName(AttributeInfo.Name);
+								}
 							}
+
+							
 						}
 					}
 				}
