[Interchange-FBX] Optimizing BindPose transform retrieval.

#rnx

diff --git a/Engine/Plugins/Interchange/Runtime/Source/Parsers/Fbx/Private/FbxMesh.cpp b/Engine/Plugins/Interchange/Runtime/Source/Parsers/Fbx/Private/FbxMesh.cpp
--- a/Engine/Plugins/Interchange/Runtime/Source/Parsers/Fbx/Private/FbxMesh.cpp
+++ b/Engine/Plugins/Interchange/Runtime/Source/Parsers/Fbx/Private/FbxMesh.cpp
@@ -1992,73 +1992,26 @@
 	} // for GeometryCount
 }
 
-bool FFbxMesh::GetGlobalJointBindPoseTransform(FFbxParser* Parser, FbxScene* SDKScene, FbxNode* Joint, FbxAMatrix& GlobalBindPoseJointMatrix, TMap<FString, FMatrix>& MeshIdToGlobalBindPoseJointMap, TMap<FString, FMatrix>& MeshIdToGlobalBindPoseReferenceMap, bool& bBadBindPoseMessageDisplay)
+bool FFbxMesh::GetGlobalJointBindPoseTransform(FFbxParser* Parser, FbxScene* SDKScene, FbxNode* Joint, FFbxJointMeshBindPoseGenerator& FbxJointMeshBindPoseGenerator,
+	FbxAMatrix& GlobalBindPoseJointMatrix, TMap<FString, FMatrix>& MeshIdToGlobalBindPoseJointMap, TMap<FString, FMatrix>& MeshIdToGlobalBindPoseReferenceMap, 
+	bool& bBadBindPoseMessageDisplay)
 {
 	FbxManager* SDKManager = SDKScene->GetFbxManager();
-	//First look for Cluster and then look in the bind pose if no cluster was found
-
-	//Search all skeletalmesh(FbxGeometry with valid deformer) using this joint and see if there is a valid FbxCluster
-	//containing a TransformLinkMatrix for this joint
-	bool bGlobalBindPoseJointMatrixSet = false;
-	const int32 GeometryCount = SDKScene->GetGeometryCount();
-	for (int32 GeometryIndex = 0; GeometryIndex < GeometryCount; ++GeometryIndex)
-	{
-		FbxGeometry* Geometry = SDKScene->GetGeometry(GeometryIndex);
-		if (!ensure(Geometry))
-		{
-			continue;
-		}
-		FString MeshUniqueID = Parser->GetFbxHelper()->GetMeshUniqueID(Geometry);
-
-		const int32 GeometryDeformerCount = Geometry->GetDeformerCount(FbxDeformer::eSkin);
-		for (int32 GeometryDeformerIndex = 0; GeometryDeformerIndex < GeometryDeformerCount; ++GeometryDeformerIndex)
-		{
-			const FbxSkin* Skin = (FbxSkin*)Geometry->GetDeformer(GeometryDeformerIndex, FbxDeformer::eSkin);
-			if (!ensure(Skin))
-			{
-				continue;
-			}
-			const int32 ClusterCount = Skin->GetClusterCount();
-			for (int32 ClusterIndex = 0; ClusterIndex < ClusterCount; ClusterIndex++)
-			{
-				const FbxCluster* Cluster = Skin->GetCluster(ClusterIndex);
-				// When Maya plug-in exports rigid binding, it will generate "CompensationCluster" for each ancestor links.
-				// FBX writes these "CompensationCluster" out. The CompensationCluster also has weight 1 for vertices.
-				// Unreal importer should skip these clusters.
-				if (!Cluster || (FCStringAnsi::Strcmp(Cluster->GetUserDataID(), "Maya_ClusterHint") == 0 && FCStringAnsi::Strcmp(Cluster->GetUserData(), "CompensationCluster") == 0))
-				{
-					continue;
-				}
-				if (Joint == Cluster->GetLink())
-				{
-					uint64 ClusterId = Cluster->GetUniqueID();
-					uint64 MeshId = Geometry->GetUniqueID();
-
-					FbxAMatrix GlobalBindPoseJointMatrixPerGeo;
-					Cluster->GetTransformLinkMatrix(GlobalBindPoseJointMatrixPerGeo);
-					MeshIdToGlobalBindPoseJointMap.Add(MeshUniqueID, FFbxConvert::ConvertMatrix<FMatrix>(GlobalBindPoseJointMatrixPerGeo));
-					if (!bGlobalBindPoseJointMatrixSet)
-					{
-						//We use the first geo bind pose
-						GlobalBindPoseJointMatrix = GlobalBindPoseJointMatrixPerGeo;
-						bGlobalBindPoseJointMatrixSet = true;
-					}
-
-					FbxAMatrix GlobalBindPoseReferenceMatrix;
-					Cluster->GetTransformMatrix(GlobalBindPoseReferenceMatrix);
-					MeshIdToGlobalBindPoseReferenceMap.Add(MeshUniqueID, FFbxConvert::ConvertMatrix<FMatrix>(GlobalBindPoseReferenceMatrix));
-
-					//Presumes only 1 cluster per geometry that matches criteria:
-					break;
-				}
-			}
-		}
-	}
+	
+	GlobalBindPoseJointMatrix = FbxJointMeshBindPoseGenerator.GetBindPoseFbxAMatrix(Joint, GlobalBindPoseJointMatrix);
+	MeshIdToGlobalBindPoseReferenceMap = FbxJointMeshBindPoseGenerator.GetMeshIdToGlobalBindPoseReferenceMap(Joint);
+	MeshIdToGlobalBindPoseJointMap = FbxJointMeshBindPoseGenerator.GetMeshIdToGlobalBindPoseJointMap(Joint);
+	
 	if (MeshIdToGlobalBindPoseReferenceMap.Num() > 0)
 	{
 		return true;
 	}
 
+	if (!(Joint->GetMesh() && Joint->GetMesh()->GetDeformerCount(FbxDeformer::eSkin) > 0))
+	{
+		return false;
+	}
+
 	auto AcquireBindPoseMatrix = [](FbxPose* CurrentPose, FbxAMatrix& GlobalBindPoseJointMatrix, FbxNode* Joint)
 	{
 		int32 PoseLinkIndex = CurrentPose->Find(Joint);
@@ -2240,6 +2193,110 @@
 	return MeshNode;
 }
 
+FFbxJointMeshBindPoseGenerator::FFbxJointMeshBindPoseGenerator(FbxScene* SDKScene, FFbxParser& Parser)
+{
+	//First look for Cluster and then look in the bind pose if no cluster was found
+
+	//Search all skeletalmesh(FbxGeometry with valid deformer) using this joint and see if there is a valid FbxCluster
+	//containing a TransformLinkMatrix for this joint
+	const int32 GeometryCount = SDKScene->GetGeometryCount();
+	for (int32 GeometryIndex = 0; GeometryIndex < GeometryCount; ++GeometryIndex)
+	{
+		FbxGeometry* Geometry = SDKScene->GetGeometry(GeometryIndex);
+		if (!ensure(Geometry))
+		{
+			continue;
+		}
+		FString MeshUniqueID = Parser.GetFbxHelper()->GetMeshUniqueID(Geometry);
+
+		const int32 GeometryDeformerCount = Geometry->GetDeformerCount(FbxDeformer::eSkin);
+		for (int32 GeometryDeformerIndex = 0; GeometryDeformerIndex < GeometryDeformerCount; ++GeometryDeformerIndex)
+		{
+			const FbxSkin* Skin = (FbxSkin*)Geometry->GetDeformer(GeometryDeformerIndex, FbxDeformer::eSkin);
+			if (!ensure(Skin))
+			{
+				continue;
+			}
+			const int32 ClusterCount = Skin->GetClusterCount();
+			for (int32 ClusterIndex = 0; ClusterIndex < ClusterCount; ClusterIndex++)
+			{
+				const FbxCluster* Cluster = Skin->GetCluster(ClusterIndex);
+				// When Maya plug-in exports rigid binding, it will generate "CompensationCluster" for each ancestor links.
+				// FBX writes these "CompensationCluster" out. The CompensationCluster also has weight 1 for vertices.
+				// Unreal importer should skip these clusters.
+				if (!Cluster || (FCStringAnsi::Strcmp(Cluster->GetUserDataID(), "Maya_ClusterHint") == 0 && FCStringAnsi::Strcmp(Cluster->GetUserData(), "CompensationCluster") == 0))
+				{
+					continue;
+				}
+
+				const FbxNode* Joint = Cluster->GetLink();
+
+				uint64 ClusterId = Cluster->GetUniqueID();
+				uint64 MeshId = Geometry->GetUniqueID();
+
+				FbxAMatrix GlobalBindPoseJointMatrixPerGeo;
+				Cluster->GetTransformLinkMatrix(GlobalBindPoseJointMatrixPerGeo);
+
+				TMap<FString, FMatrix>& MeshIdToGlobalBindPoseJointMap = JointToMeshIdToGlobalBindPoseJointMap.FindOrAdd(Joint);
+				MeshIdToGlobalBindPoseJointMap.Add(MeshUniqueID, FFbxConvert::ConvertMatrix<FMatrix>(GlobalBindPoseJointMatrixPerGeo));
+
+				if (!JointToBindPoseMap.Contains(Joint))
+				{
+					JointToBindPoseMap.Add(Joint, GlobalBindPoseJointMatrixPerGeo);
+				}
+
+				FbxAMatrix GlobalBindPoseReferenceMatrix;
+				Cluster->GetTransformMatrix(GlobalBindPoseReferenceMatrix);
+
+				TMap<FString, FMatrix>& MeshIdToGlobalBindPoseReferenceMap = JointToMeshIdToGlobalBindPoseReferenceMap.FindOrAdd(Joint);
+				MeshIdToGlobalBindPoseReferenceMap.Add(MeshUniqueID, FFbxConvert::ConvertMatrix<FMatrix>(GlobalBindPoseReferenceMatrix));
+			}
+		}
+	}
+}
+
+TMap<FString, FMatrix>& FFbxJointMeshBindPoseGenerator::GetMeshIdToGlobalBindPoseReferenceMap(const FbxNode* Joint)
+{
+	TMap<FString, FMatrix>* MeshIdToGlobalBindPoseReferenceMap = JointToMeshIdToGlobalBindPoseReferenceMap.Find(Joint);
+	if (MeshIdToGlobalBindPoseReferenceMap)
+	{
+		return *MeshIdToGlobalBindPoseReferenceMap;
+	}
+	else
+	{
+		static TMap<FString, FMatrix> EmptyMap;
+		return EmptyMap;
+	}
+}
+
+TMap<FString, FMatrix>& FFbxJointMeshBindPoseGenerator::GetMeshIdToGlobalBindPoseJointMap(const FbxNode* Joint)
+{
+	TMap<FString, FMatrix>* MeshIdToGlobalBindPoseJointMap = JointToMeshIdToGlobalBindPoseJointMap.Find(Joint);
+	if (MeshIdToGlobalBindPoseJointMap)
+	{
+		return *MeshIdToGlobalBindPoseJointMap;
+	}
+	else
+	{
+		static TMap<FString, FMatrix> EmptyMap;
+		return EmptyMap;
+	}
+}
+
+FbxAMatrix FFbxJointMeshBindPoseGenerator::GetBindPoseFbxAMatrix(const FbxNode* Joint, FbxAMatrix& DefaultValue)
+{
+	FbxAMatrix* Matrix = JointToBindPoseMap.Find(Joint);
+	if (Matrix)
+	{
+		return *Matrix;
+	}
+	else
+	{
+		return DefaultValue;
+	}
+}
+
 } // namespace UE::Interchange::Private
 
+
 #undef LOCTEXT_NAMESPACE
diff --git a/Engine/Plugins/Interchange/Runtime/Source/Parsers/Fbx/Private/FbxMesh.h b/Engine/Plugins/Interchange/Runtime/Source/Parsers/Fbx/Private/FbxMesh.h
--- a/Engine/Plugins/Interchange/Runtime/Source/Parsers/Fbx/Private/FbxMesh.h
+++ b/Engine/Plugins/Interchange/Runtime/Source/Parsers/Fbx/Private/FbxMesh.h
@@ -22,6 +22,20 @@ namespace UE
 
 		namespace Private
 		{
+			//Helper to generate and cache per joint MeshIdToGlobalBindPose transforms.
+			struct FFbxJointMeshBindPoseGenerator
+			{
+				TMap<const FbxNode*, TMap<FString, FMatrix>> JointToMeshIdToGlobalBindPoseReferenceMap;
+				TMap<const FbxNode*, TMap<FString, FMatrix>> JointToMeshIdToGlobalBindPoseJointMap;
+				TMap<const FbxNode*, FbxAMatrix> JointToBindPoseMap;
+
+				FFbxJointMeshBindPoseGenerator(FbxScene* SDKScene, FFbxParser& Parser);
+
+				TMap<FString, FMatrix>& GetMeshIdToGlobalBindPoseReferenceMap(const FbxNode* Joint);
+				TMap<FString, FMatrix>& GetMeshIdToGlobalBindPoseJointMap(const FbxNode* Joint);
+				FbxAMatrix GetBindPoseFbxAMatrix(const FbxNode* Joint, FbxAMatrix& DefaultValue);
+			};
+
 			struct FMorphTargetAnimationBuildingData
 			{
 				double StartTime;
@@ -196,7 +210,9 @@ namespace UE
 				{}
 
 				void AddAllMeshes(FbxScene* SDKScene, FbxGeometryConverter* SDKGeometryConverter, UInterchangeBaseNodeContainer& NodeContainer, TMap<FString, TSharedPtr<FPayloadContextBase>>& PayloadContexts);
-				static bool GetGlobalJointBindPoseTransform(FFbxParser* Parser, FbxScene* SDKScene, FbxNode* Joint, FbxAMatrix& GlobalBindPoseJointMatrix, TMap<FString, FMatrix>& MeshIdToGlobalBindPoseJointMap, TMap<FString, FMatrix>& MeshIdToGlobalBindPoseReferenceMap, bool& bBadBindPoseMessageDisplay);
+				static bool GetGlobalJointBindPoseTransform(FFbxParser* Parser, FbxScene* SDKScene, FbxNode* Joint, FFbxJointMeshBindPoseGenerator& FbxJointMeshBindPoseGenerator,
+					FbxAMatrix& GlobalBindPoseJointMatrix, TMap<FString, FMatrix>& MeshIdToGlobalBindPoseJointMap, TMap<FString, FMatrix>& MeshIdToGlobalBindPoseReferenceMap, 
+					bool& bBadBindPoseMessageDisplay);
 
 				TArray<FMorphTargetAnimationBuildingData>& GetMorphTargetAnimationsBuildingData() { return MorphTargetAnimationsBuildingData; }
 			protected:
diff --git a/Engine/Plugins/Interchange/Runtime/Source/Parsers/Fbx/Private/FbxScene.cpp b/Engine/Plugins/Interchange/Runtime/Source/Parsers/Fbx/Private/FbxScene.cpp
--- a/Engine/Plugins/Interchange/Runtime/Source/Parsers/Fbx/Private/FbxScene.cpp
+++ b/Engine/Plugins/Interchange/Runtime/Source/Parsers/Fbx/Private/FbxScene.cpp
@@ -141,7 +141,8 @@ namespace UE
 				, UInterchangeBaseNodeContainer& NodeContainer
 				, TMap<FString, TSharedPtr<FPayloadContextBase, ESPMode::ThreadSafe>>& PayloadContexts
 				, TArray<FbxNode*>& ForceJointNodes
-				, bool& bBadBindPoseMessageDisplay)
+				, bool& bBadBindPoseMessageDisplay
+				, FFbxJointMeshBindPoseGenerator& FbxJointMeshBindPoseGenerator)
 			{
 				constexpr bool bResetCache = false;
 				FString NodeName = Parser.GetFbxHelper()->GetFbxObjectName(Node);
@@ -192,7 +193,7 @@ namespace UE
 					}
 				}
 
-				auto ApplySkeletonAttribute = [this, &SDKScene, &UnrealNode, &Node, &NodeContainer, &bResetCache, &GetConvertedTransform, &bBadBindPoseMessageDisplay]()
+				auto ApplySkeletonAttribute = [this, &SDKScene, &UnrealNode, &Node, &NodeContainer, &bResetCache, &GetConvertedTransform, &bBadBindPoseMessageDisplay, &FbxJointMeshBindPoseGenerator]()
 				{
 					if (FRootJointInfo* RootJointInfo = CommonJointRootNodes.Find(Node))
 					{
@@ -201,13 +202,14 @@ namespace UE
 							UnrealNode->SetCustomHasBindPose(false);
 						}
 					}
+
 					//Add the joint specialized type
 					UnrealNode->AddSpecializedType(FSceneNodeStaticData::GetJointSpecializeTypeString());
 					//Get the bind pose transform for this joint
 					FbxAMatrix GlobalBindPoseJointMatrix = SDKScene->GetAnimationEvaluator()->GetNodeGlobalTransform(Node, 0);
 					TMap<FString, FMatrix> MeshIdToGlobalBindPoseReferenceMap;
 					TMap<FString, FMatrix> MeshIdToGlobalBindPoseJointMap;
-					FFbxMesh::GetGlobalJointBindPoseTransform(&Parser, SDKScene, Node, GlobalBindPoseJointMatrix, MeshIdToGlobalBindPoseJointMap, MeshIdToGlobalBindPoseReferenceMap, bBadBindPoseMessageDisplay);
+					FFbxMesh::GetGlobalJointBindPoseTransform(&Parser, SDKScene, Node, FbxJointMeshBindPoseGenerator, GlobalBindPoseJointMatrix, MeshIdToGlobalBindPoseJointMap, MeshIdToGlobalBindPoseReferenceMap, bBadBindPoseMessageDisplay);
 
 					FTransform GlobalBindPoseJointTransform = GetConvertedTransform(GlobalBindPoseJointMatrix);
 					UnrealNode->SetGlobalBindPoseReferenceForMeshUIDs(MeshIdToGlobalBindPoseReferenceMap);
@@ -226,7 +228,7 @@ namespace UE
 						FbxAMatrix GlobalFbxParentMatrix = SDKScene->GetAnimationEvaluator()->GetNodeGlobalTransform(ParentNode, 0);
 						TMap<FString, FMatrix> ParentMeshIdToGlobalBindPoseReferenceMap;
 						TMap<FString, FMatrix> ParentMeshIdToGlobalBindPoseJointMap;
-						FFbxMesh::GetGlobalJointBindPoseTransform(&Parser, SDKScene, ParentNode, GlobalFbxParentMatrix, ParentMeshIdToGlobalBindPoseJointMap, ParentMeshIdToGlobalBindPoseReferenceMap, bBadBindPoseMessageDisplay);
+						FFbxMesh::GetGlobalJointBindPoseTransform(&Parser, SDKScene, ParentNode, FbxJointMeshBindPoseGenerator, GlobalFbxParentMatrix, ParentMeshIdToGlobalBindPoseJointMap, ParentMeshIdToGlobalBindPoseReferenceMap, bBadBindPoseMessageDisplay);
 						
 						FbxAMatrix LocalFbxMatrix = GlobalFbxParentMatrix.Inverse() * GlobalBindPoseJointMatrix;
 						FTransform LocalBindPoseJointTransform = GetConvertedTransform(LocalFbxMatrix);
@@ -475,7 +477,7 @@ namespace UE
 				for (int32 ChildIndex = 0; ChildIndex < ChildCount; ++ChildIndex)
 				{
 					FbxNode* ChildNode = Node->GetChild(ChildIndex);
-					AddHierarchyRecursively(UnrealNode, ChildNode, SDKScene, NodeContainer, PayloadContexts, ForceJointNodes, bBadBindPoseMessageDisplay);
+					AddHierarchyRecursively(UnrealNode, ChildNode, SDKScene, NodeContainer, PayloadContexts, ForceJointNodes, bBadBindPoseMessageDisplay, FbxJointMeshBindPoseGenerator);
 				}
 			}
 
@@ -606,8 +608,10 @@ namespace UE
 					RootJointInfo.Value.bValidBindPose = IsValidBindPose(SDKScene, RootJointInfo.Key);
 				}
 
+				FFbxJointMeshBindPoseGenerator FbxJointMeshBindPoseGenerator(SDKScene, Parser);
+
 				bool bBadBindPoseMessageDisplay = false;
-				AddHierarchyRecursively(nullptr, RootNode, SDKScene, NodeContainer, PayloadContexts, ForceJointNodes, bBadBindPoseMessageDisplay);
+				AddHierarchyRecursively(nullptr, RootNode, SDKScene, NodeContainer, PayloadContexts, ForceJointNodes, bBadBindPoseMessageDisplay, FbxJointMeshBindPoseGenerator);
 
 				int32 NodeCount = SDKScene->GetNodeCount();
 				for (int32 NodeIndex = 0; NodeIndex < NodeCount; ++NodeIndex)
@@ -618,7 +622,7 @@ namespace UE
 						{
 							if (Node->GetParent() == nullptr)
 							{
-								AddHierarchyRecursively(nullptr, Node, SDKScene, NodeContainer, PayloadContexts, ForceJointNodes, bBadBindPoseMessageDisplay);
+								AddHierarchyRecursively(nullptr, Node, SDKScene, NodeContainer, PayloadContexts, ForceJointNodes, bBadBindPoseMessageDisplay, FbxJointMeshBindPoseGenerator);
 							}
 						}
 					}
diff --git a/Engine/Plugins/Interchange/Runtime/Source/Parsers/Fbx/Private/FbxScene.h b/Engine/Plugins/Interchange/Runtime/Source/Parsers/Fbx/Private/FbxScene.h
--- a/Engine/Plugins/Interchange/Runtime/Source/Parsers/Fbx/Private/FbxScene.h
+++ b/Engine/Plugins/Interchange/Runtime/Source/Parsers/Fbx/Private/FbxScene.h
@@ -9,6 +9,7 @@
 namespace UE::Interchange::Private
 {
 	class FPayloadContextBase;
+	struct FFbxJointMeshBindPoseGenerator;
 }
 class UInterchangeBaseNodeContainer;
 class UInterchangeSceneNode;
@@ -50,7 +51,8 @@ namespace UE
 					, UInterchangeBaseNodeContainer& NodeContainer
 					, TMap<FString, TSharedPtr<FPayloadContextBase, ESPMode::ThreadSafe>>& PayloadContexts
 					, TArray<FbxNode*>& ForceJointNodes
-					, bool& bBadBindPoseMessageDisplay);
+					, bool& bBadBindPoseMessageDisplay
+					, FFbxJointMeshBindPoseGenerator& FbxJointMeshBindPoseGenerator);
 
 				void AddAnimationRecursively(FbxNode* Node
 					, FbxScene* SDKScene
