USD: Fix broken build on Linux since cl 39367237.

diff --git a/Engine/Plugins/Interchange/Runtime/Source/Import/Private/Usd/InterchangeUsdTranslator.cpp b/Engine/Plugins/Interchange/Runtime/Source/Import/Private/Usd/InterchangeUsdTranslator.cpp
--- a/Engine/Plugins/Interchange/Runtime/Source/Import/Private/Usd/InterchangeUsdTranslator.cpp
+++ b/Engine/Plugins/Interchange/Runtime/Source/Import/Private/Usd/InterchangeUsdTranslator.cpp
@@ -2433,8 +2433,8 @@ namespace UE::InterchangeUsdTranslator::Private
 				{TEXT("uchar[]"), 		TranslateAttribute<TArray<uint8>>},
 				{TEXT("int"), 			TranslateAttribute<int32>},
 				{TEXT("uint"), 			TranslateAttribute<uint32>},
-				{TEXT("int64"), 		TranslateAttribute<int64>},
-				{TEXT("uint64"), 		TranslateAttribute<uint64>},
+				// {TEXT("int64"), 		TranslateAttribute<int64_t>},  // Temporarily disabled as they don't work on Linux
+				// {TEXT("uint64"), 		TranslateAttribute<uint64_t>},
 
 				{TEXT("half"),	 		TranslateAttribute<FFloat16>},
 				{TEXT("float"), 		TranslateAttribute<float>},
