Interchange: fix the skeletalmesh morph targets import not registering to skeleton curve metadata, when the morph target curves are not animated.

#rnx

diff --git a/Engine/Plugins/Interchange/Editor/Source/InterchangeEditor/Private/InterchangeFbxAssetImportDataConverter.cpp b/Engine/Plugins/Interchange/Editor/Source/InterchangeEditor/Private/InterchangeFbxAssetImportDataConverter.cpp
--- a/Engine/Plugins/Interchange/Editor/Source/InterchangeEditor/Private/InterchangeFbxAssetImportDataConverter.cpp
+++ b/Engine/Plugins/Interchange/Editor/Source/InterchangeEditor/Private/InterchangeFbxAssetImportDataConverter.cpp
@@ -222,7 +222,7 @@ namespace UE::Interchange::Private
 			break;
 		}
 		}
-		DestinationAnimSequenceImportData->bAddCurveMetadataToSkeleton = GenericAssetPipeline->AnimationPipeline->bAddCurveMetadataToSkeleton;
+		DestinationAnimSequenceImportData->bAddCurveMetadataToSkeleton = GenericAssetPipeline->CommonSkeletalMeshesAndAnimationsProperties->bAddCurveMetadataToSkeleton;
 		DestinationAnimSequenceImportData->bDeleteExistingCustomAttributeCurves = GenericAssetPipeline->AnimationPipeline->bDeleteExistingCustomAttributeCurves;
 		DestinationAnimSequenceImportData->bDeleteExistingMorphTargetCurves = GenericAssetPipeline->AnimationPipeline->bDeleteExistingMorphTargetCurves;
 		DestinationAnimSequenceImportData->bDeleteExistingNonCurveCustomAttributes = GenericAssetPipeline->AnimationPipeline->bDeleteExistingNonCurveCustomAttributes;
@@ -430,7 +430,7 @@ namespace UE::Interchange::Private
 			break;
 		}
 		}
-		GenericAssetPipeline->AnimationPipeline->bAddCurveMetadataToSkeleton = AnimSequenceImportData->bAddCurveMetadataToSkeleton;
+		GenericAssetPipeline->CommonSkeletalMeshesAndAnimationsProperties->bAddCurveMetadataToSkeleton = AnimSequenceImportData->bAddCurveMetadataToSkeleton;
 		GenericAssetPipeline->AnimationPipeline->bDeleteExistingCustomAttributeCurves = AnimSequenceImportData->bDeleteExistingCustomAttributeCurves;
 		GenericAssetPipeline->AnimationPipeline->bDeleteExistingMorphTargetCurves = AnimSequenceImportData->bDeleteExistingMorphTargetCurves;
 		GenericAssetPipeline->AnimationPipeline->bDeleteExistingNonCurveCustomAttributes = AnimSequenceImportData->bDeleteExistingNonCurveCustomAttributes;
diff --git a/Engine/Plugins/Interchange/Runtime/Source/FactoryNodes/Private/InterchangeSkeletalMeshFactoryNode.cpp b/Engine/Plugins/Interchange/Runtime/Source/FactoryNodes/Private/InterchangeSkeletalMeshFactoryNode.cpp
--- a/Engine/Plugins/Interchange/Runtime/Source/FactoryNodes/Private/InterchangeSkeletalMeshFactoryNode.cpp
+++ b/Engine/Plugins/Interchange/Runtime/Source/FactoryNodes/Private/InterchangeSkeletalMeshFactoryNode.cpp
@@ -96,6 +96,16 @@ bool UInterchangeSkeletalMeshFactoryNode::SetCustomImportMorphTarget(const bool&
 	IMPLEMENT_NODE_ATTRIBUTE_SETTER_NODELEGATE(ImportMorphTarget, bool)
 }
 
+bool UInterchangeSkeletalMeshFactoryNode::GetCustomAddCurveMetadataToSkeleton(bool& AttributeValue) const
+{
+	IMPLEMENT_NODE_ATTRIBUTE_GETTER(AddCurveMetadataToSkeleton, bool);
+}
+
+bool UInterchangeSkeletalMeshFactoryNode::SetCustomAddCurveMetadataToSkeleton(const bool& AttributeValue)
+{
+	IMPLEMENT_NODE_ATTRIBUTE_SETTER_NODELEGATE(AddCurveMetadataToSkeleton, bool);
+}
+
 bool UInterchangeSkeletalMeshFactoryNode::GetCustomImportVertexAttributes(bool& AttributeValue) const
 {
 	IMPLEMENT_NODE_ATTRIBUTE_GETTER(ImportVertexAttributes, bool)
diff --git a/Engine/Plugins/Interchange/Runtime/Source/FactoryNodes/Public/InterchangeSkeletalMeshFactoryNode.h b/Engine/Plugins/Interchange/Runtime/Source/FactoryNodes/Public/InterchangeSkeletalMeshFactoryNode.h
--- a/Engine/Plugins/Interchange/Runtime/Source/FactoryNodes/Public/InterchangeSkeletalMeshFactoryNode.h
+++ b/Engine/Plugins/Interchange/Runtime/Source/FactoryNodes/Public/InterchangeSkeletalMeshFactoryNode.h
@@ -63,6 +63,22 @@ public:
 	UFUNCTION(BlueprintCallable, Category = "Interchange | Node | SkeletalMesh")
 	bool SetCustomImportMorphTarget(const bool& AttributeValue);
 
+	/**
+	 * Get the custom attribute AddCurveMetadataToSkeleton. Return false if the attribute is not set.
+	 * 
+	 * Note - If this setting is disabled, curve metadata will be added to skeletal meshes for morph targets, but no metadata entry will be created for general curves.
+	 */
+	UFUNCTION(BlueprintCallable, Category = "Interchange | Node | AnimSequence")
+	bool GetCustomAddCurveMetadataToSkeleton(bool& AttributeValue) const;
+
+	/**
+	 * Set the custom attribute AddCurveMetadataToSkeleton. Return false if the attribute could not be set.
+	 * 
+	 * Note - If this setting is disabled, curve metadata will be added to skeletal meshes for morph targets, but no metadata entry will be created for general curves.
+	 */
+	UFUNCTION(BlueprintCallable, Category = "Interchange | Node | AnimSequence")
+	bool SetCustomAddCurveMetadataToSkeleton(const bool& AttributeValue);
+
 	/** Query whether the skeletal mesh factory should import vertex attributes. Return false if the attribute was not set.*/
 	UFUNCTION(BlueprintCallable, Category = "Interchange | Node | SkeletalMesh")
 	bool GetCustomImportVertexAttributes(bool& AttributeValue) const;
@@ -180,6 +196,7 @@ private:
 
 	const UE::Interchange::FAttributeKey ClassNameAttributeKey = UE::Interchange::FBaseNodeStaticData::ClassTypeAttributeKey();
 	const UE::Interchange::FAttributeKey Macro_CustomImportMorphTargetKey = UE::Interchange::FAttributeKey(TEXT("ImportMorphTarget"));
+	const UE::Interchange::FAttributeKey Macro_CustomAddCurveMetadataToSkeletonKey = UE::Interchange::FAttributeKey(TEXT("AddCurveMetadataToSkeleton"));
 	const UE::Interchange::FAttributeKey Macro_CustomImportVertexAttributesKey = UE::Interchange::FAttributeKey(TEXT("ImportVertexAttributes"));
 	const UE::Interchange::FAttributeKey Macro_CustomSkeletonSoftObjectPathKey = UE::Interchange::FAttributeKey(TEXT("SkeletonSoftObjectPath"));
 	const UE::Interchange::FAttributeKey Macro_CustomCreatePhysicsAssetKey = UE::Interchange::FAttributeKey(TEXT("CreatePhysicsAsset"));
diff --git a/Engine/Plugins/Interchange/Runtime/Source/Import/Private/Mesh/InterchangeSkeletalMeshFactory.cpp b/Engine/Plugins/Interchange/Runtime/Source/Import/Private/Mesh/InterchangeSkeletalMeshFactory.cpp
--- a/Engine/Plugins/Interchange/Runtime/Source/Import/Private/Mesh/InterchangeSkeletalMeshFactory.cpp
+++ b/Engine/Plugins/Interchange/Runtime/Source/Import/Private/Mesh/InterchangeSkeletalMeshFactory.cpp
@@ -163,7 +163,10 @@ namespace UE
 				}
 			}
 
-			void CopyMorphTargetsMeshDescriptionToSkeletalMeshImportData(const TMap<FString, TOptional<UE::Interchange::FMeshPayloadData>>& LodMorphTargetMeshDescriptions, FSkeletalMeshImportData& DestinationSkeletalMeshImportData, const bool bMergeMorphTargetWithSameName)
+			void CopyMorphTargetsMeshDescriptionToSkeletalMeshImportData(TArray<FString>& SkeletonMorphCurveMetadataNames
+																		, const TMap<FString, TOptional<UE::Interchange::FMeshPayloadData>>& LodMorphTargetMeshDescriptions
+																		, FSkeletalMeshImportData& DestinationSkeletalMeshImportData
+																		, const bool bMergeMorphTargetWithSameName)
 			{
 				TRACE_CPUPROFILER_EVENT_SCOPE(CopyMorphTargetsMeshDescriptionToSkeletalMeshImportData)
 				const int32 OriginalMorphTargetCount = LodMorphTargetMeshDescriptions.Num();
@@ -313,6 +316,8 @@ namespace UE
 					}
 				}
 				, EParallelForFlags::BackgroundPriority);
+
+				SkeletonMorphCurveMetadataNames = DestinationSkeletalMeshImportData.MorphTargetNames;
 				return;
 			}
 
@@ -495,16 +500,18 @@ namespace UE
 			}
 
 			void RetrieveAllSkeletalMeshPayloadsAndFillImportData(const UInterchangeSkeletalMeshFactoryNode* SkeletalMeshFactoryNode
+																  , UInterchangeSkeletalMeshFactory::FImportAssetObjectLODData& ImportAssetObjectLODData
 																  , FSkeletalMeshImportData& DestinationImportData
-																  , TArray<FMeshNodeContext>& MeshReferences
 																  , UInterchangeSkeletalMeshFactory::FLodPayloads& LodPayloads
-																  , TArray<SkeletalMeshImportData::FBone>& RefBonesBinary
 																  , const UInterchangeSkeletalMeshFactory::FImportAssetObjectParams& Arguments
-																  , const bool bSkinControlPointToTimeZero
 																  , const UInterchangeBaseNodeContainer* NodeContainer
 																  , const FString& RootJointNodeId)
 			{
 				TRACE_CPUPROFILER_EVENT_SCOPE(RetrieveAllSkeletalMeshPayloadsAndFillImportData)
+				TArray<FMeshNodeContext>& MeshReferences = ImportAssetObjectLODData.MeshNodeContexts;
+				TArray<SkeletalMeshImportData::FBone>& RefBonesBinary = ImportAssetObjectLODData.RefBonesBinary;
+				const bool bSkinControlPointToTimeZero = ImportAssetObjectLODData.bUseTimeZeroAsBindPose && ImportAssetObjectLODData.bDiffPose;
+
 				FMeshDescription LodMeshDescription;
 				FSkeletalMeshAttributes SkeletalMeshAttributes(LodMeshDescription);
 				SkeletalMeshAttributes.Register();
@@ -719,7 +726,7 @@ namespace UE
 				bool bMergeMorphTargetWithSameName = false;
 				SkeletalMeshFactoryNode->GetCustomMergeMorphTargetShapeWithSameName(bMergeMorphTargetWithSameName);
 				//Copy all the lod morph targets data to the DestinationImportData.
-				CopyMorphTargetsMeshDescriptionToSkeletalMeshImportData(MorphTargetMeshDescriptionsPerMorphTargetName, DestinationImportData, bMergeMorphTargetWithSameName);
+				CopyMorphTargetsMeshDescriptionToSkeletalMeshImportData(ImportAssetObjectLODData.SkeletonMorphCurveMetadataNames, MorphTargetMeshDescriptionsPerMorphTargetName, DestinationImportData, bMergeMorphTargetWithSameName);
 			}
 
 			void ProcessImportMeshInfluences(const int32 WedgeCount, TArray<SkeletalMeshImportData::FRawBoneInfluence>& Influences)
@@ -1781,15 +1788,12 @@ UInterchangeFactoryBase::FImportAssetResult UInterchangeSkeletalMeshFactory::Imp
 
 		//Add the lod mesh data to the skeletalmesh
 		FSkeletalMeshImportData SkeletalMeshImportData;
-		const bool bSkinControlPointToTimeZero = ImportAssetObjectLODData.bUseTimeZeroAsBindPose && ImportAssetObjectLODData.bDiffPose;
 		//Get all meshes and morph targets payload and fill the SkeletalMeshImportData structure
 		UE::Interchange::Private::RetrieveAllSkeletalMeshPayloadsAndFillImportData(SkeletalMeshFactoryNode
+																					, ImportAssetObjectLODData
 																					, SkeletalMeshImportData
-																					, ImportAssetObjectLODData.MeshNodeContexts
 																					, LodPayloads
-																					, ImportAssetObjectLODData.RefBonesBinary
 																					, Arguments
-																					, bSkinControlPointToTimeZero
 																					, Arguments.NodeContainer
 																					, RootJointNodeId);
 		//////////////////////////////////////////////////////////////////////////
@@ -2255,6 +2259,8 @@ void UInterchangeSkeletalMeshFactory::SetupObject_GameThread(const FSetupObjectP
 	check(IsInGameThread());
 	Super::SetupObject_GameThread(Arguments);
 
+	const UInterchangeSkeletalMeshFactoryNode* SkeletalMeshFactoryNode = Cast<const UInterchangeSkeletalMeshFactoryNode>(Arguments.NodeContainer->GetFactoryNode(Arguments.NodeUniqueID));
+
 	//TODO make sure this work at runtime
 #if WITH_EDITORONLY_DATA
 	if (ensure(Arguments.ImportedObject && Arguments.SourceData))
@@ -2271,7 +2277,7 @@ void UInterchangeSkeletalMeshFactory::SetupObject_GameThread(const FSetupObjectP
 																										  , Arguments.OriginalPipelines
 																										  , Arguments.Translator);
 
-		ImportDataPtr = UE::Interchange::FFactoryCommon::UpdateImportAssetData(UpdateImportAssetDataParameters, [&Arguments, SkeletalMesh](UInterchangeAssetImportData* AssetImportData)
+		ImportDataPtr = UE::Interchange::FFactoryCommon::UpdateImportAssetData(UpdateImportAssetDataParameters, [&Arguments, SkeletalMeshFactoryNode, SkeletalMesh](UInterchangeAssetImportData* AssetImportData)
 			{
 				auto GetSourceIndexFromContentType = [](const EInterchangeSkeletalMeshContentType& ImportContentType)->int32
 				{
@@ -2291,7 +2297,7 @@ void UInterchangeSkeletalMeshFactory::SetupObject_GameThread(const FSetupObjectP
 						: NSSkeletalMeshSourceFileLabels::GeoAndSkinningText().ToString();
 				};
 
-				if (const UInterchangeSkeletalMeshFactoryNode* SkeletalMeshFactoryNode = Cast<const UInterchangeSkeletalMeshFactoryNode>(Arguments.NodeContainer->GetFactoryNode(Arguments.NodeUniqueID)))
+				if (SkeletalMeshFactoryNode)
 				{
 					EInterchangeSkeletalMeshContentType ImportContentType = EInterchangeSkeletalMeshContentType::All;
 					SkeletalMeshFactoryNode->GetCustomImportContentType(ImportContentType);
@@ -2333,13 +2339,61 @@ void UInterchangeSkeletalMeshFactory::SetupObject_GameThread(const FSetupObjectP
 
 		
 		SkeletalMesh->SetAssetImportData(ImportDataPtr);
+
+		//Ensure we have curve meta data for every morph target
+		if (SkeletalMeshFactoryNode)
+		{
+			USkeleton* SkeletonReference = ImportAssetObjectData.SkeletonReference;
+			bool bAddMorphTargetToSkeletonCurveMetadata = false;
+			if (SkeletonReference)
+			{
+				bAddMorphTargetToSkeletonCurveMetadata = true;
+				SkeletalMeshFactoryNode->GetCustomAddCurveMetadataToSkeleton(bAddMorphTargetToSkeletonCurveMetadata);
+			}
+			constexpr bool bInTransactFalse = false;
+			for (int32 LodIndex = 0; LodIndex < ImportAssetObjectData.LodDatas.Num(); ++LodIndex)
+			{
+				// Ensure that we have curve metadata for morph target (either skeleton or mesh)
+				for (const FString& MorphTargetName : ImportAssetObjectData.LodDatas[LodIndex].SkeletonMorphCurveMetadataNames)
+				{
+					FName CurveName = *MorphTargetName;
+					if (bAddMorphTargetToSkeletonCurveMetadata)
+					{
+						//bAddMorphTargetToSkeletonCurveMetadata is true only if the skeleton is not null and the node setting is undefined or true
+						if (ensure(SkeletonReference))
+						{
+							SkeletonReference->AddCurveMetaData(CurveName, bInTransactFalse);
+							// Ensure we have a morph flag set
+							FCurveMetaData* CurveMetaData = SkeletonReference->GetCurveMetaData(CurveName);
+							check(CurveMetaData);
+							CurveMetaData->Type.bMorphtarget = true;
+						}
+					}
+					else
+					{
+						UAnimCurveMetaData* AnimCurveMetaData = SkeletalMesh->GetAssetUserData<UAnimCurveMetaData>();
+						if (AnimCurveMetaData == nullptr)
+						{
+							AnimCurveMetaData = NewObject<UAnimCurveMetaData>(SkeletalMesh, NAME_None, RF_Transactional);
+							SkeletalMesh->AddAssetUserData(AnimCurveMetaData);
+						}
+						AnimCurveMetaData->AddCurveMetaData(CurveName, bInTransactFalse);
+						// Ensure we have a morph flag set
+						FCurveMetaData* CurveMetaData = AnimCurveMetaData->GetCurveMetaData(CurveName);
+						check(CurveMetaData);
+						CurveMetaData->Type.bMorphtarget = true;
+					}
+				}
+			}
+		}
+
 #if WITH_EDITOR
 		//Re-apply the alternate skinning data
 		if (ImportAssetObjectData.ExistingSkinWeightProfileInfos.Num() > 0)
 		{
 			//Reimport alternate skin of the reimport lod only by looking at the factory node lod data count.
 			int32 LodCount = SkeletalMesh->GetLODNum();
-			if (const UInterchangeSkeletalMeshFactoryNode* SkeletalMeshFactoryNode = Cast<const UInterchangeSkeletalMeshFactoryNode>(Arguments.NodeContainer->GetFactoryNode(Arguments.NodeUniqueID)))
+			if (SkeletalMeshFactoryNode)
 			{
 				int32 ImportedLodCount = SkeletalMeshFactoryNode->GetLodDataCount();
 				if (ImportedLodCount >  0 && ImportedLodCount <= SkeletalMesh->GetLODNum())
diff --git a/Engine/Plugins/Interchange/Runtime/Source/Import/Public/Mesh/InterchangeSkeletalMeshFactory.h b/Engine/Plugins/Interchange/Runtime/Source/Import/Public/Mesh/InterchangeSkeletalMeshFactory.h
--- a/Engine/Plugins/Interchange/Runtime/Source/Import/Public/Mesh/InterchangeSkeletalMeshFactory.h
+++ b/Engine/Plugins/Interchange/Runtime/Source/Import/Public/Mesh/InterchangeSkeletalMeshFactory.h
@@ -55,6 +55,9 @@ public:
 		TArray<UE::Interchange::FMeshNodeContext> MeshNodeContexts;
 		bool bUseTimeZeroAsBindPose = false;
 		bool bDiffPose = false;
+
+		//Store Morphtarget name, we want to add skeleton curve meta data ingame thread (FinalizeObject_GameThread)
+		TArray<FString> SkeletonMorphCurveMetadataNames;
 	};
 
 	struct FImportAssetObjectData
diff --git a/Engine/Plugins/Interchange/Runtime/Source/Pipelines/Private/InterchangeGenericAnimationPipeline.cpp b/Engine/Plugins/Interchange/Runtime/Source/Pipelines/Private/InterchangeGenericAnimationPipeline.cpp
--- a/Engine/Plugins/Interchange/Runtime/Source/Pipelines/Private/InterchangeGenericAnimationPipeline.cpp
+++ b/Engine/Plugins/Interchange/Runtime/Source/Pipelines/Private/InterchangeGenericAnimationPipeline.cpp
@@ -14,6 +14,7 @@
 #include "InterchangeHelper.h"
 #include "InterchangeMeshNode.h"
 #include "InterchangePipelineLog.h"
+#include "InterchangePipelineObjectVersion.h"
 #include "InterchangeSceneNode.h"
 #include "InterchangeSkeletalMeshFactoryNode.h"
 #include "InterchangeSkeletalMeshLodDataNode.h"
@@ -140,6 +141,22 @@ namespace UE::Interchange::Private
 	}
 }
 
+void UInterchangeGenericAnimationPipeline::Serialize(FArchive& Ar)
+{
+	Ar.UsingCustomVersion(FFortniteMainInterchangePipelineObjectVersion::GUID);
+
+	Super::Serialize(Ar);
+
+	//If we load an old pipeline, we must transfer the value of the bAddCurveMetadataToSkeleton hold by 
+	if (Ar.IsLoading() && Ar.CustomVer(FFortniteMainInterchangePipelineObjectVersion::GUID) < FFortniteMainInterchangePipelineObjectVersion::InterchangeAddCurveMetadataToSkeletonPropertyMove)
+	{
+		if (ensure(CommonSkeletalMeshesAndAnimationsProperties.IsValid()))
+		{
+			CommonSkeletalMeshesAndAnimationsProperties->bAddCurveMetadataToSkeleton = bAddCurveMetadataToSkeleton_DEPRECATED;
+		}
+	}
+}
+
 FString UInterchangeGenericAnimationPipeline::GetPipelineCategory(UClass* AssetClass)
 {
 	return TEXT("Animations");
@@ -575,6 +592,7 @@ void UInterchangeGenericAnimationPipeline::CreateAnimSequenceFactoryNode(UInterc
 		return;
 	}
 	const bool bImportOnlyAnimation = CommonSkeletalMeshesAndAnimationsProperties->bImportOnlyAnimations;
+	const bool bAddCurveMetadataToSkeleton = CommonSkeletalMeshesAndAnimationsProperties->bAddCurveMetadataToSkeleton;
 
 	const UInterchangeSkeletonFactoryNode* SkeletonFactoryNode = nullptr;
 	const UInterchangeSkeletalMeshFactoryNode* SkeletalMeshFactoryNode = nullptr;
diff --git a/Engine/Plugins/Interchange/Runtime/Source/Pipelines/Private/InterchangeGenericSkeletalMeshPipeline.cpp b/Engine/Plugins/Interchange/Runtime/Source/Pipelines/Private/InterchangeGenericSkeletalMeshPipeline.cpp
--- a/Engine/Plugins/Interchange/Runtime/Source/Pipelines/Private/InterchangeGenericSkeletalMeshPipeline.cpp
+++ b/Engine/Plugins/Interchange/Runtime/Source/Pipelines/Private/InterchangeGenericSkeletalMeshPipeline.cpp
@@ -450,6 +450,7 @@ UInterchangeSkeletalMeshFactoryNode* UInterchangeGenericMeshPipeline::CreateSkel
 
 	SkeletalMeshFactoryNode->SetCustomImportContentType(SkeletalMeshImportContentType);
 
+	SkeletalMeshFactoryNode->SetCustomAddCurveMetadataToSkeleton(CommonSkeletalMeshesAndAnimationsProperties->bAddCurveMetadataToSkeleton);
 	//If we have a specified skeleton
 	if (CommonSkeletalMeshesAndAnimationsProperties->Skeleton.IsValid())
 	{
diff --git a/Engine/Plugins/Interchange/Runtime/Source/Pipelines/Private/InterchangePipelinesModule.cpp b/Engine/Plugins/Interchange/Runtime/Source/Pipelines/Private/InterchangePipelinesModule.cpp
--- a/Engine/Plugins/Interchange/Runtime/Source/Pipelines/Private/InterchangePipelinesModule.cpp
+++ b/Engine/Plugins/Interchange/Runtime/Source/Pipelines/Private/InterchangePipelinesModule.cpp
@@ -4,10 +4,18 @@
 
 #include "CoreMinimal.h"
 #include "InterchangePipelineLog.h"
+#include "InterchangePipelineObjectVersion.h"
+#include "Misc/Guid.h"
 #include "Modules/ModuleManager.h"
+#include "UObject/DevObjectVersion.h"
 
 DEFINE_LOG_CATEGORY(LogInterchangePipeline);
 
+// Unique serialization id for InterchangePipelineObjectVersion.
+const FGuid FFortniteMainInterchangePipelineObjectVersion::GUID(0xB69D2E47, 0xE2A84003, 0xBF7718A4, 0x92C4D899);
+
+static FDevVersionRegistration GRegisterInterchangePipelineObjectVersion(FFortniteMainInterchangePipelineObjectVersion::GUID, FFortniteMainInterchangePipelineObjectVersion::LatestVersion, TEXT("FN-Main-InterchangePipeline"));
+
 class FInterchangePipelinesModule : public IInterchangePipelinesModule
 {
 	/** IModuleInterface implementation */
diff --git a/Engine/Plugins/Interchange/Runtime/Source/Pipelines/Public/InterchangeGenericAnimationPipeline.h b/Engine/Plugins/Interchange/Runtime/Source/Pipelines/Public/InterchangeGenericAnimationPipeline.h
--- a/Engine/Plugins/Interchange/Runtime/Source/Pipelines/Public/InterchangeGenericAnimationPipeline.h
+++ b/Engine/Plugins/Interchange/Runtime/Source/Pipelines/Public/InterchangeGenericAnimationPipeline.h
@@ -41,6 +41,8 @@ public:
 	virtual bool CanEditChange(const FProperty* InProperty) const override;
 #endif
 
+	virtual void Serialize(FArchive& Ar) override;
+
 	static FString GetPipelineCategory(UClass* AssetClass);
 
 	//Common SkeletalMeshes And Animations Properties Settings Pointer
@@ -85,8 +87,8 @@ public:
 	bool bImportCustomAttribute = true;
 
 	/** Determines whether to automatically add curve metadata to an animation's skeleton. If this setting is disabled, curve metadata will be added to skeletal meshes for morph targets, but no metadata entry will be created for general curves. */
-	UPROPERTY(EditAnywhere, BlueprintReadWrite, Category = "Animations", meta = (SubCategory = "Curves", EditCondition = "bImportAnimations && bImportCustomAttribute"))
-	bool bAddCurveMetadataToSkeleton = true;
+	UPROPERTY()
+	bool bAddCurveMetadataToSkeleton_DEPRECATED = true;
 
 	/** Set the material curve type for all custom attributes. */
 	UPROPERTY(EditAnywhere, BlueprintReadWrite, Category = "Animations", meta = (SubCategory = "Curves", EditCondition = "bImportAnimations && bImportCustomAttribute", DisplayName = "Set Material Curve Type"))
diff --git a/Engine/Plugins/Interchange/Runtime/Source/Pipelines/Public/InterchangeGenericAssetsPipelineSharedSettings.h b/Engine/Plugins/Interchange/Runtime/Source/Pipelines/Public/InterchangeGenericAssetsPipelineSharedSettings.h
--- a/Engine/Plugins/Interchange/Runtime/Source/Pipelines/Public/InterchangeGenericAssetsPipelineSharedSettings.h
+++ b/Engine/Plugins/Interchange/Runtime/Source/Pipelines/Public/InterchangeGenericAssetsPipelineSharedSettings.h
@@ -177,9 +177,14 @@ public:
 	UPROPERTY(EditAnywhere, BlueprintReadWrite, Category = "Common Skeletal Meshes and Animations", meta=(DisplayName="Use T0 As Ref Pose"))
 	bool bUseT0AsRefPose = false;
 
+	/** Determines whether to automatically add curve metadata to a skeleton. If this setting is disabled, curve metadata will be added to skeletal meshes for morph targets, but no metadata entry will be created for general curves. */
+	UPROPERTY(EditAnywhere, BlueprintReadWrite, Category = "Common Skeletal Meshes and Animations")
+	bool bAddCurveMetadataToSkeleton = true;
+
 	/** If enabled, all static meshes that have morph targets will be imported as skeletal meshes instead. */
 	UPROPERTY(EditAnywhere, BlueprintReadWrite, Category = "Static Meshes")
 	bool bConvertStaticsWithMorphTargetsToSkeletals = false;
+
 #if WITH_EDITOR
 	virtual bool IsPropertyChangeNeedRefresh(const FPropertyChangedEvent& PropertyChangedEvent) const override
 	{
diff --git a/Engine/Plugins/Interchange/Runtime/Source/Pipelines/Public/InterchangePipelineObjectVersion.h b/Engine/Plugins/Interchange/Runtime/Source/Pipelines/Public/InterchangePipelineObjectVersion.h
new file mode 100644
--- /dev/null
+++ b/Engine/Plugins/Interchange/Runtime/Source/Pipelines/Public/InterchangePipelineObjectVersion.h
@@ -0,0 +1,25 @@
+// Copyright Epic Games, Inc. All Rights Reserved.
+#pragma once
+
+#include "Misc/Guid.h"
+
+struct FFortniteMainInterchangePipelineObjectVersion
+{
+	// Not instantiable.
+	FFortniteMainInterchangePipelineObjectVersion() = delete;
+
+	enum Type
+	{
+		InitialVersion,
+
+		// Interchange Move generic animation pipeline property bAddCurveMetadataToSkeleton to the generic shared skeletalmesh and animation pipeline
+		InterchangeAddCurveMetadataToSkeletonPropertyMove,
+
+		// -----<new versions can be added above this line>-------------------------------------------------
+		VersionPlusOne,
+		LatestVersion = VersionPlusOne - 1
+	};
+
+	// The GUID for this custom version number
+	INTERCHANGEPIPELINES_API const static FGuid GUID;
+};
