GLTF Interchange Analytics Improvements
- Application Name, Version are now passed in the Analytics

#rnx

diff --git a/Engine/Plugins/Interchange/Runtime/Source/Import/Private/Gltf/InterchangeGltfTranslator.cpp b/Engine/Plugins/Interchange/Runtime/Source/Import/Private/Gltf/InterchangeGltfTranslator.cpp
--- a/Engine/Plugins/Interchange/Runtime/Source/Import/Private/Gltf/InterchangeGltfTranslator.cpp
+++ b/Engine/Plugins/Interchange/Runtime/Source/Import/Private/Gltf/InterchangeGltfTranslator.cpp
@@ -503,6 +503,28 @@ bool UInterchangeGLTFTranslator::Translate( UInterchangeBaseNodeContainer& NodeC
 	if (GltfAsset.Metadata.GeneratorName.Len() > 0)
 	{
 		SourceNode->SetExtraInformation(TEXT("Generator Name"), GltfAsset.Metadata.GeneratorName);
+
+		bool bVersionStringFound = false;
+		int32 CharSpaceIndex = -1;
+		if (GltfAsset.Metadata.GeneratorName.FindLastChar(' ', CharSpaceIndex))
+		{
+			if ((CharSpaceIndex < GltfAsset.Metadata.GeneratorName.Len() - 1) && GltfAsset.Metadata.GeneratorName[CharSpaceIndex + 1] == 'v')
+			{
+				const int32 VersionStringCount = GltfAsset.Metadata.GeneratorName.Len() - (CharSpaceIndex + 1);
+				const FString ApplicationVersionString = GltfAsset.Metadata.GeneratorName.Right(VersionStringCount);
+				const FString ApplicationNameString = GltfAsset.Metadata.GeneratorName.Left(CharSpaceIndex);
+
+				// To achieve uniformity across all translators.
+				SourceNode->SetExtraInformation(UE::Interchange::FSourceNodeExtraInfoStaticData::GetApplicationNameExtraInfoKey(), ApplicationNameString);
+				SourceNode->SetExtraInformation(UE::Interchange::FSourceNodeExtraInfoStaticData::GetApplicationVersionExtraInfoKey(), ApplicationVersionString);
+				bVersionStringFound = true;
+			}
+		}
+		
+		if (!bVersionStringFound)
+		{
+			SourceNode->SetExtraInformation(UE::Interchange::FSourceNodeExtraInfoStaticData::GetApplicationNameExtraInfoKey(), GltfAsset.Metadata.GeneratorName);
+		}
 	}
 	for (const GLTF::FMetadata::FExtraData& Extra : GltfAsset.Metadata.Extras)
 	{
