MaterialX:
- Add an alternative when <heighttonormal> input is not directly connected to an <image> node by using the material function HeightFromNormalSmooth
- added a material expression Contrast instead of the node graph
- created a material function for <colorcorrect> instead of using the node graph
- created a material function for <hsvadjust> instead of using the node graph
- Fix a "leak" of other shadergraph that can happen when several mtlx coming from a single source such as a usd file, where the surfaceshader node is the same across all materials, leading to all materials having the same graph.

#virtualized

diff --git a/Engine/Plugins/Interchange/Runtime/Source/Common/Public/MaterialX/InterchangeMaterialXDefinitions.h b/Engine/Plugins/Interchange/Runtime/Source/Common/Public/MaterialX/InterchangeMaterialXDefinitions.h
--- a/Engine/Plugins/Interchange/Runtime/Source/Common/Public/MaterialX/InterchangeMaterialXDefinitions.h
+++ b/Engine/Plugins/Interchange/Runtime/Source/Common/Public/MaterialX/InterchangeMaterialXDefinitions.h
@@ -148,6 +148,7 @@ namespace UE
 			static constexpr const TCHAR* BlackBody = TEXT("BlackBody");
 			static constexpr const TCHAR* Ceil = TEXT("Ceil");
 			static constexpr const TCHAR* Clamp = TEXT("Clamp");
+			static constexpr const TCHAR* Contrast = TEXT("MaterialXContrast");
 			static constexpr const TCHAR* Cosine = TEXT("Cosine");
 			static constexpr const TCHAR* Crossproduct = TEXT("Crossproduct");
 			static constexpr const TCHAR* Desaturation = TEXT("Desaturation");
@@ -223,6 +224,7 @@ namespace UE
 			static constexpr const TCHAR* AGreaterThanB = TEXT("AGreaterThanB");
 			static constexpr const TCHAR* Alpha = TEXT("Alpha");
 			static constexpr const TCHAR* ALessThanB = TEXT("ALessThanB");
+			static constexpr const TCHAR* Amount= TEXT("Amount");
 			static constexpr const TCHAR* Amplitude = TEXT("Amplitude");
 			static constexpr const TCHAR* B = TEXT("B");
 			static constexpr const TCHAR* Background = TEXT("Background");
@@ -608,6 +610,7 @@ MATERIALX_NAMESPACE_BEGIN
 		static constexpr const char* NewName = "UE:NewName";
 		static constexpr const char* ParentName = "UE:ParentName";
 		static constexpr const char* UniqueName = "UE:UniqueName";
+		static constexpr const char* GeomPropImage = "UE:GeomPropImage";
 	}
 
 	namespace Category
@@ -727,6 +730,9 @@ MATERIALX_NAMESPACE_BEGIN
 		static constexpr const char* PointLight = "point_light";
 		static constexpr const char* SpotLight = "spot_light";
 		// Adjustment nodes
+		static constexpr const char* ColorCorrect = "colorcorrect";
+		static constexpr const char* Contrast = "contrast";
+		static constexpr const char* HsvAdjust = "hsvadjust";
 		static constexpr const char* HsvToRgb = "hsvtorgb";
 		static constexpr const char* Luminance = "luminance";
 		static constexpr const char* Range = "range";
diff --git a/Engine/Plugins/Interchange/Runtime/Source/Import/Private/MaterialX/MaterialExpressions/MaterialExpressionContrast.cpp b/Engine/Plugins/Interchange/Runtime/Source/Import/Private/MaterialX/MaterialExpressions/MaterialExpressionContrast.cpp
new file mode 100644
--- /dev/null
+++ b/Engine/Plugins/Interchange/Runtime/Source/Import/Private/MaterialX/MaterialExpressions/MaterialExpressionContrast.cpp
@@ -0,0 +1,53 @@
+// Copyright Epic Games, Inc. All Rights Reserved.
+
+#include "MaterialExpressionContrast.h"
+#include "MaterialCompiler.h"
+
+#include UE_INLINE_GENERATED_CPP_BY_NAME(MaterialExpressionContrast)
+
+#define LOCTEXT_NAMESPACE "MaterialExpressionMaterialXContrast"
+
+UMaterialExpressionMaterialXContrast::UMaterialExpressionMaterialXContrast(const FObjectInitializer& ObjectInitializer)
+	: Super(ObjectInitializer)
+{
+	// Structure to hold one-time initialization
+	struct FConstructorStatics
+	{
+		FText NAME_MaterialX;
+		FConstructorStatics()
+			: NAME_MaterialX(LOCTEXT("MaterialX", "MaterialX"))
+		{}
+	};
+	static FConstructorStatics ConstructorStatics;
+
+#if WITH_EDITORONLY_DATA
+	MenuCategories.Add(ConstructorStatics.NAME_MaterialX);
+#endif
+}
+
+#if WITH_EDITOR
+int32 UMaterialExpressionMaterialXContrast::Compile(FMaterialCompiler* Compiler, int32 OutputIndex)
+{
+	if(!Input.GetTracedInput().Expression)
+	{
+		return Compiler->Errorf(TEXT("Missing MaterialX Contrast Input"));
+	}
+
+	int32 IndexAmount = Amount.GetTracedInput().Expression ? Amount.Compile(Compiler) : Compiler->Constant(ConstAmount);
+	int32 IndexPivot = Pivot.GetTracedInput().Expression ? Pivot.Compile(Compiler) : Compiler->Constant(ConstPivot);
+
+	int32 IndexSub = Compiler->Sub(Input.Compile(Compiler), IndexPivot);
+	int32 IndexMul = Compiler->Mul(IndexSub, IndexAmount);
+	int32 IndexAdd = Compiler->Add(IndexMul, IndexPivot);
+
+	return IndexAdd;
+}
+
+void UMaterialExpressionMaterialXContrast::GetCaption(TArray<FString>& OutCaptions) const
+{
+	OutCaptions.Add(TEXT("MaterialX Contrast"));
+}
+
+#endif
+
+#undef LOCTEXT_NAMESPACE 
\ No newline at end of file
diff --git a/Engine/Plugins/Interchange/Runtime/Source/Import/Private/MaterialX/MaterialExpressions/MaterialExpressionContrast.h b/Engine/Plugins/Interchange/Runtime/Source/Import/Private/MaterialX/MaterialExpressions/MaterialExpressionContrast.h
new file mode 100644
--- /dev/null
+++ b/Engine/Plugins/Interchange/Runtime/Source/Import/Private/MaterialX/MaterialExpressions/MaterialExpressionContrast.h
@@ -0,0 +1,41 @@
+// Copyright Epic Games, Inc. All Rights Reserved.
+#pragma once
+
+#include "CoreMinimal.h"
+#include "Materials/MaterialExpression.h"
+
+#include "MaterialExpressionContrast.generated.h"
+
+/**
+ * A material expression that increases or decreases contrast of a float/color value using a linear slope multiplier.
+ */
+UCLASS(MinimalAPI, collapsecategories, hidecategories = Object, meta = (Private))
+class UMaterialExpressionMaterialXContrast : public UMaterialExpression
+{
+	GENERATED_UCLASS_BODY()
+
+	UPROPERTY()
+	FExpressionInput Input;
+
+	UPROPERTY()
+	FExpressionInput Amount;
+
+	UPROPERTY()
+	FExpressionInput Pivot;
+
+	/** only used if Amount is not hooked up */
+	UPROPERTY(EditAnywhere, Category = MaterialExpressionContrast, meta = (OverridingInputProperty = "Amount"))
+	float ConstAmount = 1.f;
+
+	/** only used if Pivot is not hooked up */
+	UPROPERTY(EditAnywhere, Category = MaterialExpressionContrast, meta = (OverridingInputProperty = "Pivot"))
+	float ConstPivot= 0.5f;
+
+	//~ Begin UMaterialExpressionMaterialX Interface
+#if WITH_EDITOR
+	virtual int32 Compile(class FMaterialCompiler* Compiler, int32 OutputIndex) override;
+	virtual void GetCaption(TArray<FString>& OutCaptions) const override;
+#endif
+	//~ End UMaterialExpressionMaterialX Interface
+};
+
diff --git a/Engine/Plugins/Interchange/Runtime/Source/Import/Private/MaterialX/MaterialXUtils/MaterialXBase.cpp b/Engine/Plugins/Interchange/Runtime/Source/Import/Private/MaterialX/MaterialXUtils/MaterialXBase.cpp
--- a/Engine/Plugins/Interchange/Runtime/Source/Import/Private/MaterialX/MaterialXUtils/MaterialXBase.cpp
+++ b/Engine/Plugins/Interchange/Runtime/Source/Import/Private/MaterialX/MaterialXUtils/MaterialXBase.cpp
@@ -20,9 +20,12 @@ void FMaterialXBase::UpdateDocumentRecursively(MaterialX::GraphElementPtr Graph)
 	static std::set<std::string> SkippedCategories{
 		mx::Category::ACEScgToLinRec709,
 		mx::Category::AdobeRgbToLinRec709,
+		mx::Category::ColorCorrect,
+		mx::Category::Contrast,
 		mx::Category::Distance,
 		mx::Category::Extract,
 		mx::Category::Fractal3D,
+		mx::Category::HsvAdjust,
 		mx::Category::LinAdobeRgbToLinRec709,
 		mx::Category::LinDisplayP3ToLinRec709,
 		mx::Category::Noise2D,
diff --git a/Engine/Plugins/Interchange/Runtime/Source/Import/Private/MaterialX/MaterialXUtils/MaterialXManager.cpp b/Engine/Plugins/Interchange/Runtime/Source/Import/Private/MaterialX/MaterialXUtils/MaterialXManager.cpp
--- a/Engine/Plugins/Interchange/Runtime/Source/Import/Private/MaterialX/MaterialXUtils/MaterialXManager.cpp
+++ b/Engine/Plugins/Interchange/Runtime/Source/Import/Private/MaterialX/MaterialXUtils/MaterialXManager.cpp
@@ -32,6 +32,7 @@ const TCHAR FMaterialXManager::TexturePayloadSeparator = TEXT('{');
 FMaterialXManager::FMaterialXManager()
 	: MatchingInputNames{
 	   {{TEXT(""),                            TEXT("amplitude")},   ExpressionInput(UE::Expressions::Inputs::Amplitude)},
+	   {{TEXT(""),                            TEXT("amount")},      ExpressionInput(UE::Expressions::Inputs::Amount)},
 	   {{TEXT(""),                            TEXT("bg")},          ExpressionInput(UE::Expressions::Inputs::B)},
 	   {{TEXT(""),                            TEXT("center")},      ExpressionInput(UE::Expressions::Inputs::Center)},
 	   {{TEXT(""),                            TEXT("diminish")},    ExpressionInput(UE::Expressions::Inputs::Diminish)},
@@ -51,6 +52,7 @@ FMaterialXManager::FMaterialXManager()
 	   {{TEXT(""),                            TEXT("lumacoeffs")},  ExpressionInput(UE::Expressions::Inputs::LuminanceFactors)}, // for the moment not yet handled by Interchange, because of the attribute being an advanced pin
 	   {{TEXT(""),                            TEXT("mix")},         ExpressionInput(UE::Expressions::Inputs::Alpha)},
 	   {{TEXT(""),                            TEXT("offset")},      ExpressionInput(UE::Expressions::Inputs::Offset)},
+	   {{TEXT(""),                            TEXT("pivot")},       ExpressionInput(UE::Expressions::Inputs::Pivot)},
 	   {{TEXT(""),                            TEXT("position")},    ExpressionInput(UE::Expressions::Inputs::Position)},
 	   {{TEXT(""),                            TEXT("texcoord")},    ExpressionInput(UE::Expressions::Inputs::Coordinates)},
 	   {{TEXT(""),                            TEXT("octaves")},     ExpressionInput(UE::Expressions::Inputs::Octaves)},
@@ -181,6 +183,7 @@ FMaterialXManager::FMaterialXManager()
 		// Geometric nodes 
 		{mx::Category::GeomColor,    UE::Expressions::Names::VertexColor},
 		// Adjustment nodes,
+		{mx::Category::Contrast,     UE::Expressions::Names::Contrast},
 		{mx::Category::HsvToRgb,     UE::Expressions::Names::HsvToRgb},
 		{mx::Category::Luminance,    UE::Expressions::Names::MaterialXLuminance},
 		{mx::Category::Range,        UE::Expressions::Names::MaterialXRange},
@@ -226,16 +229,19 @@ FMaterialXManager::FMaterialXManager()
 		// Math
 		{mx::Category::Place2D,					FMaterialXMaterialFunction{TInPlaceType<FString>{}, TEXT("/InterchangeAssets/Functions/MX_Place2D.MX_Place2D")}},
 		{mx::Category::Refract,					FMaterialXMaterialFunction{TInPlaceType<FString>{}, TEXT("/Engine/Functions/Engine_MaterialFunctions01/Vectors/Refract.Refract")}},
-		//ColorTransform
+		// ColorTransform
 		{mx::Category::ACEScgToLinRec709,		FMaterialXMaterialFunction{TInPlaceType<FString>{}, TEXT("/InterchangeAssets/Functions/MX_ACEScgToRec709.MX_ACEScgToRec709")}},
 		{mx::Category::AdobeRgbToLinRec709,		FMaterialXMaterialFunction{TInPlaceType<FString>{}, TEXT("/InterchangeAssets/Functions/MX_AdobeRGBToRec709.MX_AdobeRGBToRec709")}},
 		{mx::Category::LinAdobeRgbToLinRec709,	FMaterialXMaterialFunction{TInPlaceType<FString>{}, TEXT("/InterchangeAssets/Functions/MX_LinearAdobeRGBToRec709.MX_LinearAdobeRGBToRec709")}},
 		{mx::Category::LinDisplayP3ToLinRec709,	FMaterialXMaterialFunction{TInPlaceType<FString>{}, TEXT("/InterchangeAssets/Functions/MX_LinearDisplayP3ToRec709.MX_LinearDisplayP3ToRec709")}},
 		{mx::Category::SrgbDisplayP3ToLinRec709,FMaterialXMaterialFunction{TInPlaceType<FString>{}, TEXT("/InterchangeAssets/Functions/MX_SrgbDisplayP3ToRec709.MX_SrgbDisplayP3ToRec709")}},
 		{mx::Category::SrgbTextureToLinRec709,	FMaterialXMaterialFunction{TInPlaceType<FString>{}, TEXT("/InterchangeAssets/Functions/MX_SrgbToRec709.MX_SrgbToRec709")}},
-		//Procedural
+		// Procedural
 		{mx::Category::RandomFloat,	            FMaterialXMaterialFunction{TInPlaceType<FString>{}, TEXT("/InterchangeAssets/Functions/MX_RandomFloat.MX_RandomFloat")}},
 		{mx::Category::RandomColor,	            FMaterialXMaterialFunction{TInPlaceType<FString>{}, TEXT("/InterchangeAssets/Functions/MX_RandomColor.MX_RandomColor")}},
+		// Adjustment
+		{mx::Category::ColorCorrect,	        FMaterialXMaterialFunction{TInPlaceType<FString>{}, TEXT("/InterchangeAssets/Functions/MX_ColorCorrect.MX_ColorCorrect")}},
+		{mx::Category::HsvAdjust,	            FMaterialXMaterialFunction{TInPlaceType<FString>{}, TEXT("/InterchangeAssets/Functions/MX_HSVAdjust.MX_HSVAdjust")}},
 	}
 	, bIsSubstrateEnabled{ IInterchangeImportModule::IsAvailable() ? IInterchangeImportModule::Get().IsSubstrateEnabled() : false }
 {
@@ -390,6 +396,9 @@ namespace UE::Interchange::MaterialX
 																TEXT("MaterialFunction'/InterchangeAssets/Functions/MX_SrgbDisplayP3ToRec709.MX_SrgbDisplayP3ToRec709'"),
 																TEXT("MaterialFunction'/InterchangeAssets/Functions/MX_RandomFloat.MX_RandomFloat'"),
 																TEXT("MaterialFunction'/InterchangeAssets/Functions/MX_RandomColor.MX_RandomColor'"),
+																TEXT("MaterialFunction'/Engine/Functions/Engine_MaterialFunctions03/Procedurals/HeightToNormalSmooth.HeightToNormalSmooth'"),
+																TEXT("MaterialFunction'/InterchangeAssets/Functions/MX_ColorCorrect.MX_ColorCorrect'"),
+																TEXT("MaterialFunction'/InterchangeAssets/Functions/MX_HSVAdjust.MX_HSVAdjust'"),
 															  });
 
 		return bPackagesLoaded;
diff --git a/Engine/Plugins/Interchange/Runtime/Source/Import/Private/MaterialX/MaterialXUtils/MaterialXSurfaceShaderAbstract.cpp b/Engine/Plugins/Interchange/Runtime/Source/Import/Private/MaterialX/MaterialXUtils/MaterialXSurfaceShaderAbstract.cpp
--- a/Engine/Plugins/Interchange/Runtime/Source/Import/Private/MaterialX/MaterialXUtils/MaterialXSurfaceShaderAbstract.cpp
+++ b/Engine/Plugins/Interchange/Runtime/Source/Import/Private/MaterialX/MaterialXUtils/MaterialXSurfaceShaderAbstract.cpp
@@ -6,9 +6,7 @@
 #include "InterchangeImportLog.h"
 #include "InterchangeTexture2DNode.h"
 #include "InterchangeTextureBlurNode.h"
-#include "Materials/MaterialExpressionNoise.h"
 #include "Materials/MaterialExpressionTransform.h"
-#include "Materials/MaterialExpressionTransformPosition.h"
 #include "Materials/MaterialExpressionVectorNoise.h"
 #include "MaterialX/MaterialXUtils/MaterialXManager.h"
 #include "MaterialX/MaterialExpressions/MaterialExpressionTextureSampleParameterBlur.h"
@@ -579,6 +577,12 @@ void FMaterialXSurfaceShaderAbstract::ConnectImageInputToOutput(const FConnectNo
 
 		UInterchangeShaderNode* TextureShaderNode = CreateShaderNode(Connect.UpstreamNode, Connect.UpstreamNode->getName().c_str(), TextureSample::Name.ToString());
 		TextureShaderNode->AddStringAttribute(UInterchangeShaderPortsAPI::MakeInputValueKey(TextureSample::Inputs::Texture.ToString()), TextureNode->GetUniqueID());
+
+		if (Connect.UpstreamNode->getTypedAttribute<bool>(mx::Attributes::GeomPropImage))
+		{
+			// Add a custom attribute that this image node is in fact a geompropvalue converted to an image, in order to update the TextureSample material expression with the baked texture
+			TextureShaderNode->AddBooleanAttribute(mx::Attributes::GeomPropImage, true);
+		}
 		UInterchangeShaderPortsAPI::ConnectOuputToInputByName(Connect.ParentShaderNode, Connect.InputChannelName, TextureShaderNode->GetUniqueID(), OutputChannel);
 	}
 	else
@@ -1200,9 +1204,20 @@ void FMaterialXSurfaceShaderAbstract::ConnectHeightToNormalInputToOutput(const F
 		}
 		else
 		{
-			// For the moment it doesn't make sense to plug a value to it, so let's plug directly the child to the parent, in the future we could implement a Sobel and handle a multi output
-			SetAttributeNewName(Input, TCHAR_TO_UTF8(*Connect.InputChannelName)); //let's take the parent node's input name
-			ShaderNodes.Add({ Connect.UpstreamNode->getName().c_str(), Connect.OutputName }, Connect.ParentShaderNode);
+			using namespace UE::Interchange::Materials::Standard::Nodes;
+			SetAttributeNewName(Connect.UpstreamNode->getInput("in"), "Height");
+			// This MF outputs a normal vector in World Space, we need to transform the output in Tangent Space
+			UInterchangeFunctionCallShaderNode* HeightToNormalSmoothShaderNode =
+				CreateFunctionCallShaderNode(Connect.UpstreamNode,
+											 Connect.UpstreamNode->getName().c_str(),
+											 TEXT("/Engine/Functions/Engine_MaterialFunctions03/Procedurals/HeightToNormalSmooth.HeightToNormalSmooth"));
+
+			UInterchangeShaderNode* TransformNode = CreateShaderNode(Connect.UpstreamNode, (Connect.UpstreamNode->getName() + "_TS").c_str(), TransformVector::Name.ToString());
+			TransformNode->AddInt32Attribute(TransformVector::Attributes::TransformSourceType.ToString(), EMaterialVectorCoordTransformSource::TRANSFORMSOURCE_World);
+			TransformNode->AddInt32Attribute(TransformVector::Attributes::TransformType.ToString(), EMaterialVectorCoordTransform::TRANSFORM_Tangent);
+
+			UInterchangeShaderPortsAPI::ConnectDefaultOuputToInput(TransformNode, TransformVector::Inputs::Input.ToString(), HeightToNormalSmoothShaderNode->GetUniqueID());
+			UInterchangeShaderPortsAPI::ConnectDefaultOuputToInput(Connect.ParentShaderNode, Connect.InputChannelName, TransformNode->GetUniqueID());
 		}
 	}
 }
@@ -1777,7 +1792,7 @@ UInterchangeShaderNode* FMaterialXSurfaceShaderAbstract::Translate(EInterchangeM
 {
 	constexpr bool bInputInTangentSpace = true;
 
-	UInterchangeFunctionCallShaderNode* FunctionCallShaderNode = CreateFunctionCallShaderNode(SurfaceShaderNode->getName().c_str(), UE::Interchange::MaterialX::IndexSurfaceShaders, uint8(ShaderType));
+	UInterchangeFunctionCallShaderNode* FunctionCallShaderNode = CreateFunctionCallShaderNode(SurfaceMaterialName + TEXT("_") + SurfaceShaderNode->getName().c_str(), UE::Interchange::MaterialX::IndexSurfaceShaders, uint8(ShaderType));
 
 	for(mx::InputPtr Input : SurfaceShaderNode->getInputs())
 	{
diff --git a/Engine/Plugins/Interchange/Runtime/Source/Import/Public/MaterialX/MaterialXUtils/MaterialXSurfaceShaderAbstract.h b/Engine/Plugins/Interchange/Runtime/Source/Import/Public/MaterialX/MaterialXUtils/MaterialXSurfaceShaderAbstract.h
--- a/Engine/Plugins/Interchange/Runtime/Source/Import/Public/MaterialX/MaterialXUtils/MaterialXSurfaceShaderAbstract.h
+++ b/Engine/Plugins/Interchange/Runtime/Source/Import/Public/MaterialX/MaterialXUtils/MaterialXSurfaceShaderAbstract.h
@@ -437,7 +437,7 @@ protected:
 					TextureNode->SetCustomWrapV(WrapModeV);
 
 					// Encode the compression in the payloadKey
-					if(bTangentSpaceInput)
+					if(bTangentSpaceInput && Node->getType() == "vector3")
 					{
 						TextureNode->SetPayLoadKey(TextureNode->GetPayLoadKey().GetValue() + FMaterialXManager::TexturePayloadSeparator + FString::FromInt(TextureCompressionSettings::TC_Normalmap));
 					}
