Interchange: Add support for OpenVDB assets via SparseVolumeTextures.

#virtualized

diff --git a/Engine/Plugins/Interchange/Runtime/Source/Common/Public/Volume/InterchangeVolumeDefinitions.h b/Engine/Plugins/Interchange/Runtime/Source/Common/Public/Volume/InterchangeVolumeDefinitions.h
new file mode 100644
--- /dev/null
+++ b/Engine/Plugins/Interchange/Runtime/Source/Common/Public/Volume/InterchangeVolumeDefinitions.h
@@ -0,0 +1,48 @@
+// Copyright Epic Games, Inc. All Rights Reserved.
+
+#pragma once
+
+#include "Containers/StaticArray.h"
+#include "CoreMinimal.h"
+
+#include "InterchangeVolumeDefinitions.generated.h"
+
+UENUM(Blueprintable)
+enum class EInterchangeSparseVolumeTextureFormat : uint8
+{
+	Unorm8 = 0,
+	Float16 = 1,
+	Float32 = 2,
+};
+
+namespace UE::Interchange::Volume
+{
+	const FString DensityGridName = TEXT("density");
+	const FString GridNameAndComponentIndexSeparator = TEXT("_");
+
+	// These structs are direct copies of the ones from "OpenVDBImportOptions.h".
+	// We do this instead of using them directly as that belongs to an editor-only module, and we can't
+	// use #if WITH_EDITOR here as the FVolumePayloadKey shows up on the payload interface, which would
+	// force every implementation to also use #if WITH_EDITOR
+
+	// Describes what should go on a specific texture channel (e.g. AttributesA.Z) within a SparseVolumeTexture
+	struct FComponentMapping
+	{
+		int32 SourceGridIndex = INDEX_NONE;
+		int32 SourceComponentIndex = INDEX_NONE;
+	};
+
+	// Describes a specific texture (e.g. AttributesA) within a SparseVolumeTexture
+	struct FTextureInfo
+	{
+		TStaticArray<FComponentMapping, 4> Mappings;
+		EInterchangeSparseVolumeTextureFormat Format = EInterchangeSparseVolumeTextureFormat::Unorm8;
+	};
+
+	// Describes the full assignment info for a particular SparseVolumeTexture
+	struct FAssignmentInfo
+	{
+		TStaticArray<FTextureInfo, 2> Attributes;
+		bool bIsSequence = false;
+	};
+}	 // namespace UE::Interchange::Volume
diff --git a/Engine/Plugins/Interchange/Runtime/Source/FactoryNodes/InterchangeFactoryNodes.Build.cs b/Engine/Plugins/Interchange/Runtime/Source/FactoryNodes/InterchangeFactoryNodes.Build.cs
--- a/Engine/Plugins/Interchange/Runtime/Source/FactoryNodes/InterchangeFactoryNodes.Build.cs
+++ b/Engine/Plugins/Interchange/Runtime/Source/FactoryNodes/InterchangeFactoryNodes.Build.cs
@@ -13,6 +13,7 @@ namespace UnrealBuildTool.Rules
 					"Core",
 					"CoreUObject",
 					"Engine",
+					"InterchangeCommon",
 					"InterchangeCore",
 					"InterchangeEngine",
 					"InterchangeNodes",
diff --git a/Engine/Plugins/Interchange/Runtime/Source/FactoryNodes/Private/InterchangeSparseVolumeTextureFactoryNode.cpp b/Engine/Plugins/Interchange/Runtime/Source/FactoryNodes/Private/InterchangeSparseVolumeTextureFactoryNode.cpp
new file mode 100644
--- /dev/null
+++ b/Engine/Plugins/Interchange/Runtime/Source/FactoryNodes/Private/InterchangeSparseVolumeTextureFactoryNode.cpp
@@ -0,0 +1,124 @@
+// Copyright Epic Games, Inc. All Rights Reserved.
+
+#include "InterchangeSparseVolumeTextureFactoryNode.h"
+
+#include UE_INLINE_GENERATED_CPP_BY_NAME(InterchangeSparseVolumeTextureFactoryNode)
+
+#if WITH_ENGINE
+#include "SparseVolumeTexture/SparseVolumeTexture.h"
+#endif
+
+FString UInterchangeSparseVolumeTextureFactoryNode::GetTypeName() const
+{
+	const FString TypeName = TEXT("SparseVolumeTextureFactoryNode");
+	return TypeName;
+}
+
+UClass* UInterchangeSparseVolumeTextureFactoryNode::GetObjectClass() const
+{
+#if WITH_ENGINE
+	return USparseVolumeTexture::StaticClass();
+#else
+	return nullptr;
+#endif
+}
+
+bool UInterchangeSparseVolumeTextureFactoryNode::GetCustomAttributesADataType(EInterchangeSparseVolumeTextureFormat& AttributeValue) const
+{
+	IMPLEMENT_NODE_ATTRIBUTE_GETTER(AttributesADataType, EInterchangeSparseVolumeTextureFormat);
+}
+
+bool UInterchangeSparseVolumeTextureFactoryNode::SetCustomAttributesADataType(EInterchangeSparseVolumeTextureFormat AttributeValue)
+{
+	IMPLEMENT_NODE_ATTRIBUTE_SETTER_NODELEGATE(AttributesADataType, EInterchangeSparseVolumeTextureFormat);
+}
+
+bool UInterchangeSparseVolumeTextureFactoryNode::GetCustomAttributesBDataType(EInterchangeSparseVolumeTextureFormat& AttributeValue) const
+{
+	IMPLEMENT_NODE_ATTRIBUTE_GETTER(AttributesBDataType, EInterchangeSparseVolumeTextureFormat);
+}
+
+bool UInterchangeSparseVolumeTextureFactoryNode::SetCustomAttributesBDataType(EInterchangeSparseVolumeTextureFormat AttributeValue)
+{
+	IMPLEMENT_NODE_ATTRIBUTE_SETTER_NODELEGATE(AttributesBDataType, EInterchangeSparseVolumeTextureFormat);
+}
+
+bool UInterchangeSparseVolumeTextureFactoryNode::GetCustomAttributesAChannelX(FString& AttributeValue) const
+{
+	IMPLEMENT_NODE_ATTRIBUTE_GETTER(AttributesAChannelX, FString);
+}
+
+bool UInterchangeSparseVolumeTextureFactoryNode::SetCustomAttributesAChannelX(FString AttributeValue)
+{
+	IMPLEMENT_NODE_ATTRIBUTE_SETTER_NODELEGATE(AttributesAChannelX, FString);
+}
+
+bool UInterchangeSparseVolumeTextureFactoryNode::GetCustomAttributesAChannelY(FString& AttributeValue) const
+{
+	IMPLEMENT_NODE_ATTRIBUTE_GETTER(AttributesAChannelY, FString);
+}
+
+bool UInterchangeSparseVolumeTextureFactoryNode::SetCustomAttributesAChannelY(FString AttributeValue)
+{
+	IMPLEMENT_NODE_ATTRIBUTE_SETTER_NODELEGATE(AttributesAChannelY, FString);
+}
+
+bool UInterchangeSparseVolumeTextureFactoryNode::GetCustomAttributesAChannelZ(FString& AttributeValue) const
+{
+	IMPLEMENT_NODE_ATTRIBUTE_GETTER(AttributesAChannelZ, FString);
+}
+
+bool UInterchangeSparseVolumeTextureFactoryNode::SetCustomAttributesAChannelZ(FString AttributeValue)
+{
+	IMPLEMENT_NODE_ATTRIBUTE_SETTER_NODELEGATE(AttributesAChannelZ, FString);
+}
+
+bool UInterchangeSparseVolumeTextureFactoryNode::GetCustomAttributesAChannelW(FString& AttributeValue) const
+{
+	IMPLEMENT_NODE_ATTRIBUTE_GETTER(AttributesAChannelW, FString);
+}
+
+bool UInterchangeSparseVolumeTextureFactoryNode::SetCustomAttributesAChannelW(FString AttributeValue)
+{
+	IMPLEMENT_NODE_ATTRIBUTE_SETTER_NODELEGATE(AttributesAChannelW, FString);
+}
+
+bool UInterchangeSparseVolumeTextureFactoryNode::GetCustomAttributesBChannelX(FString& AttributeValue) const
+{
+	IMPLEMENT_NODE_ATTRIBUTE_GETTER(AttributesBChannelX, FString);
+}
+
+bool UInterchangeSparseVolumeTextureFactoryNode::SetCustomAttributesBChannelX(FString AttributeValue)
+{
+	IMPLEMENT_NODE_ATTRIBUTE_SETTER_NODELEGATE(AttributesBChannelX, FString);
+}
+
+bool UInterchangeSparseVolumeTextureFactoryNode::GetCustomAttributesBChannelY(FString& AttributeValue) const
+{
+	IMPLEMENT_NODE_ATTRIBUTE_GETTER(AttributesBChannelY, FString);
+}
+
+bool UInterchangeSparseVolumeTextureFactoryNode::SetCustomAttributesBChannelY(FString AttributeValue)
+{
+	IMPLEMENT_NODE_ATTRIBUTE_SETTER_NODELEGATE(AttributesBChannelY, FString);
+}
+
+bool UInterchangeSparseVolumeTextureFactoryNode::GetCustomAttributesBChannelZ(FString& AttributeValue) const
+{
+	IMPLEMENT_NODE_ATTRIBUTE_GETTER(AttributesBChannelZ, FString);
+}
+
+bool UInterchangeSparseVolumeTextureFactoryNode::SetCustomAttributesBChannelZ(FString AttributeValue)
+{
+	IMPLEMENT_NODE_ATTRIBUTE_SETTER_NODELEGATE(AttributesBChannelZ, FString);
+}
+
+bool UInterchangeSparseVolumeTextureFactoryNode::GetCustomAttributesBChannelW(FString& AttributeValue) const
+{
+	IMPLEMENT_NODE_ATTRIBUTE_GETTER(AttributesBChannelW, FString);
+}
+
+bool UInterchangeSparseVolumeTextureFactoryNode::SetCustomAttributesBChannelW(FString AttributeValue)
+{
+	IMPLEMENT_NODE_ATTRIBUTE_SETTER_NODELEGATE(AttributesBChannelW, FString);
+}
diff --git a/Engine/Plugins/Interchange/Runtime/Source/FactoryNodes/Public/InterchangeSparseVolumeTextureFactoryNode.h b/Engine/Plugins/Interchange/Runtime/Source/FactoryNodes/Public/InterchangeSparseVolumeTextureFactoryNode.h
new file mode 100644
--- /dev/null
+++ b/Engine/Plugins/Interchange/Runtime/Source/FactoryNodes/Public/InterchangeSparseVolumeTextureFactoryNode.h
@@ -0,0 +1,114 @@
+// Copyright Epic Games, Inc. All Rights Reserved.
+
+#pragma once
+
+#include "CoreMinimal.h"
+#include "Nodes/InterchangeFactoryBaseNode.h"
+#include "Volume/InterchangeVolumeDefinitions.h"
+
+#include "InterchangeSparseVolumeTextureFactoryNode.generated.h"
+
+UCLASS(BlueprintType)
+class INTERCHANGEFACTORYNODES_API UInterchangeSparseVolumeTextureFactoryNode : public UInterchangeFactoryBaseNode
+{
+	GENERATED_BODY()
+
+public:
+	virtual FString GetTypeName() const override;
+	virtual class UClass* GetObjectClass() const override;
+
+public:
+	/** Gets the data type of the AttributesA texture of the SparseVolumeTexture we'll create */
+	UFUNCTION(BlueprintCallable, Category = "Interchange | Node | Material")
+	bool GetCustomAttributesADataType(EInterchangeSparseVolumeTextureFormat& DataType) const;
+
+	/** Sets the data type of the AttributesA texture of the SparseVolumeTexture we'll create */
+	UFUNCTION(BlueprintCallable, Category = "Interchange | Node | Material")
+	bool SetCustomAttributesADataType(EInterchangeSparseVolumeTextureFormat DataType);
+
+	/** Gets the data type of the AttributesB texture of the SparseVolumeTexture we'll create */
+	UFUNCTION(BlueprintCallable, Category = "Interchange | Node | Material")
+	bool GetCustomAttributesBDataType(EInterchangeSparseVolumeTextureFormat& DataType) const;
+
+	/** Sets the data type of the AttributesB texture of the SparseVolumeTexture we'll create */
+	UFUNCTION(BlueprintCallable, Category = "Interchange | Node | Material")
+	bool SetCustomAttributesBDataType(EInterchangeSparseVolumeTextureFormat DataType);
+
+	/** Gets the grid name and component index that will be assigned to the AttributesA texture, channel X (e.g. "density_0" or "temperature_2") */
+	UFUNCTION(BlueprintCallable, Category = "Interchange | Node | Material")
+	bool GetCustomAttributesAChannelX(FString& GridNameAndComponentIndex) const;
+
+	/** Sets the grid name and component index that will be assigned to the AttributesA texture, channel X (e.g. "density_0" or "temperature_2") */
+	UFUNCTION(BlueprintCallable, Category = "Interchange | Node | Material")
+	bool SetCustomAttributesAChannelX(FString GridNameAndComponentIndex);
+
+	/** Gets the grid name and component index that will be assigned to the AttributesA texture, channel Y (e.g. "density_0" or "temperature_2") */
+	UFUNCTION(BlueprintCallable, Category = "Interchange | Node | Material")
+	bool GetCustomAttributesAChannelY(FString& GridNameAndComponentIndex) const;
+
+	/** Sets the grid name and component index that will be assigned to the AttributesA texture, channel Y (e.g. "density_0" or "temperature_2") */
+	UFUNCTION(BlueprintCallable, Category = "Interchange | Node | Material")
+	bool SetCustomAttributesAChannelY(FString GridNameAndComponentIndex);
+
+	/** Gets the grid name and component index that will be assigned to the AttributesA texture, channel Z (e.g. "density_0" or "temperature_2") */
+	UFUNCTION(BlueprintCallable, Category = "Interchange | Node | Material")
+	bool GetCustomAttributesAChannelZ(FString& GridNameAndComponentIndex) const;
+
+	/** Sets the grid name and component index that will be assigned to the AttributesA texture, channel Z (e.g. "density_0" or "temperature_2") */
+	UFUNCTION(BlueprintCallable, Category = "Interchange | Node | Material")
+	bool SetCustomAttributesAChannelZ(FString GridNameAndComponentIndex);
+
+	/** Gets the grid name and component index that will be assigned to the AttributesA texture, channel W (e.g. "density_0" or "temperature_2") */
+	UFUNCTION(BlueprintCallable, Category = "Interchange | Node | Material")
+	bool GetCustomAttributesAChannelW(FString& GridNameAndComponentIndex) const;
+
+	/** Sets the grid name and component index that will be assigned to the AttributesA texture, channel W (e.g. "density_0" or "temperature_2") */
+	UFUNCTION(BlueprintCallable, Category = "Interchange | Node | Material")
+	bool SetCustomAttributesAChannelW(FString GridNameAndComponentIndex);
+
+	/** Gets the grid name and component index that will be assigned to the AttributesB texture, channel X (e.g. "density_0" or "temperature_2") */
+	UFUNCTION(BlueprintCallable, Category = "Interchange | Node | Material")
+	bool GetCustomAttributesBChannelX(FString& GridNameAndComponentIndex) const;
+
+	/** Sets the grid name and component index that will be assigned to the AttributesB texture, channel X (e.g. "density_0" or "temperature_2") */
+	UFUNCTION(BlueprintCallable, Category = "Interchange | Node | Material")
+	bool SetCustomAttributesBChannelX(FString GridNameAndComponentIndex);
+
+	/** Gets the grid name and component index that will be assigned to the AttributesB texture, channel Y (e.g. "density_0" or "temperature_2") */
+	UFUNCTION(BlueprintCallable, Category = "Interchange | Node | Material")
+	bool GetCustomAttributesBChannelY(FString& GridNameAndComponentIndex) const;
+
+	/** Sets the grid name and component index that will be assigned to the AttributesB texture, channel Y (e.g. "density_0" or "temperature_2") */
+	UFUNCTION(BlueprintCallable, Category = "Interchange | Node | Material")
+	bool SetCustomAttributesBChannelY(FString GridNameAndComponentIndex);
+
+	/** Gets the grid name and component index that will be assigned to the AttributesB texture, channel Z (e.g. "density_0" or "temperature_2") */
+	UFUNCTION(BlueprintCallable, Category = "Interchange | Node | Material")
+	bool GetCustomAttributesBChannelZ(FString& GridNameAndComponentIndex) const;
+
+	/** Sets the grid name and component index that will be assigned to the AttributesB texture, channel Z (e.g. "density_0" or "temperature_2") */
+	UFUNCTION(BlueprintCallable, Category = "Interchange | Node | Material")
+	bool SetCustomAttributesBChannelZ(FString GridNameAndComponentIndex);
+
+	/** Gets the grid name and component index that will be assigned to the AttributesB texture, channel W (e.g. "density_0" or "temperature_2") */
+	UFUNCTION(BlueprintCallable, Category = "Interchange | Node | Material")
+	bool GetCustomAttributesBChannelW(FString& GridNameAndComponentIndex) const;
+
+	/** Sets the grid name and component index that will be assigned to the AttributesB texture, channel W (e.g. "density_0" or "temperature_2") */
+	UFUNCTION(BlueprintCallable, Category = "Interchange | Node | Material")
+	bool SetCustomAttributesBChannelW(FString GridNameAndComponentIndex);
+
+private:
+	IMPLEMENT_NODE_ATTRIBUTE_KEY(AttributesADataType)
+	IMPLEMENT_NODE_ATTRIBUTE_KEY(AttributesBDataType)
+
+	IMPLEMENT_NODE_ATTRIBUTE_KEY(AttributesAChannelX)
+	IMPLEMENT_NODE_ATTRIBUTE_KEY(AttributesAChannelY)
+	IMPLEMENT_NODE_ATTRIBUTE_KEY(AttributesAChannelZ)
+	IMPLEMENT_NODE_ATTRIBUTE_KEY(AttributesAChannelW)
+
+	IMPLEMENT_NODE_ATTRIBUTE_KEY(AttributesBChannelX)
+	IMPLEMENT_NODE_ATTRIBUTE_KEY(AttributesBChannelY)
+	IMPLEMENT_NODE_ATTRIBUTE_KEY(AttributesBChannelZ)
+	IMPLEMENT_NODE_ATTRIBUTE_KEY(AttributesBChannelW)
+};
diff --git a/Engine/Plugins/Interchange/Runtime/Source/Import/InterchangeImport.Build.cs b/Engine/Plugins/Interchange/Runtime/Source/Import/InterchangeImport.Build.cs
--- a/Engine/Plugins/Interchange/Runtime/Source/Import/InterchangeImport.Build.cs
+++ b/Engine/Plugins/Interchange/Runtime/Source/Import/InterchangeImport.Build.cs
@@ -62,6 +62,7 @@ namespace UnrealBuildTool.Rules
 						"InterchangeFbxParser",
 						"MaterialEditor",
 						"SkeletalMeshUtilitiesCommon",
+						"SparseVolumeTexture",
 						"UnrealEd",
 						"VariantManager",
 					}
diff --git a/Engine/Plugins/Interchange/Runtime/Source/Import/Private/InterchangeImportModule.cpp b/Engine/Plugins/Interchange/Runtime/Source/Import/Private/InterchangeImportModule.cpp
--- a/Engine/Plugins/Interchange/Runtime/Source/Import/Private/InterchangeImportModule.cpp
+++ b/Engine/Plugins/Interchange/Runtime/Source/Import/Private/InterchangeImportModule.cpp
@@ -20,24 +20,26 @@
 #include "Mesh/InterchangeStaticMeshFactory.h"
 #include "Misc/CoreDelegates.h"
 #include "Modules/ModuleManager.h"
+#include "OpenVDB/InterchangeOpenVDBTranslator.h"
 #include "Scene/InterchangeActorFactory.h"
 #include "Scene/InterchangeCameraActorFactory.h"
-#include "Scene/InterchangeSceneImportAssetFactory.h"
+#include "Scene/InterchangeDecalActorFactory.h"
 #include "Scene/InterchangeLevelFactory.h"
 #include "Scene/InterchangeLevelInstanceActorFactory.h"
 #include "Scene/InterchangeLightActorFactory.h"
+#include "Scene/InterchangeSceneImportAssetFactory.h"
 #include "Scene/InterchangeSceneVariantSetsFactory.h"
-#include "Scene/InterchangeStaticMeshActorFactory.h"
 #include "Scene/InterchangeSkeletalMeshActorFactory.h"
-#include "Scene/InterchangeDecalActorFactory.h"
-#include "Texture/InterchangeImageWrapperTranslator.h"
+#include "Scene/InterchangeStaticMeshActorFactory.h"
 #include "Texture/InterchangeDDSTranslator.h"
 #include "Texture/InterchangeIESTranslator.h"
+#include "Texture/InterchangeImageWrapperTranslator.h"
 #include "Texture/InterchangeJPGTranslator.h"
-#include "Texture/InterchangeUEJPEGTranslator.h"
 #include "Texture/InterchangePSDTranslator.h"
 #include "Texture/InterchangeTextureFactory.h"
+#include "Texture/InterchangeUEJPEGTranslator.h"
 #include "Usd/InterchangeUsdTranslator.h"
+#include "Volume/InterchangeSparseVolumeTextureFactory.h"
 
 DEFINE_LOG_CATEGORY(LogInterchangeImport);
 
@@ -98,6 +100,13 @@ void FInterchangeImportModule::StartupModule()
 		InterchangeManager.RegisterTranslator(UInterchangePSDTranslator::StaticClass());
 		InterchangeManager.RegisterTranslator(UInterchangeIESTranslator::StaticClass());
 
+		//Volumes
+#if WITH_EDITOR && OPENVDB_AVAILABLE
+		// Editor and OpenVDB-only because we use GetOpenVDBGridInfo and ConvertOpenVDBToSparseVolumeTexture
+		// which use OpenVDB and are in an Editor-only module
+		InterchangeManager.RegisterTranslator(UInterchangeOpenVDBTranslator::StaticClass());
+#endif
+
 		//Register the factories
 		InterchangeManager.RegisterFactory(UInterchangeTextureFactory::StaticClass());
 		InterchangeManager.RegisterFactory(UInterchangeMaterialFactory::StaticClass());
@@ -119,6 +128,10 @@ void FInterchangeImportModule::StartupModule()
 		InterchangeManager.RegisterFactory(UInterchangeSceneVariantSetsFactory::StaticClass());
 		InterchangeManager.RegisterFactory(UInterchangeLightActorFactory::StaticClass());
 		InterchangeManager.RegisterFactory(UInterchangeSceneImportAssetFactory::StaticClass());
+#if WITH_EDITOR
+		// Editor-only because SVT internals like UStreamableSparseVolumeTexture::AppendFrame are editor-only
+		InterchangeManager.RegisterFactory(UInterchangeSparseVolumeTextureFactory::StaticClass());
+#endif
 	};
 	
 	if (GEngine)
diff --git a/Engine/Plugins/Interchange/Runtime/Source/Import/Private/OpenVDB/InterchangeOpenVDBTranslator.cpp b/Engine/Plugins/Interchange/Runtime/Source/Import/Private/OpenVDB/InterchangeOpenVDBTranslator.cpp
new file mode 100644
--- /dev/null
+++ b/Engine/Plugins/Interchange/Runtime/Source/Import/Private/OpenVDB/InterchangeOpenVDBTranslator.cpp
@@ -0,0 +1,377 @@
+// Copyright Epic Games, Inc. All Rights Reserved.
+
+#include "OpenVDB/InterchangeOpenVDBTranslator.h"
+
+#include "InterchangeImportLog.h"
+#include "InterchangeManager.h"
+#include "InterchangeSceneNode.h"
+#include "InterchangeVolumeNode.h"
+
+#include "Async/ParallelFor.h"
+#include "HAL/IConsoleManager.h"
+
+#if WITH_EDITOR
+#include "OpenVDBImportOptions.h"
+#include "SparseVolumeTextureOpenVDBUtility.h"
+#endif	  // WITH_EDITOR
+
+#include UE_INLINE_GENERATED_CPP_BY_NAME(InterchangeOpenVDBTranslator)
+
+#define LOCTEXT_NAMESPACE "InterchangeOpenVDBTranslator"
+
+static bool GInterchangeEnableOpenVDBImport = false;
+static FAutoConsoleVariableRef CVarInterchangeEnableOpenVDBImport(
+	TEXT("Interchange.FeatureFlags.Import.OpenVDB"),
+	GInterchangeEnableOpenVDBImport,
+	TEXT("Whether OpenVDB support is enabled via Interchange")
+);
+
+namespace UE::InterchangeOpenVDBTranslator::Private
+{
+#if WITH_EDITOR
+	struct FOpenVDBFileInfo
+	{
+		TArray64<uint8> FileBytes;
+		TArray<FOpenVDBGridInfo> GridInfo;
+	};
+
+	EVolumeGridElementType GridTypeToInterchangeGridType(EOpenVDBGridType Type)
+	{
+		switch (Type)
+		{
+			case EOpenVDBGridType::Half:
+			case EOpenVDBGridType::Half2:
+			case EOpenVDBGridType::Half3:
+			case EOpenVDBGridType::Half4:
+			{
+				return EVolumeGridElementType::Half;
+			}
+			case EOpenVDBGridType::Float:
+			case EOpenVDBGridType::Float2:
+			case EOpenVDBGridType::Float3:
+			case EOpenVDBGridType::Float4:
+			{
+				return EVolumeGridElementType::Float;
+			}
+			case EOpenVDBGridType::Double:
+			case EOpenVDBGridType::Double2:
+			case EOpenVDBGridType::Double3:
+			case EOpenVDBGridType::Double4:
+			{
+				return EVolumeGridElementType::Double;
+			}
+			default:
+			{
+				return EVolumeGridElementType::Unknown;
+			}
+		}
+	}
+
+	class UInterchangeOpenVDBTranslatorImpl
+	{
+	public:
+		TMap<FString, FOpenVDBFileInfo> FileNameToInfo;
+	};
+#else
+	class UInterchangeOpenVDBTranslatorImpl
+	{
+	};
+#endif	  // WITH_EDITOR
+}	 // namespace UE::InterchangeOpenVDBTranslator::Private
+
+UInterchangeOpenVDBTranslatorSettings::UInterchangeOpenVDBTranslatorSettings()
+	: bTranslateAdjacentNumberedFiles(true)
+{
+}
+
+UInterchangeOpenVDBTranslator::UInterchangeOpenVDBTranslator()
+	: Impl(MakeUnique<UE::InterchangeOpenVDBTranslator::Private::UInterchangeOpenVDBTranslatorImpl>())
+{
+}
+
+EInterchangeTranslatorType UInterchangeOpenVDBTranslator::GetTranslatorType() const
+{
+	return EInterchangeTranslatorType::Assets;
+}
+
+EInterchangeTranslatorAssetType UInterchangeOpenVDBTranslator::GetSupportedAssetTypes() const
+{
+	return EInterchangeTranslatorAssetType::Textures;
+}
+
+TArray<FString> UInterchangeOpenVDBTranslator::GetSupportedFormats() const
+{
+	TArray<FString> Extensions;
+#if WITH_EDITOR
+	if (GInterchangeEnableOpenVDBImport)
+	{
+		Extensions.Add(TEXT("vdb;OpenVDB files"));
+	}
+#endif	  // WITH_EDITOR
+	return Extensions;
+}
+
+bool UInterchangeOpenVDBTranslator::Translate(UInterchangeBaseNodeContainer& NodeContainer) const
+{
+#if WITH_EDITOR && OPENVDB_AVAILABLE
+	using namespace UE::InterchangeOpenVDBTranslator::Private;
+
+	// References:
+	// - SparseVolumeTextureFactory.cpp, LoadOpenVDBPreviewData()
+
+	UInterchangeOpenVDBTranslatorImpl* ImplPtr = Impl.Get();
+	if (!ImplPtr)
+	{
+		return false;
+	}
+
+	UInterchangeOpenVDBTranslatorSettings* Settings = Cast<UInterchangeOpenVDBTranslatorSettings>(GetSettings());
+	if (!Settings)
+	{
+		return false;
+	}
+
+	// Get all relevant VDB filenames to import
+	TArray<FString> Filenames;
+	{
+		const FString MainFilename = GetSourceData()->GetFilename();
+
+		if (Settings->bTranslateAdjacentNumberedFiles)
+		{
+			Filenames = FindOpenVDBSequenceFileNames(MainFilename);
+		}
+		else
+		{
+			Filenames = {MainFilename};
+		}
+	}
+
+	// Create all the entries we'll need in one thread
+	ImplPtr->FileNameToInfo.Reset();
+	for (const FString& Filename : Filenames)
+	{
+		ImplPtr->FileNameToInfo.Add(Filename);
+	}
+
+	// Collect all the frame info in parallel
+	const int32 NumBatches = Filenames.Num();
+	ParallelFor(
+		NumBatches,
+		[ImplPtr, &Filenames](int32 FilenameIndex)
+		{
+			const FString& Filename = Filenames[FilenameIndex];
+			FOpenVDBFileInfo& FileInfo = ImplPtr->FileNameToInfo[Filename];
+
+			if (!FPaths::FileExists(Filename))
+			{
+				UE_LOG(LogInterchangeImport, Error, TEXT("OpenVDB file could not be found: %s"), *Filename);
+				return;
+			}
+
+			if (!FFileHelper::LoadFileToArray(FileInfo.FileBytes, *Filename))
+			{
+				UE_LOG(LogInterchangeImport, Error, TEXT("OpenVDB file could not be loaded: %s"), *Filename);
+				return;
+			}
+
+			const bool bCreateStrings = true;
+			if (!GetOpenVDBGridInfo(FileInfo.FileBytes, bCreateStrings, &FileInfo.GridInfo))
+			{
+				UE_LOG(LogInterchangeImport, Error, TEXT("Failed to read OpenVDB file: %s"), *Filename);
+				return;
+			}
+		}
+	);
+
+	for (const TPair<FString, FOpenVDBFileInfo>& Pair : ImplPtr->FileNameToInfo)
+	{
+		const FString& Filename = Pair.Key;
+		const FOpenVDBFileInfo& FileInfo = Pair.Value;
+
+		const FString VolumeNodeUid = FPaths::GetBaseFilename(Filename);
+
+		// Generate node for this file
+		UInterchangeVolumeNode* VolumeNode = nullptr;
+		{
+			const UInterchangeVolumeNode* Node = Cast<UInterchangeVolumeNode>(NodeContainer.GetNode(VolumeNodeUid));
+			if (Node)
+			{
+				continue;
+			}
+
+			VolumeNode = NewObject<UInterchangeVolumeNode>(&NodeContainer);
+			VolumeNode->InitializeNode(VolumeNodeUid, VolumeNodeUid, EInterchangeNodeContainerType::TranslatedAsset);
+			VolumeNode->SetAssetName(VolumeNodeUid);
+			VolumeNode->SetCustomFileName(Filename);
+			NodeContainer.AddNode(VolumeNode);
+		}
+
+		// Generate nodes for each volume grid
+		for (const FOpenVDBGridInfo& GridInfo : FileInfo.GridInfo)
+		{
+			const FString& GridName = GridInfo.Name;
+			const FString GridNodeUid = VolumeNodeUid + TEXT("_") + GridName;
+
+			const UInterchangeVolumeGridNode* Node = Cast<UInterchangeVolumeGridNode>(NodeContainer.GetNode(VolumeNodeUid));
+			if (Node)
+			{
+				continue;
+			}
+
+			UInterchangeVolumeGridNode* GridNode = NewObject<UInterchangeVolumeGridNode>(&NodeContainer);
+			GridNode->InitializeNode(GridNodeUid, GridName, EInterchangeNodeContainerType::TranslatedAsset);
+			GridNode->SetCustomElementType(GridTypeToInterchangeGridType(GridInfo.Type));
+			GridNode->SetCustomNumComponents(GridInfo.NumComponents);
+			GridNode->SetCustomGridTransform(GridInfo.Transform);
+			GridNode->SetCustomGridActiveAABBMin(GridInfo.VolumeActiveAABBMin);
+			GridNode->SetCustomGridActiveAABBMax(GridInfo.VolumeActiveAABBMax);
+			GridNode->SetCustomGridActiveDimensions(GridInfo.VolumeActiveDim);
+
+			VolumeNode->AddGridDependencyUid(GridNodeUid);
+
+			NodeContainer.AddNode(GridNode);
+			NodeContainer.SetNodeParentUid(GridNodeUid, VolumeNodeUid);
+		}
+	}
+
+	return true;
+#else
+	return false;
+#endif	  // WITH_EDITOR && OPENVDB_AVAILABLE
+}
+
+void UInterchangeOpenVDBTranslator::ReleaseSource()
+{
+#if WITH_EDITOR
+	UE::InterchangeOpenVDBTranslator::Private::UInterchangeOpenVDBTranslatorImpl* ImplPtr = Impl.Get();
+	if (!ImplPtr)
+	{
+		return;
+	}
+
+	ImplPtr->FileNameToInfo.Reset();
+
+	if (TranslatorSettings)
+	{
+		TranslatorSettings->ClearFlags(RF_Standalone);
+		TranslatorSettings = nullptr;
+	}
+#endif	  // WITH_EDITOR
+}
+
+UInterchangeTranslatorSettings* UInterchangeOpenVDBTranslator::GetSettings() const
+{
+#if WITH_EDITOR
+	using namespace UE::InterchangeOpenVDBTranslator::Private;
+
+	UInterchangeOpenVDBTranslatorImpl* ImplPtr = Impl.Get();
+	if (!ImplPtr)
+	{
+		return nullptr;
+	}
+
+	if (!TranslatorSettings)
+	{
+		TranslatorSettings = DuplicateObject<UInterchangeOpenVDBTranslatorSettings>(
+			UInterchangeOpenVDBTranslatorSettings::StaticClass()->GetDefaultObject<UInterchangeOpenVDBTranslatorSettings>(),
+			GetTransientPackage()
+		);
+		TranslatorSettings->LoadSettings();
+		TranslatorSettings->ClearFlags(RF_ArchetypeObject);
+		TranslatorSettings->SetFlags(RF_Standalone);
+		TranslatorSettings->ClearInternalFlags(EInternalObjectFlags::Async);
+	}
+	return TranslatorSettings;
+#else
+	return nullptr;
+#endif	  // WITH_EDITOR
+}
+
+void UInterchangeOpenVDBTranslator::SetSettings(const UInterchangeTranslatorSettings* InterchangeTranslatorSettings)
+{
+#if WITH_EDITOR
+	using namespace UE::InterchangeOpenVDBTranslator::Private;
+
+	UInterchangeOpenVDBTranslatorImpl* ImplPtr = Impl.Get();
+	if (!ImplPtr)
+	{
+		return;
+	}
+
+	if (TranslatorSettings)
+	{
+		TranslatorSettings->ClearFlags(RF_Standalone);
+		TranslatorSettings->ClearInternalFlags(EInternalObjectFlags::Async);
+		TranslatorSettings = nullptr;
+	}
+	if (const UInterchangeOpenVDBTranslatorSettings* OpenVDBTranslatorSettings = Cast<UInterchangeOpenVDBTranslatorSettings>(
+			InterchangeTranslatorSettings
+		))
+	{
+		TranslatorSettings = DuplicateObject<UInterchangeOpenVDBTranslatorSettings>(OpenVDBTranslatorSettings, GetTransientPackage());
+		TranslatorSettings->ClearInternalFlags(EInternalObjectFlags::Async);
+		TranslatorSettings->SetFlags(RF_Standalone);
+	}
+#endif	  // WITH_EDITOR
+}
+
+TOptional<UE::Interchange::FVolumePayloadData> UInterchangeOpenVDBTranslator::GetVolumePayloadData(
+	const UE::Interchange::FVolumePayloadKey& PayloadKey
+) const
+{
+#if WITH_EDITOR && OPENVDB_AVAILABLE
+	using namespace UE::InterchangeOpenVDBTranslator::Private;
+
+	UInterchangeOpenVDBTranslatorImpl* ImplPtr = Impl.Get();
+	if (!ImplPtr)
+	{
+		return {};
+	}
+
+	FOpenVDBFileInfo* FoundInfo = ImplPtr->FileNameToInfo.Find(PayloadKey.FileName);
+	if (!FoundInfo)
+	{
+		return {};
+	}
+
+	// Convert our custom struct to the one used by ConvertOpenVDBToSparseVolumeTexture
+	FOpenVDBImportOptions ImportOptions;
+	{
+		ImportOptions.bIsSequence = PayloadKey.AssignmentInfo.bIsSequence;
+
+		for (int32 TextureIndex = 0; TextureIndex < ImportOptions.Attributes.Num(); ++TextureIndex)
+		{
+			const UE::Interchange::Volume::FTextureInfo& InTextureInfo = PayloadKey.AssignmentInfo.Attributes[TextureIndex];
+			FOpenVDBSparseVolumeAttributesDesc& OutTextureInfo = ImportOptions.Attributes[TextureIndex];
+
+			static_assert((int)ESparseVolumeAttributesFormat::Unorm8 == (int)EInterchangeSparseVolumeTextureFormat::Unorm8);
+			static_assert((int)ESparseVolumeAttributesFormat::Float16 == (int)EInterchangeSparseVolumeTextureFormat::Float16);
+			static_assert((int)ESparseVolumeAttributesFormat::Float32 == (int)EInterchangeSparseVolumeTextureFormat::Float32);
+			OutTextureInfo.Format = static_cast<ESparseVolumeAttributesFormat>(InTextureInfo.Format);
+
+			for (int32 ChannelIndex = 0; ChannelIndex < OutTextureInfo.Mappings.Num(); ++ChannelIndex)
+			{
+				OutTextureInfo.Mappings[ChannelIndex].SourceComponentIndex = InTextureInfo.Mappings[ChannelIndex].SourceComponentIndex;
+				OutTextureInfo.Mappings[ChannelIndex].SourceGridIndex = InTextureInfo.Mappings[ChannelIndex].SourceGridIndex;
+			}
+		}
+	}
+
+	UE::Interchange::FVolumePayloadData Result;
+	bool bSucess = ConvertOpenVDBToSparseVolumeTexture(
+		FoundInfo->FileBytes,
+		ImportOptions,
+		PayloadKey.VolumeBoundsMin,
+		Result.TextureData,
+		Result.Transform
+	);
+	if (bSucess)
+	{
+		return Result;
+	}
+#endif	  // WITH_EDITOR && OPENVDB_AVAILABLE
+
+	return {};
+}
+
+#undef LOCTEXT_NAMESPACE
diff --git a/Engine/Plugins/Interchange/Runtime/Source/Import/Private/Volume/InterchangeSparseVolumeTextureFactory.cpp b/Engine/Plugins/Interchange/Runtime/Source/Import/Private/Volume/InterchangeSparseVolumeTextureFactory.cpp
new file mode 100644
--- /dev/null
+++ b/Engine/Plugins/Interchange/Runtime/Source/Import/Private/Volume/InterchangeSparseVolumeTextureFactory.cpp
@@ -0,0 +1,976 @@
+// Copyright Epic Games, Inc. All Rights Reserved.
+
+#include "Volume/InterchangeSparseVolumeTextureFactory.h"
+
+#include "InterchangeAssetImportData.h"
+#include "InterchangeImportCommon.h"
+#include "InterchangeResult.h"
+#include "InterchangeSparseVolumeTextureFactoryNode.h"
+#include "InterchangeTranslatorBase.h"
+#include "InterchangeVolumeNode.h"
+#include "Nodes/InterchangeBaseNode.h"
+#include "Volume/InterchangeVolumeDefinitions.h"
+#include "Volume/InterchangeVolumePayloadData.h"
+#include "Volume/InterchangeVolumePayloadInterface.h"
+
+#include "Async/ParallelFor.h"
+#include "Async/TaskGraphInterfaces.h"
+#include "Misc/CoreStats.h"
+#include "Misc/Paths.h"
+#include "ProfilingDebugging/CpuProfilerTrace.h"
+#include "Serialization/EditorBulkData.h"
+#include "SparseVolumeTexture/SparseVolumeTexture.h"
+#include "UDIMUtilities.h"
+#include "UObject/ObjectMacros.h"
+#if WITH_EDITORONLY_DATA
+#include "EditorFramework/AssetImportData.h"
+#endif	  // WITH_EDITORONLY_DATA
+
+#include UE_INLINE_GENERATED_CPP_BY_NAME(InterchangeSparseVolumeTextureFactory)
+
+#define LOCTEXT_NAMESPACE "InterchangeSparseVolumeTextureFactory"
+
+namespace UE::InterchangeSparseVolumeTextureFactory::Private
+{
+	// Returns the UInterchangeVolumeNodes referenced as targets by the provided UInterchangeSparseVolumeTextureFactoryNode
+	TArray<const UInterchangeVolumeNode*> GetVolumeNodes(
+		const UInterchangeSparseVolumeTextureFactoryNode* FactoryNode,
+		const UInterchangeBaseNodeContainer* NodeContainer
+	)
+	{
+		if (!FactoryNode || !NodeContainer)
+		{
+			return {};
+		}
+
+		TArray<const UInterchangeVolumeNode*> VolumeNodes;
+
+		TArray<FString> TargetNodeUids;
+		FactoryNode->GetTargetNodeUids(TargetNodeUids);
+
+		for (int32 TargetNodeIndex = TargetNodeUids.Num() - 1; TargetNodeIndex >= 0; --TargetNodeIndex)
+		{
+			const UInterchangeBaseNode* BaseNode = NodeContainer->GetNode(TargetNodeUids[TargetNodeIndex]);
+			if (const UInterchangeVolumeNode* VolumeNode = Cast<UInterchangeVolumeNode>(BaseNode))
+			{
+				VolumeNodes.Add(VolumeNode);
+			}
+		}
+
+		return VolumeNodes;
+	}
+
+	// Returns the UInterchangeVolumeGridNodes referenced as dependencies by the provided UInterchangeVolumeNode
+	TArray<const UInterchangeVolumeGridNode*> GetGridNodes(
+		const UInterchangeVolumeNode* VolumeNode,
+		const UInterchangeBaseNodeContainer* NodeContainer
+	)
+	{
+		if (!VolumeNode || !NodeContainer)
+		{
+			return {};
+		}
+
+		TArray<FString> GridNodeUids;
+		VolumeNode->GetGridDependencies(GridNodeUids);
+
+		TArray<const UInterchangeVolumeGridNode*> Result;
+		Result.Reserve(GridNodeUids.Num());
+		for (const FString& GridNodeUid : GridNodeUids)
+		{
+			Result.Add(Cast<UInterchangeVolumeGridNode>(NodeContainer->GetNode(GridNodeUid)));
+		}
+
+		return Result;
+	}
+
+#if WITH_EDITOR
+	void HashSourceFiles(TArray<FAssetImportInfo::FSourceFile>& SourceFiles)
+	{
+		// References:
+		// - GenerateHashSourceFilesTasks from InterchangeTextureFactory.cpp
+
+		struct FHashTask
+		{
+			FHashTask(FAssetImportInfo::FSourceFile* InSourceFile)
+				: SourceFile(InSourceFile)
+			{
+			}
+
+			ENamedThreads::Type GetDesiredThread()
+			{
+				return ENamedThreads::AnyBackgroundThreadNormalTask;
+			}
+
+			TStatId GetStatId() const
+			{
+				return GET_STATID(STAT_TaskGraph_OtherTasks);
+			}
+
+			static ESubsequentsMode::Type GetSubsequentsMode()
+			{
+				return ESubsequentsMode::TrackSubsequents;
+			}
+
+			void DoTask(ENamedThreads::Type CurrentThread, const FGraphEventRef& MyCompletionGraphEvent)
+			{
+				SourceFile->FileHash = FMD5Hash::HashFile(*SourceFile->RelativeFilename);
+				SourceFile->Timestamp = IFileManager::Get().GetTimeStamp(*SourceFile->RelativeFilename);
+			}
+
+		private:
+			FAssetImportInfo::FSourceFile* SourceFile;
+		};
+
+		FGraphEventArray TasksToDo;
+		for (FAssetImportInfo::FSourceFile& SourceFile : SourceFiles)
+		{
+			TasksToDo.Add(TGraphTask<FHashTask>::CreateTask().ConstructAndDispatchWhenReady(&SourceFile));
+		}
+
+		ENamedThreads::Type NamedThread = IsInGameThread() ? ENamedThreads::GameThread : ENamedThreads::AnyThread;
+		FTaskGraphInterface::Get().WaitUntilTasksComplete(TasksToDo, NamedThread);
+	}
+#endif	  // WITH_EDITOR
+
+	// Actually fills in the provided SVT asset with the payload / source data
+	void SetupSparseVolumeTexture(
+		USparseVolumeTexture* SparseVolumeTexture,
+		TArray<TOptional<UE::Interchange::FVolumePayloadData>>& ProcessedPayloads,
+		bool bIsReimport
+	)
+	{
+		UStreamableSparseVolumeTexture* Streamable = Cast<UStreamableSparseVolumeTexture>(SparseVolumeTexture);
+		int32 NumPayloads = ProcessedPayloads.Num();
+		if (!Streamable || NumPayloads < 1)
+		{
+			return;
+		}
+
+		bool bSucess = Streamable->BeginInitialize(NumPayloads);
+		if (!bSucess)
+		{
+			return;
+		}
+
+		for (TOptional<UE::Interchange::FVolumePayloadData>& ProcessedPayload : ProcessedPayloads)
+		{
+			if (!ProcessedPayload.IsSet())
+			{
+				continue;
+			}
+
+			bSucess = Streamable->AppendFrame(ProcessedPayload->TextureData, ProcessedPayload->Transform);
+			if (!bSucess)
+			{
+				return;
+			}
+		}
+
+		bSucess = Streamable->EndInitialize();
+		if (!bSucess)
+		{
+			return;
+		}
+	}
+
+	// We have to convert the assignment info we spread out into individual string attributes on the factory node back into an
+	// index-based description of assignment info via FAssignmentInfo, which will ultimately be converted into an
+	// FOpenVDBImportOptions by the translator when retrieving the payload, as that is what the OpenVDB utils expect.
+	//
+	// It's probably for the best though, as this can just live in here and instead users would only interact with the string attributes
+	// on the factory nodes
+	TOptional<UE::Interchange::Volume::FAssignmentInfo> GetAssignmentInfo(
+		UInterchangeSparseVolumeTextureFactoryNode* FactoryNode,
+		TArray<const UInterchangeVolumeNode*>& VolumeNodes,
+		const UInterchangeBaseNodeContainer* NodeContainer
+	)
+	{
+		using namespace UE::Interchange::Volume;
+
+		if (!FactoryNode || !NodeContainer || VolumeNodes.Num() == 0)
+		{
+			return {};
+		}
+
+		FAssignmentInfo Result;
+		Result.bIsSequence = VolumeNodes.Num() > 1;
+
+		TArray<const UInterchangeVolumeGridNode*> OrderedGridNodes;
+		{
+			// The grid dependency order still matches the grid order in the OpenVDB file.
+			//
+			// If we're importing multiple volume nodes, it's implied that they're separate frames of the same animated SparseVolumeTexture,
+			// and that means they should have the same grid and assignment, so picking any volume node should do
+			TArray<FString> GridNodeUids;
+			VolumeNodes[0]->GetGridDependencies(GridNodeUids);
+
+			OrderedGridNodes.Reserve(GridNodeUids.Num());
+			for (const FString& GridNodeUid : GridNodeUids)
+			{
+				const UInterchangeVolumeGridNode* GridNode = Cast<UInterchangeVolumeGridNode>(NodeContainer->GetNode(GridNodeUid));
+				ensure(GridNode);
+
+				OrderedGridNodes.Add(GridNode);
+			}
+		}
+
+		TMap<FString, int32> GridNameToIndex;
+		GridNameToIndex.Reserve(OrderedGridNodes.Num());
+		for (int32 Index = 0; Index < OrderedGridNodes.Num(); ++Index)
+		{
+			const UInterchangeVolumeGridNode* GridNode = OrderedGridNodes[Index];
+			GridNameToIndex.Add(GridNode->GetDisplayLabel(), Index);
+		}
+
+		// Convert formats
+		{
+			EInterchangeSparseVolumeTextureFormat Format;
+
+			Result.Attributes[0].Format = EInterchangeSparseVolumeTextureFormat::Float16;
+			if (FactoryNode->GetCustomAttributesADataType(Format))
+			{
+				Result.Attributes[0].Format = Format;
+			}
+
+			Result.Attributes[1].Format = EInterchangeSparseVolumeTextureFormat::Float16;
+			if (FactoryNode->GetCustomAttributesBDataType(Format))
+			{
+				Result.Attributes[1].Format = Format;
+			}
+		}
+
+		// Convert channels
+		using GetterFunc = decltype(&UInterchangeSparseVolumeTextureFactoryNode::GetCustomAttributesAChannelX);
+		static const TArray<GetterFunc> AttributeChannelGetters = {
+			&UInterchangeSparseVolumeTextureFactoryNode::GetCustomAttributesAChannelX,
+			&UInterchangeSparseVolumeTextureFactoryNode::GetCustomAttributesAChannelY,
+			&UInterchangeSparseVolumeTextureFactoryNode::GetCustomAttributesAChannelZ,
+			&UInterchangeSparseVolumeTextureFactoryNode::GetCustomAttributesAChannelW,
+			&UInterchangeSparseVolumeTextureFactoryNode::GetCustomAttributesBChannelX,
+			&UInterchangeSparseVolumeTextureFactoryNode::GetCustomAttributesBChannelY,
+			&UInterchangeSparseVolumeTextureFactoryNode::GetCustomAttributesBChannelZ,
+			&UInterchangeSparseVolumeTextureFactoryNode::GetCustomAttributesBChannelW,
+		};
+		for (int32 ChannelIndex = 0; ChannelIndex < AttributeChannelGetters.Num(); ++ChannelIndex)
+		{
+			// e.g. "temperature_6"
+			FString GridNameAndComponentIndex;
+			bool bGetResult = (FactoryNode->*AttributeChannelGetters[ChannelIndex])(GridNameAndComponentIndex);
+			if (!bGetResult || GridNameAndComponentIndex.IsEmpty())
+			{
+				continue;
+			}
+
+			FString GridName;
+			FString ComponentIndexStr;
+			bool bSplit = GridNameAndComponentIndex.Split(
+				GridNameAndComponentIndexSeparator,
+				&GridName,
+				&ComponentIndexStr,
+				ESearchCase::CaseSensitive,
+				ESearchDir::FromEnd
+			);
+			if (!bSplit || GridName.IsEmpty() || ComponentIndexStr.IsEmpty())
+			{
+				continue;
+			}
+
+			int32* FoundGridIndex = GridNameToIndex.Find(GridName);
+			if (!FoundGridIndex)
+			{
+				continue;
+			}
+
+			int32 ComponentIndex = INDEX_NONE;
+			bool bLexed = LexTryParseString(ComponentIndex, *ComponentIndexStr);
+			if (!bLexed)
+			{
+				continue;
+			}
+
+			int32 TextureIndex = ChannelIndex / 4;
+			int32 ChannelComponentIndex = ChannelIndex % 4;
+
+			FTextureInfo& TextureInfo = Result.Attributes[TextureIndex];
+			FComponentMapping& Mapping = TextureInfo.Mappings[ChannelComponentIndex];
+			Mapping.SourceGridIndex = *FoundGridIndex;
+			Mapping.SourceComponentIndex = ComponentIndex;
+		}
+
+		return Result;
+	}
+
+	// Computes the combined AABB bounds for all the volume grids in the provided nodes.
+	// The OpenVDB utils need this when retrieving the final payload data.
+	void ComputeExpandedVolumeBounds(
+		TArray<const UInterchangeVolumeNode*>& InVolumeNodes,
+		const UInterchangeBaseNodeContainer* InNodeContainer,
+		FIntVector3& OutVolumeBoundsMin,
+		FIntVector3& OutVolumeBoundsMax
+	)
+	{
+		FIntVector3 VolumeBoundsMin = FIntVector3(INT32_MAX, INT32_MAX, INT32_MAX);
+		FIntVector3 VolumeBoundsMax = FIntVector3(INT32_MIN, INT32_MIN, INT32_MIN);
+
+		for (const UInterchangeVolumeNode* VolumeNode : InVolumeNodes)
+		{
+			for (const UInterchangeVolumeGridNode* GridNode : GetGridNodes(VolumeNode, InNodeContainer))
+			{
+				if (!GridNode)
+				{
+					continue;
+				}
+
+				FIntVector Min;
+				if (GridNode->GetCustomGridActiveAABBMin(Min))
+				{
+					FIntVector3 GridMin{Min};
+					VolumeBoundsMin.X = FMath::Min(VolumeBoundsMin.X, GridMin.X);
+					VolumeBoundsMin.Y = FMath::Min(VolumeBoundsMin.Y, GridMin.Y);
+					VolumeBoundsMin.Z = FMath::Min(VolumeBoundsMin.Z, GridMin.Z);
+				}
+
+				FIntVector Max;
+				if (GridNode->GetCustomGridActiveAABBMax(Max))
+				{
+					FIntVector3 GridMax{Max};
+					VolumeBoundsMax.X = FMath::Min(VolumeBoundsMax.X, GridMax.X);
+					VolumeBoundsMax.Y = FMath::Min(VolumeBoundsMax.Y, GridMax.Y);
+					VolumeBoundsMax.Z = FMath::Min(VolumeBoundsMax.Z, GridMax.Z);
+				}
+			}
+		}
+
+		OutVolumeBoundsMin = VolumeBoundsMin;
+		OutVolumeBoundsMax = VolumeBoundsMax;
+	}
+
+	// Sanity check that all the provided volume nodes have the same grid arrangement, with the same number
+	// of components and data types. If this passes, we can make a single animated SVT from those volume frames.
+	bool CheckForGridConsistency(const TArray<const UInterchangeVolumeNode*>& VolumeNodes, const UInterchangeBaseNodeContainer* NodeContainer)
+	{
+		if (VolumeNodes.Num() < 2)
+		{
+			// A single volume is always "consistent"
+			return true;
+		}
+
+		struct FGridInfo
+		{
+			EVolumeGridElementType GridType = EVolumeGridElementType::Unknown;
+			int32 NumComponents = 0;
+		};
+
+		// Get info from the first volume
+		TMap<FString, FGridInfo> GridNameToInfo;
+		{
+			const UInterchangeVolumeNode* FirstVolume = VolumeNodes[0];
+			for (const UInterchangeVolumeGridNode* GridNode : GetGridNodes(FirstVolume, NodeContainer))
+			{
+				FString DisplayLabel = GridNode->GetDisplayLabel();
+
+				int32 NumComponents = 0;
+				GridNode->GetCustomNumComponents(NumComponents);
+
+				EVolumeGridElementType GridType = EVolumeGridElementType::Unknown;
+				GridNode->GetCustomElementType(GridType);
+
+				if (GridNameToInfo.Contains(DisplayLabel))
+				{
+					// Grid names should be unique within a volume
+					return false;
+				}
+				GridNameToInfo.Add(DisplayLabel, FGridInfo{GridType, NumComponents});
+			}
+		}
+
+		// Compare it with the other volumes
+		for (int32 Index = 0; Index < VolumeNodes.Num(); ++Index)
+		{
+			const UInterchangeVolumeNode* OtherVolume = VolumeNodes[Index];
+
+			TArray<const UInterchangeVolumeGridNode*> GridNodes = GetGridNodes(OtherVolume, NodeContainer);
+			if (GridNodes.Num() != GridNameToInfo.Num())
+			{
+				// Should have the same number of grids as the first volume
+				return false;
+			}
+
+			for (const UInterchangeVolumeGridNode* GridNode : GridNodes)
+			{
+				FString DisplayLabel = GridNode->GetDisplayLabel();
+
+				int32 NumComponents = 0;
+				GridNode->GetCustomNumComponents(NumComponents);
+
+				EVolumeGridElementType GridType = EVolumeGridElementType::Unknown;
+				GridNode->GetCustomElementType(GridType);
+
+				const FGridInfo* FoundInfo = GridNameToInfo.Find(DisplayLabel);
+				if (!FoundInfo)
+				{
+					// Grid is not present on the first volume
+					return false;
+				}
+
+				if (FoundInfo->NumComponents != NumComponents || FoundInfo->GridType != GridType)
+				{
+					// Grid is different from the corresponding grid of the first volume
+					return false;
+				}
+			}
+		}
+
+		return true;
+	}
+}	 // namespace UE::InterchangeSparseVolumeTextureFactory::Private
+
+UClass* UInterchangeSparseVolumeTextureFactory::GetFactoryClass() const
+{
+	return USparseVolumeTexture::StaticClass();
+}
+
+UInterchangeFactoryBase::FImportAssetResult UInterchangeSparseVolumeTextureFactory::BeginImportAsset_GameThread(
+	const FImportAssetObjectParams& Arguments
+)
+{
+	TRACE_CPUPROFILER_EVENT_SCOPE(UInterchangeSparseVolumeTextureFactory::BeginImportAsset_GameThread);
+
+	using namespace UE::Interchange;
+	using namespace UE::InterchangeSparseVolumeTextureFactory::Private;
+
+	FImportAssetResult ImportAssetResult;
+	USparseVolumeTexture* SparseVolumeTexture = nullptr;
+
+	auto HandleFailure = [this, &Arguments, &ImportAssetResult](const FText& Info)
+	{
+		UInterchangeResultError_Generic* Message = AddMessage<UInterchangeResultError_Generic>();
+		Message->SourceAssetName = Arguments.SourceData->GetFilename();
+		Message->DestinationAssetName = Arguments.AssetName;
+		Message->AssetType = GetFactoryClass();
+		Message->Text = FText::Format(
+			LOCTEXT(
+				"SparseVolumeTextureFactory_Failure",
+				"UInterchangeSparseVolumeTextureFactory: Could not create SparseVolumeTexture asset '{0}'. Reason: {1}"
+			),
+			FText::FromString(Arguments.AssetName),
+			Info
+		);
+		bSkipImport = true;
+		ImportAssetResult.bIsFactorySkipAsset = true;
+	};
+
+	UInterchangeSparseVolumeTextureFactoryNode* FactoryNode = Cast<UInterchangeSparseVolumeTextureFactoryNode>(Arguments.AssetNode);
+	if (!FactoryNode)
+	{
+		HandleFailure(
+			LOCTEXT("SparseVolumeTextureFactory_AssetNodeNull", "Asset node parameter is not an UInterchangeSparseVolumeTextureFactoryNode.")
+		);
+		return ImportAssetResult;
+	}
+
+	const UClass* ObjectClass = Arguments.AssetNode->GetObjectClass();
+	if (!ObjectClass || !ObjectClass->IsChildOf(USparseVolumeTexture::StaticClass()))
+	{
+		HandleFailure(
+			LOCTEXT("SparseVolumeTextureFactory_NodeClassMissmatch", "Asset node parameter class doesn't derive the USparseVolumeTexture class.")
+		);
+		return ImportAssetResult;
+	}
+
+	// Get the source translated nodes for this volume (could have multiple if this is an animated SVT and we have multiple frames)
+	TArray<const UInterchangeVolumeNode*> VolumeNodes = GetVolumeNodes(FactoryNode, Arguments.NodeContainer);
+	if (VolumeNodes.Num() == 0)
+	{
+		HandleFailure(LOCTEXT("SparseVolumeTextureFactory_NoVolumes", "Asset node parameter class doesn't target any UInterchangeVolumeNode."));
+		return ImportAssetResult;
+	}
+
+	// If we're trying to import an animation, the available grids must be identical for each volume frame
+	if (VolumeNodes.Num() > 1)
+	{
+		bool bConsistentGrids = CheckForGridConsistency(VolumeNodes, Arguments.NodeContainer);
+		if (!bConsistentGrids)
+		{
+			HandleFailure(LOCTEXT(
+				"SparseVolumeTextureFactory_InconsistentGrids",
+				"UInterchangeVolumeNodes provided for animated SparseVolumeTexture import don't have consistent volume grids"
+			));
+			return ImportAssetResult;
+		}
+	}
+
+	UObject* ExistingAsset = Arguments.ReimportObject;
+	if (!ExistingAsset)
+	{
+		FSoftObjectPath ReferenceObject;
+		if (FactoryNode->GetCustomReferenceObject(ReferenceObject))
+		{
+			ExistingAsset = ReferenceObject.TryLoad();
+		}
+
+		// If we're not reimporting this asset directly, and are instead importing/reimporting some other asset type and just
+		// ran into this asset as a "dependency" (e.g. importing USD and found a volume already in the folder, or reimporting
+		// a material and found the SVT already in the folder), then mark this as a "FactorySkipAsset" and just reuse it.
+		// With bSkipImport == true the further stages won't modify the asset, just carry it along
+		if (ExistingAsset && ExistingAsset->IsA(USparseVolumeTexture::StaticClass())
+			&& Arguments.Translator->GetSupportedAssetTypes() != EInterchangeTranslatorAssetType::Textures)
+		{
+			bSkipImport = true;
+			ImportAssetResult.bIsFactorySkipAsset = true;
+			ImportAssetResult.ImportedObject = ExistingAsset;
+		}
+	}
+
+	// Actually create a new asset
+	if (!bSkipImport)
+	{
+		// Pick between static and animated SVT
+		//
+		// Note that SVT assets cannot be reused for reimport: They have an internal InitState flag that only lets
+		// them ever receive their "source data" exactly once. The legacy factory reimports by calling NewObject with the
+		// same class, name and outer, which resets the SVT. We'll do this here too, we just have to be careful to
+		// pick the same class as the existing asset if we have one (during reimports): We will crash if we try
+		// overwriting an animated SVT asset with a static one, for example.
+		UClass* ClassToUse = ExistingAsset			  ? ExistingAsset->GetClass()
+							 : VolumeNodes.Num() == 1 ? UStaticSparseVolumeTexture::StaticClass()
+													  : UAnimatedSparseVolumeTexture::StaticClass();
+
+		// If we have a mismatch we should return now before we stomp the existing asset for no reason.
+		// Note that it's fine to try and stuff a single frame payload into an animated SVT: It will just be a one-frame-animation
+		if (ClassToUse == UStaticSparseVolumeTexture::StaticClass() && VolumeNodes.Num() != 1)
+		{
+			HandleFailure(LOCTEXT("SparseVolumeTextureFactory_Mismatch", "Cannot import multiple volume frames into a StaticSparseVolumeTexture."));
+			return ImportAssetResult;
+		}
+
+		SparseVolumeTexture = NewObject<USparseVolumeTexture>(Arguments.Parent, ClassToUse, *Arguments.AssetName, RF_Public | RF_Standalone);
+		if (!SparseVolumeTexture)
+		{
+			HandleFailure(LOCTEXT("SparseVolumeTextureFactory_TextureCreateFail", "SparseVolumeTexture creation failed."));
+			return ImportAssetResult;
+		}
+
+		ImportAssetResult.ImportedObject = SparseVolumeTexture;
+	}
+
+	return ImportAssetResult;
+}
+
+UInterchangeFactoryBase::FImportAssetResult UInterchangeSparseVolumeTextureFactory::ImportAsset_Async(const FImportAssetObjectParams& Arguments)
+{
+	TRACE_CPUPROFILER_EVENT_SCOPE(UInterchangeSparseVolumeTextureFactory::ImportAsset_Async);
+
+	using namespace UE::Interchange;
+	using namespace UE::Interchange::Volume;
+	using namespace UE::InterchangeSparseVolumeTextureFactory::Private;
+
+	FImportAssetResult ImportAssetResult;
+	ImportAssetResult.bIsFactorySkipAsset = bSkipImport;
+
+	auto HandleFailure = [this, &Arguments](const FText& Info)
+	{
+		UInterchangeResultError_Generic* Message = AddMessage<UInterchangeResultError_Generic>();
+		Message->SourceAssetName = Arguments.SourceData->GetFilename();
+		Message->DestinationAssetName = Arguments.AssetName;
+		Message->AssetType = GetFactoryClass();
+		Message->Text = Info;
+	};
+
+	if (!Arguments.AssetNode)
+	{
+		HandleFailure(
+			LOCTEXT("SparseVolumeTextureFactory_Async_AssetNodeNull", "UInterchangeSparseVolumeTextureFactory: Asset node parameter is null.")
+		);
+		return ImportAssetResult;
+	}
+
+	UInterchangeSparseVolumeTextureFactoryNode* FactoryNode = Cast<UInterchangeSparseVolumeTextureFactoryNode>(Arguments.AssetNode);
+	if (!FactoryNode)
+	{
+		HandleFailure(LOCTEXT(
+			"SparseVolumeTextureFactory_Async_NodeWrongClass",
+			"UInterchangeSparseVolumeTextureFactory: Asset node parameter is not a child of UInterchangeSparseVolumeTextureFactoryNode."
+		));
+		return ImportAssetResult;
+	}
+
+	UObject* CreatedAsset = nullptr;
+	FSoftObjectPath ReferenceObject;
+	if (Arguments.AssetNode->GetCustomReferenceObject(ReferenceObject))
+	{
+		CreatedAsset = ReferenceObject.TryLoad();
+	}
+
+	// Do not override an asset we skip
+	if (bSkipImport)
+	{
+		ImportAssetResult.ImportedObject = CreatedAsset;
+		return ImportAssetResult;
+	}
+
+	USparseVolumeTexture* SparseVolumeTexture = Cast<USparseVolumeTexture>(CreatedAsset);
+	if (!SparseVolumeTexture)
+	{
+		HandleFailure(LOCTEXT(
+			"SparseVolumeTextureFactory_Async_CannotCreateAsync",
+			"UInterchangeSparseVolumeTextureFactory: Could not create an USparseVolumeTexture asset."
+		));
+		return ImportAssetResult;
+	}
+
+	const IInterchangeVolumePayloadInterface* TranslatorInterface = Cast<IInterchangeVolumePayloadInterface>(Arguments.Translator);
+	if (!TranslatorInterface)
+	{
+		HandleFailure(LOCTEXT(
+			"SparseVolumeTextureFactory_Async_TranslatorPayloadInterface",
+			"UInterchangeSparseVolumeTextureFactory: The translator does not implement the Interchange volume payload interface."
+		));
+		return ImportAssetResult;
+	}
+
+	// Get the source translated nodes for this volume (could have multiple if this is an animated SVT and we have multiple frames)
+	TArray<const UInterchangeVolumeNode*> VolumeNodes = GetVolumeNodes(FactoryNode, Arguments.NodeContainer);
+	if (VolumeNodes.Num() == 0)
+	{
+		HandleFailure(
+			LOCTEXT("SparseVolumeTextureFactory_Async_NoTranslatedNode", "Asset node parameter class doesn't target any UInterchangeVolumeNode.")
+		);
+		return ImportAssetResult;
+	}
+
+	// Produce an FAssignmentInfo we'll reuse for all payload keys below
+	TOptional<FAssignmentInfo> AssignmentInfo = GetAssignmentInfo(FactoryNode, VolumeNodes, Arguments.NodeContainer);
+	if (!AssignmentInfo.IsSet())
+	{
+		HandleFailure(LOCTEXT(
+			"SparseVolumeTextureFactory_Async_NoAssignmentInfo",
+			"Failed to extract a valid grid assignment info from UInterchangeVolumeNode."
+		));
+		return ImportAssetResult;
+	}
+
+	// Get a single volume bounds we'll share across all payloads
+	FIntVector3 VolumeBoundsMin;
+	FIntVector3 VolumeBoundsMax;
+	ComputeExpandedVolumeBounds(VolumeNodes, Arguments.NodeContainer, VolumeBoundsMin, VolumeBoundsMax);
+
+	// Collect the payload keys from each volume node
+	bool bAtLeastOneValidPayload = false;
+	TArray<FVolumePayloadKey> PayloadKeys;
+	PayloadKeys.Reserve(VolumeNodes.Num());
+	for (const UInterchangeVolumeNode* VolumeNode : VolumeNodes)
+	{
+		// Emplace even if we fail so that we keep a one-to-one match with the VolumeNodes array
+		FVolumePayloadKey& Key = PayloadKeys.Emplace_GetRef();
+
+		FString FileName;
+		bool bSuccess = VolumeNode->GetCustomFileName(FileName);
+		if (!bSuccess || FileName.IsEmpty())
+		{
+			continue;
+		}
+
+		bAtLeastOneValidPayload = true;
+		Key.FileName = FileName;
+		Key.AssignmentInfo = AssignmentInfo.GetValue();
+		Key.VolumeBoundsMin = VolumeBoundsMin;
+	}
+	if (!bAtLeastOneValidPayload)
+	{
+		HandleFailure(
+			LOCTEXT("SparseVolumeTextureFactory_InvalidPayloadKey", "None of the translated UInterchangeVolumeGridNodes has a valid payload key.")
+		);
+		return ImportAssetResult;
+	}
+
+	// TODO: Handle this properly when implementing the USD translator's usage of the OpenVDB translator
+	AlternateTexturePath.Reset();
+
+	// Fetch the actual payloads in parallel (presumably they're all different files, but it should work even in the edge case
+	// that the same file shows up more than once somehow)
+	TArray<TOptional<FVolumePayloadData>> PayloadData;
+	PayloadData.SetNum(VolumeNodes.Num());
+	const int32 NumBatches = PayloadKeys.Num();
+	ParallelFor(
+		NumBatches,
+		[&PayloadKeys, &PayloadData, TranslatorInterface](int32 Index)
+		{
+			PayloadData[Index] = TranslatorInterface->GetVolumePayloadData(PayloadKeys[Index]);
+		}
+	);
+
+	bool bAtLeastOnePayloadIsValid = false;
+	for (const TOptional<FVolumePayloadData>& PayloadDataEntry : PayloadData)
+	{
+		if (PayloadDataEntry.IsSet())
+		{
+			bAtLeastOnePayloadIsValid = true;
+			break;
+		}
+	}
+	if (!bAtLeastOnePayloadIsValid)
+	{
+		HandleFailure(LOCTEXT(
+			"SparseVolumeTextureFactory_Async_CannotRetrievePayload",
+			"UInterchangeSparseVolumeTextureFactory: The factory couldn't retrieve any valid payload from the translator."
+		));
+		return ImportAssetResult;
+	}
+
+	// Check if the payload is valid for the Texture
+	bool bCanSetup = bAtLeastOnePayloadIsValid;
+	if (UStaticSparseVolumeTexture* StaticVolumeTexture = Cast<UStaticSparseVolumeTexture>(SparseVolumeTexture))
+	{
+		bCanSetup = bCanSetup && PayloadKeys.Num() == 1;
+	}
+	else if (UAnimatedSparseVolumeTexture* AnimatedVolumeTexture = Cast<UAnimatedSparseVolumeTexture>(SparseVolumeTexture))
+	{
+		bCanSetup = bCanSetup && PayloadKeys.Num() >= 1;
+	}
+	if (!bCanSetup)
+	{
+		HandleFailure(LOCTEXT(
+			"SparseVolumeTextureFactory_Async_CannotSetup",
+			"UInterchangeSparseVolumeTextureFactory: The factory cannot setup the created asset with the provided payload data."
+		));
+		return ImportAssetResult;
+	}
+
+#if WITH_EDITOR
+	// Create SourceFile hashes while we're still in an async thread (we'll move this into AssetImportData later)
+	{
+		SourceFiles.Empty(VolumeNodes.Num());
+
+		for (const UInterchangeVolumeNode* VolumeNode : VolumeNodes)
+		{
+			FString FileName;
+			VolumeNode->GetCustomFileName(FileName);
+
+			FAssetImportInfo::FSourceFile& SourceFile = SourceFiles.Emplace_GetRef();
+			SourceFile.RelativeFilename = FileName;
+		}
+
+		HashSourceFiles(SourceFiles);
+	}
+#endif	  // WITH_EDITOR
+
+	ProcessedPayloads = MoveTemp(PayloadData);
+	return ImportAssetResult;
+}
+
+UInterchangeFactoryBase::FImportAssetResult UInterchangeSparseVolumeTextureFactory::EndImportAsset_GameThread(
+	const FImportAssetObjectParams& Arguments
+)
+{
+	TRACE_CPUPROFILER_EVENT_SCOPE(UInterchangeSparseVolumeTextureFactory::ImportAsset_Async);
+
+	using namespace UE::InterchangeSparseVolumeTextureFactory::Private;
+
+	FImportAssetResult ImportAssetResult;
+	ImportAssetResult.bIsFactorySkipAsset = bSkipImport;
+
+	if (!Arguments.AssetNode)
+	{
+		return ImportAssetResult;
+	}
+
+	UObject* ExistingAsset = nullptr;
+	FSoftObjectPath ReferenceObject;
+	if (Arguments.AssetNode->GetCustomReferenceObject(ReferenceObject))
+	{
+		ExistingAsset = ReferenceObject.TryLoad();
+	}
+	if (!ExistingAsset)
+	{
+		return ImportAssetResult;
+	}
+
+	if (USparseVolumeTexture* SparseVolumeTexture = Cast<USparseVolumeTexture>(ExistingAsset))
+	{
+		ImportAssetResult.ImportedObject = SparseVolumeTexture;
+
+		if (!HasValidPayloads())
+		{
+			if (!Arguments.ReimportObject)
+			{
+				// Abandon asset
+				ImportAssetResult.ImportedObject = nullptr;
+				SparseVolumeTexture->RemoveFromRoot();
+				SparseVolumeTexture->MarkAsGarbage();
+				Arguments.AssetNode->SetCustomReferenceObject(nullptr);
+			}
+
+			Arguments.AssetNode->SetEnabled(false);
+			bSkipImport = true;
+			ImportAssetResult.bIsFactorySkipAsset = bSkipImport;
+		}
+	}
+
+	return ImportAssetResult;
+}
+
+void UInterchangeSparseVolumeTextureFactory::SetupObject_GameThread(const FSetupObjectParams& Arguments)
+{
+	TRACE_CPUPROFILER_EVENT_SCOPE(UInterchangeSparseVolumeTextureFactory::SetupObject_GameThread);
+
+	using namespace UE::InterchangeSparseVolumeTextureFactory::Private;
+
+	if (bSkipImport)
+	{
+		return;
+	}
+
+	UInterchangeSparseVolumeTextureFactoryNode* SparseVolumeTextureFactoryNode = Cast<UInterchangeSparseVolumeTextureFactoryNode>(
+		Arguments.FactoryNode
+	);
+
+	USparseVolumeTexture* SparseVolumeTexture = Cast<USparseVolumeTexture>(Arguments.ImportedObject);
+
+	if (SparseVolumeTexture && SparseVolumeTextureFactoryNode && HasValidPayloads())
+	{
+		SetupSparseVolumeTexture(SparseVolumeTexture, ProcessedPayloads, Arguments.bIsReimport);
+
+#if WITH_EDITOR
+		if (!Arguments.bIsReimport)
+		{
+			SparseVolumeTextureFactoryNode->ApplyAllCustomAttributeToObject(SparseVolumeTexture);
+		}
+		else
+		{
+			UInterchangeAssetImportData* InterchangeAssetImportData = nullptr;
+			{
+				UStreamableSparseVolumeTexture* Streamable = Cast<UStreamableSparseVolumeTexture>(SparseVolumeTexture);
+				InterchangeAssetImportData = Cast<UInterchangeAssetImportData>(Streamable->AssetImportData);
+			}
+
+			UInterchangeFactoryBaseNode* PreviousNode = nullptr;
+			if (InterchangeAssetImportData)
+			{
+				PreviousNode = InterchangeAssetImportData->GetStoredFactoryNode(InterchangeAssetImportData->NodeUniqueID);
+			}
+
+			UInterchangeFactoryBaseNode* CurrentNode = NewObject<UInterchangeSparseVolumeTextureFactoryNode>(GetTransientPackage());
+			UInterchangeBaseNode::CopyStorage(SparseVolumeTextureFactoryNode, CurrentNode);
+			CurrentNode->FillAllCustomAttributeFromObject(SparseVolumeTexture);
+
+			UE::Interchange::FFactoryCommon::ApplyReimportStrategyToAsset(
+				SparseVolumeTexture,
+				PreviousNode,
+				CurrentNode,
+				SparseVolumeTextureFactoryNode
+			);
+		}
+#endif	  // WITH_EDITOR
+	}
+	else if (!Arguments.bIsReimport)
+	{
+		// Import failed, so let's abandon the asset
+		SparseVolumeTexture->RemoveFromRoot();
+		SparseVolumeTexture->MarkAsGarbage();
+	}
+
+	Super::SetupObject_GameThread(Arguments);
+
+#if WITH_EDITOR
+	// Setup asset import data
+	//
+	// The Streamable derived class is the one that stores the AssetImportData, but all SVT classes that are actual assets
+	// derive it, so it should be safe to assume we have a Streamable here
+	UStreamableSparseVolumeTexture* Streamable = Cast<UStreamableSparseVolumeTexture>(SparseVolumeTexture);
+	if (Streamable && Arguments.SourceData && HasValidPayloads())
+	{
+		// We must call the Update of the asset source file in the main thread because UAssetImportData::Update execute some delegate we do not
+		// control
+		UE::Interchange::FFactoryCommon::FSetImportAssetDataParameters SetImportAssetDataParameters(
+			Streamable,
+			Streamable->AssetImportData,
+			Arguments.SourceData,
+			Arguments.NodeUniqueID,
+			Arguments.NodeContainer,
+			Arguments.OriginalPipelines,
+			Arguments.Translator
+		);
+		SetImportAssetDataParameters.SourceFiles = MoveTemp(SourceFiles);
+
+		Streamable->AssetImportData = UE::Interchange::FFactoryCommon::SetImportAssetData(SetImportAssetDataParameters);
+
+		if (AlternateTexturePath.IsSet())
+		{
+			// A file different from the source data has been used to create the texture.
+			// Set this file's path as the first file name
+			Streamable->AssetImportData->AddFileName(AlternateTexturePath.GetValue(), 0);
+		}
+	}
+#endif	  // WITH_EDITOR
+}
+
+bool UInterchangeSparseVolumeTextureFactory::GetSourceFilenames(const UObject* Object, TArray<FString>& OutSourceFilenames) const
+{
+#if WITH_EDITOR
+	if (const UStreamableSparseVolumeTexture* Streamable = Cast<UStreamableSparseVolumeTexture>(Object))
+	{
+		return UE::Interchange::FFactoryCommon::GetSourceFilenames(Streamable->AssetImportData.Get(), OutSourceFilenames);
+	}
+#endif
+
+	return false;
+}
+
+bool UInterchangeSparseVolumeTextureFactory::SetSourceFilename(const UObject* Object, const FString& SourceFilename, int32 SourceIndex) const
+{
+#if WITH_EDITOR
+	if (const UStreamableSparseVolumeTexture* Streamable = Cast<UStreamableSparseVolumeTexture>(Object))
+	{
+		return UE::Interchange::FFactoryCommon::SetSourceFilename(Streamable->AssetImportData.Get(), SourceFilename, SourceIndex);
+	}
+#endif
+
+	return false;
+}
+
+void UInterchangeSparseVolumeTextureFactory::BackupSourceData(const UObject* Object) const
+{
+#if WITH_EDITOR
+	if (const UStreamableSparseVolumeTexture* Streamable = Cast<UStreamableSparseVolumeTexture>(Object))
+	{
+		UE::Interchange::FFactoryCommon::BackupSourceData(Streamable->AssetImportData.Get());
+	}
+#endif
+}
+
+void UInterchangeSparseVolumeTextureFactory::ReinstateSourceData(const UObject* Object) const
+{
+#if WITH_EDITOR
+	if (const UStreamableSparseVolumeTexture* Streamable = Cast<UStreamableSparseVolumeTexture>(Object))
+	{
+		UE::Interchange::FFactoryCommon::ReinstateSourceData(Streamable->AssetImportData.Get());
+	}
+#endif
+}
+
+void UInterchangeSparseVolumeTextureFactory::ClearBackupSourceData(const UObject* Object) const
+{
+#if WITH_EDITOR
+	if (const UStreamableSparseVolumeTexture* Streamable = Cast<UStreamableSparseVolumeTexture>(Object))
+	{
+		UE::Interchange::FFactoryCommon::ClearBackupSourceData(Streamable->AssetImportData.Get());
+	}
+#endif
+}
+
+bool UInterchangeSparseVolumeTextureFactory::HasValidPayloads() const
+{
+	for (const TOptional<UE::Interchange::FVolumePayloadData>& Payload : ProcessedPayloads)
+	{
+		if (Payload.IsSet())
+		{
+			return true;
+		}
+	}
+
+	return false;
+}
+
+#undef LOCTEXT_NAMESPACE
diff --git a/Engine/Plugins/Interchange/Runtime/Source/Import/Public/OpenVDB/InterchangeOpenVDBTranslator.h b/Engine/Plugins/Interchange/Runtime/Source/Import/Public/OpenVDB/InterchangeOpenVDBTranslator.h
new file mode 100644
--- /dev/null
+++ b/Engine/Plugins/Interchange/Runtime/Source/Import/Public/OpenVDB/InterchangeOpenVDBTranslator.h
@@ -0,0 +1,70 @@
+// Copyright Epic Games, Inc. All Rights Reserved.
+
+#pragma once
+
+#include "InterchangeTranslatorBase.h"
+#include "Volume/InterchangeVolumePayloadInterface.h"
+
+#include "UObject/ObjectMacros.h"
+
+#include "InterchangeOpenVDBTranslator.generated.h"
+
+namespace UE::Interchange
+{
+	struct FVolumePayloadData;
+	struct FVolumePayloadKey;
+}
+namespace UE::InterchangeOpenVDBTranslator::Private
+{
+	class UInterchangeOpenVDBTranslatorImpl;
+}
+
+UCLASS(BlueprintType, editinlinenew, MinimalAPI)
+class UInterchangeOpenVDBTranslatorSettings : public UInterchangeTranslatorSettings
+{
+	GENERATED_BODY()
+
+public:
+	/**
+	 * Whether we should search the same folder of the translated file for other .vdb files whose file names differ from it
+	 * only by a numbered suffix (e.g. "Tornado_0001.vdb, Tornado_0002.vdb, etc.), and translate those too.
+	 *
+	 * This is useful if you intend on importing an animated SparseVolumeTexture.
+	 */
+	UPROPERTY(EditAnywhere, Category = "OpenVDB Translator")
+	bool bTranslateAdjacentNumberedFiles;
+
+public:
+	UInterchangeOpenVDBTranslatorSettings();
+};
+
+UCLASS(BlueprintType)
+class UInterchangeOpenVDBTranslator
+	: public UInterchangeTranslatorBase
+	, public IInterchangeVolumePayloadInterface
+{
+	GENERATED_BODY()
+
+public:
+	UInterchangeOpenVDBTranslator();
+
+	/** Begin UInterchangeTranslatorBase interface */
+	virtual EInterchangeTranslatorType GetTranslatorType() const override;
+	virtual EInterchangeTranslatorAssetType GetSupportedAssetTypes() const override;
+	virtual TArray<FString> GetSupportedFormats() const override;
+	virtual bool Translate(UInterchangeBaseNodeContainer& BaseNodeContainer) const override;
+	virtual void ReleaseSource() override;
+	virtual UInterchangeTranslatorSettings* GetSettings() const override;
+	virtual void SetSettings(const UInterchangeTranslatorSettings* InterchangeTranslatorSettings) override;
+	/** End UInterchangeTranslatorBase interface */
+
+	/** Begin Interchange payload interfaces */
+	virtual TOptional<UE::Interchange::FVolumePayloadData> GetVolumePayloadData(const UE::Interchange::FVolumePayloadKey& PayloadKey) const override;
+	/** End Interchange payload interfaces */
+
+private:
+	TUniquePtr<UE::InterchangeOpenVDBTranslator::Private::UInterchangeOpenVDBTranslatorImpl> Impl;
+
+	UPROPERTY(Transient, DuplicateTransient)
+	mutable TObjectPtr<UInterchangeOpenVDBTranslatorSettings> TranslatorSettings = nullptr;
+};
diff --git a/Engine/Plugins/Interchange/Runtime/Source/Import/Public/Volume/InterchangeSparseVolumeTextureFactory.h b/Engine/Plugins/Interchange/Runtime/Source/Import/Public/Volume/InterchangeSparseVolumeTextureFactory.h
new file mode 100644
--- /dev/null
+++ b/Engine/Plugins/Interchange/Runtime/Source/Import/Public/Volume/InterchangeSparseVolumeTextureFactory.h
@@ -0,0 +1,63 @@
+// Copyright Epic Games, Inc. All Rights Reserved.
+
+#pragma once
+
+#include "InterchangeFactoryBase.h"
+
+#include "CoreMinimal.h"
+#include "EditorFramework/AssetImportData.h"
+#include "Serialization/EditorBulkData.h"
+#include "UObject/Object.h"
+#include "UObject/ObjectMacros.h"
+
+#include "InterchangeSparseVolumeTextureFactory.generated.h"
+
+class UInterchangeSparseVolumeTextureFactoryNode;
+namespace UE::Interchange
+{
+	struct FVolumePayloadData;
+	struct FVolumePayloadKey;
+}
+
+UCLASS(BlueprintType)
+class INTERCHANGEIMPORT_API UInterchangeSparseVolumeTextureFactory : public UInterchangeFactoryBase
+{
+	GENERATED_BODY()
+
+public:
+	// Begin UInterchangeFactoryBase interface
+	virtual UClass* GetFactoryClass() const override;
+	virtual EInterchangeFactoryAssetType GetFactoryAssetType() override
+	{
+		return EInterchangeFactoryAssetType::Textures;
+	}
+	virtual FImportAssetResult BeginImportAsset_GameThread(const FImportAssetObjectParams& Arguments) override;
+	virtual FImportAssetResult ImportAsset_Async(const FImportAssetObjectParams& Arguments) override;
+	virtual FImportAssetResult EndImportAsset_GameThread(const FImportAssetObjectParams& Arguments) override;
+	virtual void SetupObject_GameThread(const FSetupObjectParams& Arguments) override;
+	virtual bool GetSourceFilenames(const UObject* Object, TArray<FString>& OutSourceFilenames) const override;
+	virtual bool SetSourceFilename(const UObject* Object, const FString& SourceFilename, int32 SourceIndex) const override;
+	virtual void BackupSourceData(const UObject* Object) const override;
+	virtual void ReinstateSourceData(const UObject* Object) const override;
+	virtual void ClearBackupSourceData(const UObject* Object) const override;
+	// End UInterchangeFactoryBase interface
+
+private:
+	bool HasValidPayloads() const;
+
+private:
+	TOptional<FString> AlternateTexturePath;
+
+#if WITH_EDITORONLY_DATA
+	// The data for the source files will be collecetd here during the import during an async task,
+	// and later be added to AssetImportData
+	TArray<FAssetImportInfo::FSourceFile> SourceFiles;
+#endif
+
+	// Where we retain the processed payloads between the different factory interface function calls
+	TArray<TOptional<UE::Interchange::FVolumePayloadData>> ProcessedPayloads;
+
+	// We turn this to true when we run into an existing asset that we shouldn't overwrite/modify (e.g. when reimporting
+	// some other asset type, like reimporting a material and finding an existing SVT that it wants to use)
+	bool bSkipImport = false;
+};
diff --git a/Engine/Plugins/Interchange/Runtime/Source/Import/Public/Volume/InterchangeVolumePayloadData.h b/Engine/Plugins/Interchange/Runtime/Source/Import/Public/Volume/InterchangeVolumePayloadData.h
new file mode 100644
--- /dev/null
+++ b/Engine/Plugins/Interchange/Runtime/Source/Import/Public/Volume/InterchangeVolumePayloadData.h
@@ -0,0 +1,25 @@
+// Copyright Epic Games, Inc. All Rights Reserved.
+
+#pragma once
+
+#include "CoreMinimal.h"
+#include "SparseVolumeTexture/SparseVolumeTextureData.h"
+
+namespace UE::Interchange
+{
+	// Contains everything we need to use as source data for a SparseVolumeTexture asset
+	struct INTERCHANGEIMPORT_API FVolumePayloadData
+	{
+		virtual ~FVolumePayloadData() = default;
+		FVolumePayloadData() = default;
+		FVolumePayloadData(FVolumePayloadData&&) = default;
+		FVolumePayloadData& operator=(FVolumePayloadData&&) = default;
+
+		FVolumePayloadData(const FVolumePayloadData&) = delete;
+		FVolumePayloadData& operator=(const FVolumePayloadData&) = delete;
+
+	public:
+		UE::SVT::FTextureData TextureData;
+		FTransform Transform;
+	};
+}	 // namespace UE::Interchange
diff --git a/Engine/Plugins/Interchange/Runtime/Source/Import/Public/Volume/InterchangeVolumePayloadInterface.h b/Engine/Plugins/Interchange/Runtime/Source/Import/Public/Volume/InterchangeVolumePayloadInterface.h
new file mode 100644
--- /dev/null
+++ b/Engine/Plugins/Interchange/Runtime/Source/Import/Public/Volume/InterchangeVolumePayloadInterface.h
@@ -0,0 +1,33 @@
+// Copyright Epic Games, Inc. All Rights Reserved.
+
+#pragma once
+
+#include "Volume/InterchangeVolumePayloadData.h"
+#include "Volume/InterchangeVolumePayloadKey.h"
+
+#include "CoreMinimal.h"
+#include "UObject/Interface.h"
+#include "UObject/ObjectMacros.h"
+
+#include "InterchangeVolumePayloadInterface.generated.h"
+
+UINTERFACE()
+class INTERCHANGEIMPORT_API UInterchangeVolumePayloadInterface : public UInterface
+{
+	GENERATED_BODY()
+};
+
+/**
+ * Payload interface for volume data (e.g. OpenVDB)
+ */
+class INTERCHANGEIMPORT_API IInterchangeVolumePayloadInterface
+{
+	GENERATED_BODY()
+
+public:
+	virtual TOptional<UE::Interchange::FVolumePayloadData> GetVolumePayloadData(const UE::Interchange::FVolumePayloadKey& PayloadKey) const
+	{
+		return {};
+	}
+};
+
diff --git a/Engine/Plugins/Interchange/Runtime/Source/Import/Public/Volume/InterchangeVolumePayloadKey.h b/Engine/Plugins/Interchange/Runtime/Source/Import/Public/Volume/InterchangeVolumePayloadKey.h
new file mode 100644
--- /dev/null
+++ b/Engine/Plugins/Interchange/Runtime/Source/Import/Public/Volume/InterchangeVolumePayloadKey.h
@@ -0,0 +1,17 @@
+// Copyright Epic Games, Inc. All Rights Reserved.
+
+#pragma once
+
+#include "CoreMinimal.h"
+#include "Volume/InterchangeVolumeDefinitions.h"
+
+namespace UE::Interchange
+{
+	// Contains everything we need to query a factory for some FVolumePayloadData
+	struct FVolumePayloadKey
+	{
+		FString FileName;
+		Volume::FAssignmentInfo AssignmentInfo;
+		FIntVector3 VolumeBoundsMin;
+	};
+}
diff --git a/Engine/Plugins/Interchange/Runtime/Source/Nodes/Private/InterchangeVolumeNode.cpp b/Engine/Plugins/Interchange/Runtime/Source/Nodes/Private/InterchangeVolumeNode.cpp
new file mode 100644
--- /dev/null
+++ b/Engine/Plugins/Interchange/Runtime/Source/Nodes/Private/InterchangeVolumeNode.cpp
@@ -0,0 +1,122 @@
+// Copyright Epic Games, Inc. All Rights Reserved.
+
+#include "InterchangeVolumeNode.h"
+
+namespace UE::Interchange
+{
+	const FAttributeKey& FInterchangeVolumeNodeStaticData::GetGridDependenciesBaseKey()
+	{
+		static FAttributeKey AttributeKey(TEXT("__GridDependencies__"));
+		return AttributeKey;
+	}
+}
+
+UInterchangeVolumeNode::UInterchangeVolumeNode()
+{
+	GridDependencies.Initialize(Attributes, UE::Interchange::FInterchangeVolumeNodeStaticData::GetGridDependenciesBaseKey().ToString());
+}
+
+FString UInterchangeVolumeNode::GetTypeName() const
+{
+	return TEXT("Volume");
+}
+
+bool UInterchangeVolumeNode::GetCustomFileName(FString& AttributeValue) const
+{
+	IMPLEMENT_NODE_ATTRIBUTE_GETTER(FileName, FString);
+}
+
+bool UInterchangeVolumeNode::SetCustomFileName(const FString& AttributeValue)
+{
+	IMPLEMENT_NODE_ATTRIBUTE_SETTER_NODELEGATE(FileName, FString);
+}
+
+int32 UInterchangeVolumeNode::GetGridDependeciesCount() const
+{
+	return GridDependencies.GetCount();
+}
+
+void UInterchangeVolumeNode::GetGridDependencies(TArray<FString>& OutDependencies) const
+{
+	GridDependencies.GetItems(OutDependencies);
+}
+
+void UInterchangeVolumeNode::GetGridDependency(const int32 Index, FString& OutDependency) const
+{
+	GridDependencies.GetItem(Index, OutDependency);
+}
+
+bool UInterchangeVolumeNode::AddGridDependencyUid(const FString& DependencyUid)
+{
+	return GridDependencies.AddItem(DependencyUid);
+}
+
+bool UInterchangeVolumeNode::RemoveGridDependencyUid(const FString& DependencyUid)
+{
+	return GridDependencies.RemoveItem(DependencyUid);
+}
+
+FString UInterchangeVolumeGridNode::GetTypeName() const
+{
+	return TEXT("VolumeGrid");
+}
+
+bool UInterchangeVolumeGridNode::GetCustomElementType(EVolumeGridElementType& AttributeValue) const
+{
+	IMPLEMENT_NODE_ATTRIBUTE_GETTER(ElementType, EVolumeGridElementType);
+}
+
+bool UInterchangeVolumeGridNode::SetCustomElementType(const EVolumeGridElementType& AttributeValue)
+{
+	IMPLEMENT_NODE_ATTRIBUTE_SETTER_NODELEGATE(ElementType, EVolumeGridElementType);
+}
+
+bool UInterchangeVolumeGridNode::GetCustomNumComponents(int32& AttributeValue) const
+{
+	IMPLEMENT_NODE_ATTRIBUTE_GETTER(NumComponents, int32);
+}
+
+bool UInterchangeVolumeGridNode::SetCustomNumComponents(const int32& AttributeValue)
+{
+	IMPLEMENT_NODE_ATTRIBUTE_SETTER_NODELEGATE(NumComponents, int32);
+}
+
+bool UInterchangeVolumeGridNode::GetCustomGridTransform(FTransform& AttributeValue) const
+{
+	IMPLEMENT_NODE_ATTRIBUTE_GETTER(GridTransform, FTransform);
+}
+
+bool UInterchangeVolumeGridNode::SetCustomGridTransform(const FTransform& AttributeValue)
+{
+	IMPLEMENT_NODE_ATTRIBUTE_SETTER_NODELEGATE(GridTransform, FTransform);
+}
+
+bool UInterchangeVolumeGridNode::GetCustomGridActiveAABBMin(FIntVector& AttributeValue) const
+{
+	IMPLEMENT_NODE_ATTRIBUTE_GETTER(GridActiveAABBMin, FIntVector);
+}
+
+bool UInterchangeVolumeGridNode::SetCustomGridActiveAABBMin(const FIntVector& AttributeValue)
+{
+	IMPLEMENT_NODE_ATTRIBUTE_SETTER_NODELEGATE(GridActiveAABBMin, FIntVector);
+}
+
+bool UInterchangeVolumeGridNode::GetCustomGridActiveAABBMax(FIntVector& AttributeValue) const
+{
+	IMPLEMENT_NODE_ATTRIBUTE_GETTER(GridActiveAABBMin, FIntVector);
+}
+
+bool UInterchangeVolumeGridNode::SetCustomGridActiveAABBMax(const FIntVector& AttributeValue)
+{
+	IMPLEMENT_NODE_ATTRIBUTE_SETTER_NODELEGATE(GridActiveAABBMax, FIntVector);
+}
+
+bool UInterchangeVolumeGridNode::GetCustomGridActiveDimensions(FIntVector& AttributeValue) const
+{
+	IMPLEMENT_NODE_ATTRIBUTE_GETTER(GridActiveDim, FIntVector);
+}
+
+bool UInterchangeVolumeGridNode::SetCustomGridActiveDimensions(const FIntVector& AttributeValue)
+{
+	IMPLEMENT_NODE_ATTRIBUTE_SETTER_NODELEGATE(GridActiveDim, FIntVector);
+}
diff --git a/Engine/Plugins/Interchange/Runtime/Source/Nodes/Public/InterchangeVolumeNode.h b/Engine/Plugins/Interchange/Runtime/Source/Nodes/Public/InterchangeVolumeNode.h
new file mode 100644
--- /dev/null
+++ b/Engine/Plugins/Interchange/Runtime/Source/Nodes/Public/InterchangeVolumeNode.h
@@ -0,0 +1,156 @@
+// Copyright Epic Games, Inc. All Rights Reserved.
+
+#pragma once
+
+#include "CoreMinimal.h"
+#include "Nodes/InterchangeBaseNode.h"
+
+#include "InterchangeVolumeNode.generated.h"
+
+namespace UE::Interchange
+{
+	struct INTERCHANGENODES_API FInterchangeVolumeNodeStaticData : public FBaseNodeStaticData
+	{
+		static const FAttributeKey& GetGridDependenciesBaseKey();
+	};
+}
+
+/**
+ * Represents a file that contains volume data in the form of (potentially multiple) grids,
+ * which are represented as UInterchangeVolumeGridNode dependencies
+ */
+UCLASS(BlueprintType)
+class INTERCHANGENODES_API UInterchangeVolumeNode : public UInterchangeBaseNode
+{
+	GENERATED_BODY()
+
+public:
+	UInterchangeVolumeNode();
+
+	virtual FString GetTypeName() const override;
+
+public:
+	/**
+	 * Gets the filename of the file with volume data (e.g. "C:/MyFolder/File.vdb").
+	 *
+	 * This is stored on the translated node as well as the source data, as a volume import may discover other
+	 * additional files in order to handle animated volume imports.
+	 */
+	UFUNCTION(BlueprintCallable, Category = "Interchange | Node | Volume")
+	bool GetCustomFileName(FString& AttributeValue) const;
+
+	/**
+	 * Sets the filename of the file with volume data (e.g. "C:/MyFolder/File.vdb").
+	 *
+	 * This is stored on the translated node as well as the source data, as a volume import may discover other
+	 * additional files in order to handle animated volume imports.
+	 */
+	UFUNCTION(BlueprintCallable, Category = "Interchange | Node | Volume")
+	bool SetCustomFileName(const FString& AttributeValue);
+
+	/**
+	 * Gets the number of UInterchangeVolumeGridNodes declared as dependencies by this volume node
+	 * (in other words, returns the number of volume grids contained in this file)
+	 */
+	UFUNCTION(BlueprintCallable, Category = "Interchange | Node | Volume")
+	int32 GetGridDependeciesCount() const;
+
+	/** Gets the Node IDs of UInterchangeVolumeGridNodes declared as dependencies by this volume node */
+	UFUNCTION(BlueprintCallable, Category = "Interchange | Node | Volume")
+	void GetGridDependencies(TArray<FString>& OutDependencies) const;
+
+	/** Gets the Node ID of an UInterchangeVolumeGridNode dependency */
+	UFUNCTION(BlueprintCallable, Category = "Interchange | Node | Volume")
+	void GetGridDependency(const int32 Index, FString& OutDependency) const;
+
+	/** Sets the Node ID of an UInterchangeVolumeGridNode dependency */
+	UFUNCTION(BlueprintCallable, Category = "Interchange | Node | Volume")
+	bool AddGridDependencyUid(const FString& DependencyUid);
+
+	/** Removes the Node ID of an UInterchangeVolumeGridNode dependency */
+	UFUNCTION(BlueprintCallable, Category = "Interchange | Node | Volume")
+	bool RemoveGridDependencyUid(const FString& DependencyUid);
+
+private:
+	IMPLEMENT_NODE_ATTRIBUTE_KEY(FileName);
+
+	UE::Interchange::TArrayAttributeHelper<FString> GridDependencies;
+};
+
+UENUM(Blueprintable)
+enum class EVolumeGridElementType : uint8
+{
+	Unknown,
+	Half,
+	Float,
+	Double
+};
+
+/**
+ * Represents a single grid (essentially a 3d texture) within a volumetric file
+ */
+UCLASS(BlueprintType)
+class INTERCHANGENODES_API UInterchangeVolumeGridNode : public UInterchangeBaseNode
+{
+	GENERATED_BODY()
+
+public:
+	virtual FString GetTypeName() const override;
+
+public:
+	/** Gets the datatype of the value of each voxel in the grid */
+	UFUNCTION(BlueprintCallable, Category = "Interchange | Node | Texture")
+	bool GetCustomElementType(EVolumeGridElementType& AttributeValue) const;
+
+	/** Sets the datatype of the value of each voxel in the grid */
+	UFUNCTION(BlueprintCallable, Category = "Interchange | Node | Volume")
+	bool SetCustomElementType(const EVolumeGridElementType& AttributeValue);
+
+	/** Gets the number of components of each voxel of the grid (e.g. 3 components for a vector grid) */
+	UFUNCTION(BlueprintCallable, Category = "Interchange | Node | Volume")
+	bool GetCustomNumComponents(int32& NumComponents) const;
+
+	/** Sets the number of components of each voxel of the grid (e.g. 3 components for a vector grid) */
+	UFUNCTION(BlueprintCallable, Category = "Interchange | Node | Volume")
+	bool SetCustomNumComponents(const int32& NumComponents);
+
+	/** Gets the grid transform contained in the volume file */
+	UFUNCTION(BlueprintCallable, Category = "Interchange | Node | Volume")
+	bool GetCustomGridTransform(FTransform& AttributeValue) const;
+
+	/** Sets the grid transform contained in the volume file */
+	UFUNCTION(BlueprintCallable, Category = "Interchange | Node | Volume")
+	bool SetCustomGridTransform(const FTransform& AttributeValue);
+
+	/** Gets the min X, Y and Z of the grid's active axis-aligned bounding box. */
+	UFUNCTION(BlueprintCallable, Category = "Interchange | Node | Volume")
+	bool GetCustomGridActiveAABBMin(FIntVector& AttributeValue) const;
+
+	/** Sets the min X, Y and Z of the grid's active axis-aligned bounding box. */
+	UFUNCTION(BlueprintCallable, Category = "Interchange | Node | Volume")
+	bool SetCustomGridActiveAABBMin(const FIntVector& AttributeValue);
+
+	/** Gets the max X, Y and Z of the grid's active axis-aligned bounding box. */
+	UFUNCTION(BlueprintCallable, Category = "Interchange | Node | Volume")
+	bool GetCustomGridActiveAABBMax(FIntVector& AttributeValue) const;
+
+	/** Sets the max X, Y and Z of the grid's active axis-aligned bounding box. */
+	UFUNCTION(BlueprintCallable, Category = "Interchange | Node | Volume")
+	bool SetCustomGridActiveAABBMax(const FIntVector& AttributeValue);
+
+	/** Gets the size of the grid, in voxels. */
+	UFUNCTION(BlueprintCallable, Category = "Interchange | Node | Volume")
+	bool GetCustomGridActiveDimensions(FIntVector& AttributeValue) const;
+
+	/** Sets the size of the grid, in voxels. */
+	UFUNCTION(BlueprintCallable, Category = "Interchange | Node | Volume")
+	bool SetCustomGridActiveDimensions(const FIntVector& AttributeValue);
+
+private:
+	IMPLEMENT_NODE_ATTRIBUTE_KEY(ElementType);
+	IMPLEMENT_NODE_ATTRIBUTE_KEY(NumComponents);
+	IMPLEMENT_NODE_ATTRIBUTE_KEY(GridTransform);
+	IMPLEMENT_NODE_ATTRIBUTE_KEY(GridActiveAABBMin);
+	IMPLEMENT_NODE_ATTRIBUTE_KEY(GridActiveAABBMax);
+	IMPLEMENT_NODE_ATTRIBUTE_KEY(GridActiveDim);
+};
diff --git a/Engine/Plugins/Interchange/Runtime/Source/Pipelines/Private/InterchangeGenericMaterialPipeline.cpp b/Engine/Plugins/Interchange/Runtime/Source/Pipelines/Private/InterchangeGenericMaterialPipeline.cpp
--- a/Engine/Plugins/Interchange/Runtime/Source/Pipelines/Private/InterchangeGenericMaterialPipeline.cpp
+++ b/Engine/Plugins/Interchange/Runtime/Source/Pipelines/Private/InterchangeGenericMaterialPipeline.cpp
@@ -1,4 +1,4 @@
-// Copyright Epic Games, Inc. All Rights Reserved. 
+// Copyright Epic Games, Inc. All Rights Reserved.
 
 #include "InterchangeGenericMaterialPipeline.h"
 
@@ -6,24 +6,26 @@
 #include "AssetRegistry/AssetRegistryModule.h"
 #include "CoreMinimal.h"
 #include "InterchangeGenericTexturePipeline.h"
+#include "InterchangeManager.h"
 #include "InterchangeMaterialDefinitions.h"
 #include "InterchangeMaterialFactoryNode.h"
+#include "InterchangeMaterialInstanceNode.h"
+#include "InterchangePipelineLog.h"
 #include "InterchangeShaderGraphNode.h"
-#include "InterchangeTexture2DNode.h"
+#include "InterchangeSourceData.h"
+#include "InterchangeSparseVolumeTexturePipeline.h"
 #include "InterchangeTexture2DArrayNode.h"
+#include "InterchangeTexture2DNode.h"
 #include "InterchangeTextureBlurNode.h"
 #include "InterchangeTextureCubeNode.h"
 #include "InterchangeTextureFactoryNode.h"
 #include "InterchangeTextureNode.h"
-#include "InterchangePipelineLog.h"
-#include "InterchangeSourceData.h"
 #include "Materials/Material.h"
-#include "Misc/PackageName.h"
-#include "Nodes/InterchangeBaseNode.h"
-#include "Nodes/InterchangeUserDefinedAttribute.h"
-#include "InterchangeMaterialInstanceNode.h"
 #include "Materials/MaterialExpressionAdd.h"
 #include "Materials/MaterialExpressionComponentMask.h"
+#include "Materials/MaterialExpressionConstant.h"
+#include "Materials/MaterialExpressionConstant3Vector.h"
+#include "Materials/MaterialExpressionCosine.h"
 #include "Materials/MaterialExpressionDivide.h"
 #include "Materials/MaterialExpressionFresnel.h"
 #include "Materials/MaterialExpressionLinearInterpolate.h"
@@ -34,11 +36,8 @@
 #include "Materials/MaterialExpressionRotateAboutAxis.h"
 #include "Materials/MaterialExpressionRotator.h"
 #include "Materials/MaterialExpressionScalarParameter.h"
-#include "Materials/MaterialExpressionStaticBoolParameter.h"
-#include "Materials/MaterialExpressionConstant.h"
-#include "Materials/MaterialExpressionConstant3Vector.h"
-#include "Materials/MaterialExpressionCosine.h"
 #include "Materials/MaterialExpressionSine.h"
+#include "Materials/MaterialExpressionStaticBoolParameter.h"
 #include "Materials/MaterialExpressionSwitch.h"
 #include "Materials/MaterialExpressionTangent.h"
 #include "Materials/MaterialExpressionTextureCoordinate.h"
@@ -48,8 +47,8 @@
 #include "Materials/MaterialExpressionTextureSampleParameter2DArray.h"
 #include "Materials/MaterialExpressionTextureSampleParameterCube.h"
 #include "Materials/MaterialExpressionTime.h"
-#include "Materials/MaterialExpressionTransformPosition.h"
 #include "Materials/MaterialExpressionTransform.h"
+#include "Materials/MaterialExpressionTransformPosition.h"
 #include "Materials/MaterialExpressionVectorNoise.h"
 #include "Materials/MaterialExpressionVectorParameter.h"
 #include "Materials/MaterialFunction.h"
@@ -58,15 +57,17 @@
 #include "MaterialX/MaterialExpressions/MaterialExpressionSwizzle.h"
 #include "MaterialX/MaterialExpressions/MaterialExpressionTextureSampleParameterBlur.h"
 #include "Misc/CoreMisc.h"
+#include "Misc/PackageName.h"
 #include "Misc/Paths.h"
 #include "Nodes/InterchangeBaseNode.h"
+#include "Nodes/InterchangeBaseNode.h"
 #include "Nodes/InterchangeBaseNodeContainer.h"
 #include "Nodes/InterchangeSourceNode.h"
+#include "Nodes/InterchangeUserDefinedAttribute.h"
 #include "Templates/Function.h"
 #include "Templates/SubclassOf.h"
 #include "UObject/Object.h"
 #include "UObject/ObjectMacros.h"
-#include "InterchangeManager.h"
 
 #if UE_BUILD_DEBUG
 #include "HAL/PlatformFileManager.h"
@@ -593,6 +594,7 @@ HashDebugData->PopNodeAddressesToLastPopIndex();\
 UInterchangeGenericMaterialPipeline::UInterchangeGenericMaterialPipeline()
 {
 	TexturePipeline = CreateDefaultSubobject<UInterchangeGenericTexturePipeline>("TexturePipeline");
+	SparseVolumeTexturePipeline = CreateDefaultSubobject<UInterchangeSparseVolumeTexturePipeline>("SparseVolumeTexturePipeline");
 }
 
 FString UInterchangeGenericMaterialPipeline::GetPipelineCategory(UClass* AssetClass)
@@ -607,6 +609,11 @@ void UInterchangeGenericMaterialPipeline::PreDialogCleanup(const FName PipelineS
 		TexturePipeline->PreDialogCleanup(PipelineStackName);
 	}
 
+	if (SparseVolumeTexturePipeline)
+	{
+		SparseVolumeTexturePipeline->PreDialogCleanup(PipelineStackName);
+	}
+
 	//Save only pipeline if we are a stand alone pipeline (not a sub object of another pipeline)
 	if (IsStandAlonePipeline())
 	{
@@ -621,6 +628,11 @@ bool UInterchangeGenericMaterialPipeline::IsSettingsAreValid(TOptional<FText>& O
 		return false;
 	}
 
+	if (SparseVolumeTexturePipeline && !SparseVolumeTexturePipeline->IsSettingsAreValid(OutInvalidReason))
+	{
+		return false;
+	}
+
 	return Super::IsSettingsAreValid(OutInvalidReason);
 }
 
@@ -632,6 +644,10 @@ void UInterchangeGenericMaterialPipeline::AdjustSettingsForContext(const FInterc
 	{
 		TexturePipeline->AdjustSettingsForContext(ContextParams);
 	}
+	if (SparseVolumeTexturePipeline)
+	{
+		SparseVolumeTexturePipeline->AdjustSettingsForContext(ContextParams);
+	}
 #if WITH_EDITOR
 	TArray<FString> HideCategories;
 	bool bIsObjectAMaterial = !ContextParams.ReimportAsset ? false : ContextParams.ReimportAsset->IsA(UMaterialInterface::StaticClass());
@@ -703,6 +719,10 @@ void UInterchangeGenericMaterialPipeline::FilterPropertiesFromTranslatedData(UIn
 	{
 		TexturePipeline->FilterPropertiesFromTranslatedData(InBaseNodeContainer);
 	}
+	if(SparseVolumeTexturePipeline)
+	{
+		SparseVolumeTexturePipeline->FilterPropertiesFromTranslatedData(InBaseNodeContainer);
+	}
 }
 
 bool UInterchangeGenericMaterialPipeline::IsPropertyChangeNeedRefresh(const FPropertyChangedEvent& PropertyChangedEvent) const
@@ -716,6 +736,10 @@ bool UInterchangeGenericMaterialPipeline::IsPropertyChangeNeedRefresh(const FPro
 	{
 		return true;
 	}
+	if (SparseVolumeTexturePipeline && SparseVolumeTexturePipeline->IsPropertyChangeNeedRefresh(PropertyChangedEvent))
+	{
+		return true;
+	}
 	return Super::IsPropertyChangeNeedRefresh(PropertyChangedEvent);
 }
 
@@ -725,6 +749,10 @@ void UInterchangeGenericMaterialPipeline::GetSupportAssetClasses(TArray<UClass*>
 	{
 		TexturePipeline->GetSupportAssetClasses(PipelineSupportAssetClasses);
 	}
+	if (SparseVolumeTexturePipeline)
+	{
+		SparseVolumeTexturePipeline->GetSupportAssetClasses(PipelineSupportAssetClasses);
+	}
 
 	PipelineSupportAssetClasses.Add(UMaterial::StaticClass());
 	PipelineSupportAssetClasses.Add(UMaterialInstance::StaticClass());
@@ -755,6 +783,10 @@ void UInterchangeGenericMaterialPipeline::ExecutePipeline(UInterchangeBaseNodeCo
 		{
 			TexturePipeline->SetResultsContainer(Results);
 		}
+		if (SparseVolumeTexturePipeline)
+		{
+			SparseVolumeTexturePipeline->SetResultsContainer(Results);
+		}
 	}
 
 	BaseNodeContainer = InBaseNodeContainer;
@@ -768,6 +800,10 @@ void UInterchangeGenericMaterialPipeline::ExecutePipeline(UInterchangeBaseNodeCo
 	{
 		TexturePipeline->ScriptedExecutePipeline(InBaseNodeContainer, InSourceDatas, ContentBasePath);
 	}
+	if (SparseVolumeTexturePipeline)
+	{
+		SparseVolumeTexturePipeline->ScriptedExecutePipeline(InBaseNodeContainer, InSourceDatas, ContentBasePath);
+	}
 
 	
 	MaterialNodes.Empty();
@@ -1009,6 +1045,10 @@ void UInterchangeGenericMaterialPipeline::ExecutePostFactoryPipeline(const UInte
 	{
 		TexturePipeline->ScriptedExecutePostFactoryPipeline(InBaseNodeContainer, NodeKey, CreatedAsset, bIsAReimport);
 	}
+	if (SparseVolumeTexturePipeline)
+	{
+		SparseVolumeTexturePipeline->ScriptedExecutePostFactoryPipeline(InBaseNodeContainer, NodeKey, CreatedAsset, bIsAReimport);
+	}
 }
 
 void UInterchangeGenericMaterialPipeline::ExecutePostImportPipeline(const UInterchangeBaseNodeContainer* InBaseNodeContainer, const FString& NodeKey, UObject* CreatedAsset, bool bIsAReimport)
@@ -1017,6 +1057,10 @@ void UInterchangeGenericMaterialPipeline::ExecutePostImportPipeline(const UInter
 	{
 		TexturePipeline->ScriptedExecutePostImportPipeline(InBaseNodeContainer, NodeKey, CreatedAsset, bIsAReimport);
 	}
+	if (SparseVolumeTexturePipeline)
+	{
+		SparseVolumeTexturePipeline->ScriptedExecutePostImportPipeline(InBaseNodeContainer, NodeKey, CreatedAsset, bIsAReimport);
+	}
 }
 
 void UInterchangeGenericMaterialPipeline::SetReimportSourceIndex(UClass* ReimportObjectClass, const int32 SourceFileIndex)
@@ -1025,6 +1069,10 @@ void UInterchangeGenericMaterialPipeline::SetReimportSourceIndex(UClass* Reimpor
 	{
 		TexturePipeline->ScriptedSetReimportSourceIndex(ReimportObjectClass, SourceFileIndex);
 	}
+	if (SparseVolumeTexturePipeline)
+	{
+		SparseVolumeTexturePipeline->ScriptedSetReimportSourceIndex(ReimportObjectClass, SourceFileIndex);
+	}
 }
 
 UInterchangeBaseMaterialFactoryNode* UInterchangeGenericMaterialPipeline::CreateBaseMaterialFactoryNode(const UInterchangeBaseNode* MaterialNode, TSubclassOf<UInterchangeBaseMaterialFactoryNode> NodeType, bool bAddMaterialInstanceSuffix /*= false*/)
@@ -1969,9 +2017,9 @@ void UInterchangeGenericMaterialPipeline::HandleMakeFloat3Node(const UInterchang
 }
 
 void UInterchangeGenericMaterialPipeline::HandleTextureNode(
-	const UInterchangeTextureNode* TextureNode, 
-	UInterchangeBaseMaterialFactoryNode* MaterialFactoryNode, 
-	UInterchangeMaterialExpressionFactoryNode* TextureBaseFactoryNode, 
+	const UInterchangeTextureNode* TextureNode,
+	UInterchangeBaseMaterialFactoryNode* MaterialFactoryNode,
+	UInterchangeMaterialExpressionFactoryNode* TextureBaseFactoryNode,
 	const FString & ExpressionClassName,
 	bool bIsAParameter)
 {
diff --git a/Engine/Plugins/Interchange/Runtime/Source/Pipelines/Private/InterchangeSparseVolumeTexturePipeline.cpp b/Engine/Plugins/Interchange/Runtime/Source/Pipelines/Private/InterchangeSparseVolumeTexturePipeline.cpp
new file mode 100644
--- /dev/null
+++ b/Engine/Plugins/Interchange/Runtime/Source/Pipelines/Private/InterchangeSparseVolumeTexturePipeline.cpp
@@ -0,0 +1,486 @@
+// Copyright Epic Games, Inc. All Rights Reserved.
+
+#include "InterchangeSparseVolumeTexturePipeline.h"
+
+#include "InterchangePipelineLog.h"
+#include "InterchangeSparseVolumeTextureFactoryNode.h"
+#include "InterchangeVolumeNode.h"
+#include "Volume/InterchangeVolumeDefinitions.h"
+
+#include "Engine/Texture.h"
+#include "SparseVolumeTexture/SparseVolumeTexture.h"
+
+#include UE_INLINE_GENERATED_CPP_BY_NAME(InterchangeSparseVolumeTexturePipeline)
+
+namespace UE::Interchange::Private
+{
+	UInterchangeSparseVolumeTextureFactoryNode* CreateTextureFactoryNode(
+		const FString& DisplayLabel,
+		const FString& NodeUid,
+		UInterchangeBaseNodeContainer* BaseNodeContainer
+	)
+	{
+		UInterchangeSparseVolumeTextureFactoryNode* TextureFactoryNode = nullptr;
+		if (BaseNodeContainer->IsNodeUidValid(NodeUid))
+		{
+			TextureFactoryNode = Cast<UInterchangeSparseVolumeTextureFactoryNode>(BaseNodeContainer->GetFactoryNode(NodeUid));
+			if (!ensure(TextureFactoryNode))
+			{
+				// Log an error
+				return nullptr;
+			}
+		}
+		else
+		{
+			const EInterchangeNodeContainerType NodeContainerType = EInterchangeNodeContainerType::FactoryData;
+
+			TextureFactoryNode = NewObject<UInterchangeSparseVolumeTextureFactoryNode>(BaseNodeContainer);
+			TextureFactoryNode->InitializeNode(NodeUid, DisplayLabel, NodeContainerType);
+
+			BaseNodeContainer->AddNode(TextureFactoryNode);
+		}
+
+		return TextureFactoryNode;
+	}
+
+	// SparseVolumeTextures have 8 individual channels, grouped into two RGBA 'textures' called "AttributesA" and "AttributesB",
+	// each 'texture' being of a format according to EInterchangeSparseVolumeTextureFormat.
+	//
+	// The purpose of this function is to figure out some sensible default assignment/distribution of the grids of the provided
+	// volume texture across these 8 channels. The idea is that other pipelines (like the USD Pipeline) would later override these
+	// with any specific grid-to-SVT channel mapping that the source files specify.
+	//
+	// Another goal here is to match the default assignment done by the SparseVolumeTextureFactory, so that SVTs imported via
+	// Interchange match the ones imported with the legacy factory
+	void SetupDefaultOpenVDBGridAssignment(
+		UInterchangeSparseVolumeTextureFactoryNode* VolumeFactoryNode,
+		UInterchangeBaseNodeContainer* BaseNodeContainer
+	)
+	{
+		// References:
+		// - ComputeDefaultOpenVDBGridAssignment function from SparseVolumeTextureFactory.cpp
+
+		using namespace UE::Interchange::Volume;
+
+		if (!VolumeFactoryNode || !BaseNodeContainer)
+		{
+			return;
+		}
+
+		// Get the translated node for this factory node
+		const UInterchangeVolumeNode* VolumeNode = nullptr;
+		{
+			TArray<FString> TargetNodeUids;
+			VolumeFactoryNode->GetTargetNodeUids(TargetNodeUids);
+
+			for (int32 TargetNodeIndex = TargetNodeUids.Num() - 1; TargetNodeIndex >= 0; --TargetNodeIndex)
+			{
+				VolumeNode = Cast<UInterchangeVolumeNode>(BaseNodeContainer->GetNode(TargetNodeUids[TargetNodeIndex]));
+				if (VolumeNode)
+				{
+					break;
+				}
+			}
+
+			if (!VolumeNode)
+			{
+				return;
+			}
+		}
+
+		// Get all the grids contained in the given volume
+		TArray<const UInterchangeVolumeGridNode*> GridNodes;
+		{
+			TArray<FString> GridNodeUids;
+			VolumeNode->GetGridDependencies(GridNodeUids);
+
+			for (int32 GridIndex = 0; GridIndex < GridNodeUids.Num(); ++GridIndex)
+			{
+				const UInterchangeVolumeGridNode* GridNode = Cast<UInterchangeVolumeGridNode>(BaseNodeContainer->GetNode(GridNodeUids[GridIndex]));
+				if (GridNode)
+				{
+					GridNodes.Add(GridNode);
+				}
+			}
+
+			if (GridNodes.Num() == 0)
+			{
+				return;
+			}
+		}
+
+		// Check whether we have a grid named "density" (seems to be common for .vdbs)
+		int32 DensityGridIndex = INDEX_NONE;
+		int32 NumNonDensity = GridNodes.Num();
+		for (int32 Index = 0; Index < GridNodes.Num(); ++Index)
+		{
+			const UInterchangeVolumeGridNode* Node = GridNodes[Index];
+			if (Node->GetDisplayLabel() == DensityGridName)
+			{
+				DensityGridIndex = Index;
+				NumNonDensity--;
+				break;
+			}
+		}
+
+		// We use these to help distribute the grids through the different channels, as we have to iterate through them
+		int32 SetterIndex = 0;
+		using SetterFunc = decltype(&UInterchangeSparseVolumeTextureFactoryNode::SetCustomAttributesAChannelX);
+		static const TArray<SetterFunc> AttributeChannelSetters = {
+			&UInterchangeSparseVolumeTextureFactoryNode::SetCustomAttributesAChannelX,
+			&UInterchangeSparseVolumeTextureFactoryNode::SetCustomAttributesAChannelY,
+			&UInterchangeSparseVolumeTextureFactoryNode::SetCustomAttributesAChannelZ,
+			&UInterchangeSparseVolumeTextureFactoryNode::SetCustomAttributesAChannelW,
+			&UInterchangeSparseVolumeTextureFactoryNode::SetCustomAttributesBChannelX,
+			&UInterchangeSparseVolumeTextureFactoryNode::SetCustomAttributesBChannelY,
+			&UInterchangeSparseVolumeTextureFactoryNode::SetCustomAttributesBChannelZ,
+			&UInterchangeSparseVolumeTextureFactoryNode::SetCustomAttributesBChannelW,
+		};
+
+		// Optimized density assignment: "density" grid as 8bit unsigned normalized on AttributesA, and everything else on Attributes B.
+		// This is only done if we have 0, 1, 2 or 4 non-density grid components (if we have density and 3 non-density we have a total
+		// of 4, so since they would fit nicely into a single AttributesA 'texture' we just do that instead)
+		const bool bOptimizedDensityAssignment = DensityGridIndex != INDEX_NONE && (NumNonDensity <= 4 && NumNonDensity != 3);
+		if (bOptimizedDensityAssignment)
+		{
+			VolumeFactoryNode->SetCustomAttributesADataType(EInterchangeSparseVolumeTextureFormat::Unorm8);
+			VolumeFactoryNode->SetCustomAttributesAChannelX(DensityGridName + GridNameAndComponentIndexSeparator + TEXT("0"));
+
+			VolumeFactoryNode->SetCustomAttributesBDataType(EInterchangeSparseVolumeTextureFormat::Float16);
+
+			SetterIndex = 4;	// Start at SetCustomAttributesBChannelX() instead, as our AttributesA texture will hold just the density
+		}
+		else
+		{
+			VolumeFactoryNode->SetCustomAttributesADataType(EInterchangeSparseVolumeTextureFormat::Float16);
+
+			VolumeFactoryNode->SetCustomAttributesBDataType(EInterchangeSparseVolumeTextureFormat::Float16);
+		}
+
+		// Actually distribute the remaining grid/components across the channels in order
+		for (int32 Index = 0; Index < GridNodes.Num(); ++Index)
+		{
+			if (SetterIndex >= AttributeChannelSetters.Num())
+			{
+				break;
+			}
+
+			if (Index == DensityGridIndex)
+			{
+				continue;
+			}
+
+			const UInterchangeVolumeGridNode* Node = GridNodes[Index];
+
+			int32 GridNumComponents = 0;
+			bool bHasComponents = Node->GetCustomNumComponents(GridNumComponents);
+			if (!bHasComponents)
+			{
+				continue;
+			}
+
+			// e.g. "temperature_"
+			const FString& GridNameAndSeparator = Node->GetDisplayLabel() + GridNameAndComponentIndexSeparator;
+
+			for (int32 GridComponentIndex = 0;															   //
+				 GridComponentIndex < GridNumComponents && SetterIndex < AttributeChannelSetters.Num();	   //
+				 ++GridComponentIndex, ++SetterIndex)
+			{
+				// e.g. "temperature_2"
+				const FString GridNameAndComponentIndex = GridNameAndSeparator + LexToString(GridComponentIndex);
+
+				const SetterFunc& Setter = AttributeChannelSetters[SetterIndex];
+				(VolumeFactoryNode->*Setter)(GridNameAndComponentIndex);
+			}
+		}
+	}
+
+	// Splits something like "tornado_23" into the "tornado_" prefix and 23 suffix
+	void SplitNumberedSuffix(const FString& String, FString& OutPrefix, int32& OutSuffix)
+	{
+		const int32 LastNonDigitIndex = String.FindLastCharByPredicate(
+			[](TCHAR Letter)
+			{
+				return !FChar::IsDigit(Letter);
+			}
+		);
+
+		// No numbered suffix
+		if (LastNonDigitIndex == String.Len() - 1)
+		{
+			OutPrefix = String;
+			OutSuffix = INDEX_NONE;
+			return;
+		}
+
+		FString NumberSuffixStr;
+
+		// String is all numbers
+		if (LastNonDigitIndex == INDEX_NONE)
+		{
+			NumberSuffixStr = String;
+			OutPrefix = {};
+		}
+		// Some prefix, some numbers
+		else
+		{
+			NumberSuffixStr = String.RightChop(LastNonDigitIndex + 1);
+			OutPrefix = String.Left(LastNonDigitIndex + 1);
+		}
+
+		int32 Number = INDEX_NONE;
+		if (ensure(NumberSuffixStr.IsNumeric()))
+		{
+			TTypeFromString<int32>::FromString(Number, *NumberSuffixStr);
+		}
+		OutSuffix = Number;
+	}
+
+	// Turns something like "tornado_" or "tornado-" into just "tornado"
+	FString RemoveTrailingSeparators(FString String)
+	{
+		FString LastChar = String.Right(1);
+		while ((LastChar == TEXT("-") || LastChar == TEXT("_")) && String.Len() > 1)
+		{
+			String.LeftChopInline(1, EAllowShrinking::No);
+			LastChar = String.Right(1);
+		}
+		String.Shrink();
+		return String;
+	}
+}	 // namespace UE::Interchange::Private
+
+FString UInterchangeSparseVolumeTexturePipeline::GetPipelineCategory(UClass* AssetClass)
+{
+	// Ideally we'd be in a "Volumes" one, but these seem to be somewhat hard-coded?
+	return TEXT("Textures");
+}
+
+void UInterchangeSparseVolumeTexturePipeline::AdjustSettingsForContext(const FInterchangePipelineContextParams& ContextParams)
+{
+	Super::AdjustSettingsForContext(ContextParams);
+
+#if WITH_EDITOR
+	TArray<FString> HideCategories;
+	bool bIsObjectAnSVT = !ContextParams.ReimportAsset ? false : ContextParams.ReimportAsset.IsA(USparseVolumeTexture::StaticClass());
+	if ((!bIsObjectAnSVT && ContextParams.ContextType == EInterchangePipelineContext::AssetReimport)
+		|| ContextParams.ContextType == EInterchangePipelineContext::AssetCustomLODImport
+		|| ContextParams.ContextType == EInterchangePipelineContext::AssetCustomLODReimport
+		|| ContextParams.ContextType == EInterchangePipelineContext::AssetAlternateSkinningImport
+		|| ContextParams.ContextType == EInterchangePipelineContext::AssetAlternateSkinningReimport
+		|| ContextParams.ContextType == EInterchangePipelineContext::AssetCustomMorphTargetImport
+		|| ContextParams.ContextType == EInterchangePipelineContext::AssetCustomMorphTargetReImport)
+	{
+		bImportSparseVolumeTextures = false;
+		bImportAnimatedSparseVolumeTextures = false;
+		HideCategories.Add(UInterchangeSparseVolumeTexturePipeline::GetPipelineCategory(nullptr));
+	}
+
+	if (UInterchangePipelineBase* OuterMostPipeline = GetMostPipelineOuter())
+	{
+		for (const FString& HideCategoryName : HideCategories)
+		{
+			HidePropertiesOfCategory(OuterMostPipeline, this, HideCategoryName);
+		}
+	}
+#endif	  // WITH_EDITOR
+}
+
+#if WITH_EDITOR
+bool UInterchangeSparseVolumeTexturePipeline::IsPropertyChangeNeedRefresh(const FPropertyChangedEvent& PropertyChangedEvent) const
+{
+	if (PropertyChangedEvent.GetPropertyName() == GET_MEMBER_NAME_CHECKED(UInterchangeSparseVolumeTexturePipeline, bImportSparseVolumeTextures))
+	{
+		return true;
+	}
+
+	if (PropertyChangedEvent.GetPropertyName()
+		== GET_MEMBER_NAME_CHECKED(UInterchangeSparseVolumeTexturePipeline, bImportAnimatedSparseVolumeTextures))
+	{
+		return true;
+	}
+
+	return Super::IsPropertyChangeNeedRefresh(PropertyChangedEvent);
+}
+
+void UInterchangeSparseVolumeTexturePipeline::FilterPropertiesFromTranslatedData(UInterchangeBaseNodeContainer* InBaseNodeContainer)
+{
+	Super::FilterPropertiesFromTranslatedData(InBaseNodeContainer);
+
+	TArray<FString> TmpTextureNodes;
+	InBaseNodeContainer->GetNodes(UInterchangeVolumeNode::StaticClass(), TmpTextureNodes);
+	if (TmpTextureNodes.Num() == 0)
+	{
+		if (UInterchangePipelineBase* OuterMostPipeline = GetMostPipelineOuter())
+		{
+			HidePropertiesOfCategory(OuterMostPipeline, this, UInterchangeSparseVolumeTexturePipeline::GetPipelineCategory(nullptr));
+		}
+	}
+}
+
+void UInterchangeSparseVolumeTexturePipeline::GetSupportAssetClasses(TArray<UClass*>& PipelineSupportAssetClasses) const
+{
+	PipelineSupportAssetClasses.Add(USparseVolumeTexture::StaticClass());
+}
+#endif	  // WITH_EDITOR
+
+void UInterchangeSparseVolumeTexturePipeline::ExecutePipeline(
+	UInterchangeBaseNodeContainer* InBaseNodeContainer,
+	const TArray<UInterchangeSourceData*>& InSourceDatas,
+	const FString& ContentBasePath
+)
+{
+	using namespace UE::Interchange::Private;
+
+	if (!bImportSparseVolumeTextures)
+	{
+		return;
+	}
+
+	if (!InBaseNodeContainer)
+	{
+		return;
+	}
+	BaseNodeContainer = InBaseNodeContainer;
+
+	// Find all the translated nodes we need for this pipeline
+	TArray<UInterchangeVolumeNode*> VolumeNodes;
+	BaseNodeContainer->IterateNodes(
+		[&VolumeNodes](const FString& NodeUid, UInterchangeBaseNode* Node)
+		{
+			if (UInterchangeVolumeNode* TextureNode = Cast<UInterchangeVolumeNode>(Node))
+			{
+				VolumeNodes.Add(TextureNode);
+			}
+		}
+	);
+
+	TArray<UInterchangeSparseVolumeTextureFactoryNode*> CreatedFactoryNodes;
+
+	// We need to group up all the volume nodes that seem to refer to the same sequence, according to their filename.
+	// These could have been manually added by the user somehow, or discovered via the OpenVDB translator.
+	if (bImportAnimatedSparseVolumeTextures)
+	{
+		struct FNodeAndTrailingIndex
+		{
+			const UInterchangeVolumeNode* Node;
+			int32 Index;
+		};
+
+		TMap<FString, TArray<FNodeAndTrailingIndex>> NodesWithSamePrefix;
+		NodesWithSamePrefix.Reserve(VolumeNodes.Num());
+
+		// Group up volume nodes by the same filename prefix
+		for (UInterchangeVolumeNode* VolumeNode : VolumeNodes)
+		{
+			FString FileName;
+			bool bSuccess = VolumeNode->GetCustomFileName(FileName);
+			if (!bSuccess || FileName.IsEmpty())
+			{
+				continue;
+			}
+			FileName = FPaths::GetBaseFilename(FileName);
+
+			// Note: Suffix may be INDEX_NONE, but that should be alright (they'll just be grouped together
+			// at the front, if any)
+			FString Prefix;
+			int32 Suffix;
+			SplitNumberedSuffix(FileName, Prefix, Suffix);
+
+			FNodeAndTrailingIndex& NodeAndIndex = NodesWithSamePrefix.FindOrAdd(Prefix).Emplace_GetRef();
+			NodeAndIndex.Node = VolumeNode;
+			NodeAndIndex.Index = Suffix;
+		}
+
+		for (TPair<FString, TArray<FNodeAndTrailingIndex>>& Pair : NodesWithSamePrefix)
+		{
+			TArray<FNodeAndTrailingIndex>& NodeAndIndices = Pair.Value;
+			if (NodeAndIndices.Num() == 0)
+			{
+				continue;
+			}
+
+			// Sort them according to their suffix integers (note: We need all this because just sorting the filename
+			// strings is not enough in case the files have numbered suffixes without zero padding.
+			// i.e. lexicographic sort would produce an order like ["tornado_1", "tornado_12", "tornado_2", "tornado_24", etc.])
+			NodeAndIndices.Sort(
+				[](const FNodeAndTrailingIndex& LHS, const FNodeAndTrailingIndex& RHS)
+				{
+					return LHS.Index < RHS.Index;
+				}
+			);
+
+			const FString& Prefix = Pair.Key;
+			const UInterchangeVolumeNode* FirstVolume = NodeAndIndices[0].Node;
+
+			FString DisplayLabel;
+			if (NodeAndIndices.Num() == 1)
+			{
+				DisplayLabel = FirstVolume->GetDisplayLabel();
+			}
+			else
+			{
+				// Splitting the numbered suffix likely left us with something like "tornado_" as the prefix.
+				// If we're going to use that as the asset name, let's strip any separators that were left at the end like that
+				DisplayLabel = RemoveTrailingSeparators(Prefix);
+			}
+
+			const FString FactoryNodeUid = UInterchangeFactoryBaseNode::BuildFactoryNodeUid(FirstVolume->GetUniqueID());
+
+			UInterchangeSparseVolumeTextureFactoryNode* FactoryNode = CreateTextureFactoryNode(DisplayLabel, FactoryNodeUid, BaseNodeContainer);
+			if (!FactoryNode)
+			{
+				continue;
+			}
+			CreatedFactoryNodes.Add(FactoryNode);
+
+			for (const FNodeAndTrailingIndex& NodeAndIndex : NodeAndIndices)
+			{
+				FactoryNode->AddTargetNodeUid(NodeAndIndex.Node->GetUniqueID());
+				NodeAndIndex.Node->AddTargetNodeUid(FactoryNode->GetUniqueID());
+			}
+
+			SetupDefaultOpenVDBGridAssignment(FactoryNode, BaseNodeContainer);
+		}
+	}
+	// No grouping: Each volume node becomes it's own factory node
+	else
+	{
+		for (UInterchangeVolumeNode* VolumeNode : VolumeNodes)
+		{
+			const FString DisplayLabel = VolumeNode->GetDisplayLabel();
+			const FString NodeUid = UInterchangeFactoryBaseNode::BuildFactoryNodeUid(VolumeNode->GetUniqueID());
+
+			UInterchangeSparseVolumeTextureFactoryNode* FactoryNode = CreateTextureFactoryNode(DisplayLabel, NodeUid, BaseNodeContainer);
+			if (!FactoryNode)
+			{
+				continue;
+			}
+			CreatedFactoryNodes.Add(FactoryNode);
+
+			FactoryNode->AddTargetNodeUid(VolumeNode->GetUniqueID());
+			VolumeNode->AddTargetNodeUid(FactoryNode->GetUniqueID());
+
+			SetupDefaultOpenVDBGridAssignment(FactoryNode, BaseNodeContainer);
+		}
+	}
+
+	// Set an override asset name if we have exactly one factory node
+	if (CreatedFactoryNodes.Num() == 1)
+	{
+		FString OverrideAssetName = IsStandAlonePipeline() ? DestinationName : FString();
+		if (OverrideAssetName.IsEmpty() && IsStandAlonePipeline())
+		{
+			OverrideAssetName = AssetName;
+		}
+
+		UInterchangeSparseVolumeTextureFactoryNode* FactoryNode = CreatedFactoryNodes[0];
+
+		const bool bOverrideAssetName = IsStandAlonePipeline() && !OverrideAssetName.IsEmpty();
+		if (FactoryNode && bOverrideAssetName)
+		{
+			FactoryNode->SetAssetName(OverrideAssetName);
+			FactoryNode->SetDisplayLabel(OverrideAssetName);
+		}
+	}
+}
diff --git a/Engine/Plugins/Interchange/Runtime/Source/Pipelines/Public/InterchangeGenericMaterialPipeline.h b/Engine/Plugins/Interchange/Runtime/Source/Pipelines/Public/InterchangeGenericMaterialPipeline.h
--- a/Engine/Plugins/Interchange/Runtime/Source/Pipelines/Public/InterchangeGenericMaterialPipeline.h
+++ b/Engine/Plugins/Interchange/Runtime/Source/Pipelines/Public/InterchangeGenericMaterialPipeline.h
@@ -14,15 +14,16 @@
 class UInterchangeBaseMaterialFactoryNode;
 class UInterchangeFunctionCallShaderNode;
 class UInterchangeGenericTexturePipeline;
-class UInterchangeShaderGraphNode;
-class UInterchangeShaderNode;
-class UInterchangeMaterialFactoryNode;
 class UInterchangeMaterialExpressionFactoryNode;
+class UInterchangeMaterialFactoryNode;
 class UInterchangeMaterialFunctionFactoryNode;
 class UInterchangeMaterialInstanceFactoryNode;
 class UInterchangeResult;
-class UMaterialFunction;
+class UInterchangeShaderGraphNode;
+class UInterchangeShaderNode;
+class UInterchangeSparseVolumeTexturePipeline;
 class UInterchangeTextureNode;
+class UMaterialFunction;
 
 namespace UE::Interchange::Materials::HashUtils
 {
@@ -98,6 +99,9 @@ public:
 	UPROPERTY(VisibleAnywhere, BlueprintReadWrite, Instanced, Category = "Textures")
 	TObjectPtr<UInterchangeGenericTexturePipeline> TexturePipeline;
 
+	UPROPERTY(VisibleAnywhere, BlueprintReadWrite, Instanced, Category = "Textures")
+	TObjectPtr<UInterchangeSparseVolumeTexturePipeline> SparseVolumeTexturePipeline;
+
 	/** BEGIN UInterchangePipelineBase overrides */
 	virtual void PreDialogCleanup(const FName PipelineStackName) override;
 	virtual bool IsSettingsAreValid(TOptional<FText>& OutInvalidReason) const override;
diff --git a/Engine/Plugins/Interchange/Runtime/Source/Pipelines/Public/InterchangeSparseVolumeTexturePipeline.h b/Engine/Plugins/Interchange/Runtime/Source/Pipelines/Public/InterchangeSparseVolumeTexturePipeline.h
new file mode 100644
--- /dev/null
+++ b/Engine/Plugins/Interchange/Runtime/Source/Pipelines/Public/InterchangeSparseVolumeTexturePipeline.h
@@ -0,0 +1,63 @@
+// Copyright Epic Games, Inc. All Rights Reserved.
+
+#pragma once
+
+#include "InterchangePipelineBase.h"
+#include "InterchangeSourceData.h"
+#include "Nodes/InterchangeBaseNodeContainer.h"
+#include "UObject/Object.h"
+#include "UObject/ObjectMacros.h"
+
+#include "InterchangeSparseVolumeTexturePipeline.generated.h"
+
+class UInterchangeTextureFactoryNode;
+class UInterchangeTextureNode;
+
+UCLASS(BlueprintType, editinlinenew)
+class INTERCHANGEPIPELINES_API UInterchangeSparseVolumeTexturePipeline : public UInterchangePipelineBase
+{
+	GENERATED_BODY()
+
+public:
+	static FString GetPipelineCategory(UClass* AssetClass);
+
+	UPROPERTY(
+		EditAnywhere,
+		BlueprintReadWrite,
+		Category = "Textures",
+		meta = (StandAlonePipelineProperty = "True", PipelineInternalEditionData = "True")
+	)
+	FString PipelineDisplayName;
+
+	/** If enabled, imports all sparse volume texture assets found in the source. */
+	UPROPERTY(EditAnywhere, BlueprintReadWrite, Category = "Textures")
+	bool bImportSparseVolumeTextures = true;
+
+	/**
+	 * If enabled, will attempt to import volume nodes corresponding to numbered files in the same folder as individual frames
+	 * of an animated SparseVolumeTextures.
+	 */
+	UPROPERTY(EditAnywhere, BlueprintReadWrite, Category = "Textures")
+	bool bImportAnimatedSparseVolumeTextures = true;
+
+	/** If set, and there is only one asset and one source, the imported asset will be given this name. */
+	UPROPERTY(EditAnywhere, BlueprintReadWrite, Category = "Textures", meta = (StandAlonePipelineProperty = "True"))
+	FString AssetName;
+
+public:
+	virtual void AdjustSettingsForContext(const FInterchangePipelineContextParams& ContextParams) override;
+	virtual void ExecutePipeline(
+		UInterchangeBaseNodeContainer* InBaseNodeContainer,
+		const TArray<UInterchangeSourceData*>& InSourceDatas,
+		const FString& ContentBasePath
+	) override;
+#if WITH_EDITOR
+	virtual bool IsPropertyChangeNeedRefresh(const FPropertyChangedEvent& PropertyChangedEvent) const override;
+	virtual void FilterPropertiesFromTranslatedData(UInterchangeBaseNodeContainer* InBaseNodeContainer) override;
+	virtual void GetSupportAssetClasses(TArray<UClass*>& PipelineSupportAssetClasses) const override;
+#endif	  // WITH_EDITOR
+
+private:
+	UPROPERTY()
+	TObjectPtr<UInterchangeBaseNodeContainer> BaseNodeContainer;
+};
diff --git a/Engine/Source/Editor/SparseVolumeTexture/Private/SparseVolumeTextureFactory.cpp b/Engine/Source/Editor/SparseVolumeTexture/Private/SparseVolumeTextureFactory.cpp
--- a/Engine/Source/Editor/SparseVolumeTexture/Private/SparseVolumeTextureFactory.cpp
+++ b/Engine/Source/Editor/SparseVolumeTexture/Private/SparseVolumeTextureFactory.cpp
@@ -105,133 +105,6 @@ static void ComputeDefaultOpenVDBGridAssignment(const TArray<TSharedPtr<FOpenVDB
 	ImportOptions->bIsSequence = NumFiles > 1;
 }
 
-static FString GetVDBSequenceBaseFileName(const FString& FileName, bool bDiscardNumbersOnly)
-{
-	const FString CleanFileName = FPaths::GetCleanFilename(FileName);
-	int32 NumValidChars = CleanFileName.Len() - 4; // chop off the file extension
-
-	// Remove any digits at the end
-	NumValidChars = CleanFileName.FindLastCharByPredicate([&](TCHAR Letter){ return !FChar::IsDigit(Letter); }, NumValidChars) + 1;
-
-	// Optionally remove other unwanted chars like underscores and invalid object name chars that would later be replaced by underscores
-	if (!bDiscardNumbersOnly)
-	{
-		NumValidChars = CleanFileName.FindLastCharByPredicate([&](TCHAR Letter) 
-			{
-				// INVALID_OBJECTNAME_CHARACTERS is defined in NameTypes.h and is a string literal containing all the invalid chars.
-				// The number at the end of a filename in a sequence is often separated by an underscore or other special character,
-				// so when we want to get the base filename for deriving the new asset name, we also discard these characters.
-				// Underscores are not part of the invalid chars string literal, so we append them here to also discard these chars.
-				const TCHAR* InvalidChars = TEXT("_") INVALID_OBJECTNAME_CHARACTERS;
-				while (*InvalidChars)
-				{
-					if (Letter == *InvalidChars)
-					{
-						return false;
-					}
-					++InvalidChars;
-				}
-				return true;
-			}, NumValidChars) + 1;
-	}
-
-	const FString CleanFileNameWithoutSuffix = CleanFileName.Left(NumValidChars);
-	return CleanFileNameWithoutSuffix;
-}
-
-static TArray<FString> FindOpenVDBSequenceFileNames(const FString& Filename)
-{
-	TArray<FString> SequenceFilenames;
-
-	// The file is potentially a sequence if the character before the `.vdb` is a number.
-	const bool bIsFilePotentiallyPartOfASequence = FChar::IsDigit(Filename[Filename.Len() - 5]);
-
-	if (!bIsFilePotentiallyPartOfASequence)
-	{
-		SequenceFilenames.Add(Filename);
-	}
-	else
-	{
-		const FString Path = FPaths::GetPath(Filename);
-		const FString CleanFilename = FPaths::GetCleanFilename(Filename);
-		const FString CleanFilenameWithoutSuffix = GetVDBSequenceBaseFileName(Filename, true /*bDiscardNumbersOnly*/);
-
-		// Find all files potentially part of the sequence
-		TArray<FString> PotentialSequenceFilenames;
-		IFileManager::Get().FindFiles(PotentialSequenceFilenames, *Path, TEXT("*.vdb"));
-		PotentialSequenceFilenames = PotentialSequenceFilenames.FilterByPredicate([&CleanFilenameWithoutSuffix](const FString& Str) 
-			{ 
-				if (!CleanFilenameWithoutSuffix.IsEmpty())
-				{
-					// Check for the same base filename
-					return Str.StartsWith(CleanFilenameWithoutSuffix);
-				}
-				else
-				{
-					// Removing the digits at the end of the input file resulted in an empty string, so we are looking for numeric filenames only (excluding the 4 chars file extension)
-					return Str.LeftChop(4).IsNumeric();
-				}
-			});
-
-		auto GetFilenameNumberSuffix = [](const FString& Filename) -> int32
-		{
-			const FString FilenameWithoutExt = Filename.LeftChop(4);
-			const int32 LastNonDigitIndex = FilenameWithoutExt.FindLastCharByPredicate([](TCHAR Letter) { return !FChar::IsDigit(Letter); }) + 1;
-			const FString NumberSuffixStr = FilenameWithoutExt.RightChop(LastNonDigitIndex);
-
-			int32 Number = INDEX_NONE;
-			if (NumberSuffixStr.IsNumeric())
-			{
-				TTypeFromString<int32>::FromString(Number, *NumberSuffixStr);
-			}
-			return Number;
-		};
-
-		// Find range of number suffixes
-		int32 LowestIndex = INT32_MAX;
-		int32 HighestIndex = INT32_MIN;
-		for (FString& ItemFilename : PotentialSequenceFilenames)
-		{
-			const int32 Index = GetFilenameNumberSuffix(ItemFilename);
-			if (Index == INDEX_NONE)
-			{
-				ItemFilename.Empty();
-				continue;
-			}
-			LowestIndex = FMath::Min(LowestIndex, Index);
-			HighestIndex = FMath::Max(HighestIndex, Index);
-		}
-
-		check(HighestIndex >= LowestIndex);
-
-		// Sort the filenames into the result array
-		SequenceFilenames.SetNum(HighestIndex - LowestIndex + 1);
-		for (const FString& ItemFilename : PotentialSequenceFilenames)
-		{
-			const int32 Index = ItemFilename.IsEmpty() ? INDEX_NONE : GetFilenameNumberSuffix(ItemFilename);
-			if (Index == INDEX_NONE)
-			{
-				continue;
-			}
-			SequenceFilenames[Index - LowestIndex] = Path / ItemFilename;
-		}
-
-		// Chop off any items after finding the first gap
-		for (int32 i = 0; i < SequenceFilenames.Num(); ++i)
-		{
-			if (SequenceFilenames[i].IsEmpty())
-			{
-				SequenceFilenames.SetNum(i);
-				break;
-			}
-		}
-	}
-
-	check(!SequenceFilenames.IsEmpty());
-
-	return SequenceFilenames;
-}
-
 bool LoadOpenVDBPreviewData(const FString& Filename, FOpenVDBPreviewData* OutPreviewData)
 {
 	FOpenVDBPreviewData& Result = *OutPreviewData;
diff --git a/Engine/Source/Editor/SparseVolumeTexture/Private/SparseVolumeTextureOpenVDB.h b/Engine/Source/Editor/SparseVolumeTexture/Private/SparseVolumeTextureOpenVDB.h
--- a/Engine/Source/Editor/SparseVolumeTexture/Private/SparseVolumeTextureOpenVDB.h
+++ b/Engine/Source/Editor/SparseVolumeTexture/Private/SparseVolumeTextureOpenVDB.h
@@ -5,9 +5,7 @@
 #include "HAL/Platform.h"
 #include "Misc/CoreMiscDefines.h" // UE_PUSH_MACRO seems to break without this
 
-#if PLATFORM_WINDOWS || PLATFORM_LINUX
-
-#define OPENVDB_AVAILABLE 1
+#if OPENVDB_AVAILABLE
 
 THIRD_PARTY_INCLUDES_START
 
@@ -71,8 +69,4 @@ UE_POP_MACRO("check")
 
 THIRD_PARTY_INCLUDES_END
 
-#else
-
-#define OPENVDB_AVAILABLE 0
-
 #endif
diff --git a/Engine/Source/Editor/SparseVolumeTexture/Private/SparseVolumeTextureOpenVDBUtility.cpp b/Engine/Source/Editor/SparseVolumeTexture/Private/SparseVolumeTextureOpenVDBUtility.cpp
--- a/Engine/Source/Editor/SparseVolumeTexture/Private/SparseVolumeTextureOpenVDBUtility.cpp
+++ b/Engine/Source/Editor/SparseVolumeTexture/Private/SparseVolumeTextureOpenVDBUtility.cpp
@@ -9,6 +9,9 @@
 #include "OpenVDBGridAdapter.h"
 #include "OpenVDBImportOptions.h"
 
+#include "HAL/FileManager.h"
+#include "Misc/Paths.h"
+
 DEFINE_LOG_CATEGORY_STATIC(LogSparseVolumeTextureOpenVDBUtility, Log, All);
 
 #if OPENVDB_AVAILABLE
@@ -529,4 +532,131 @@ const TCHAR* OpenVDBGridTypeToString(EOpenVDBGridType Type)
 	}
 }
 
+FString GetVDBSequenceBaseFileName(const FString& FileName, bool bDiscardNumbersOnly)
+{
+	const FString CleanFileName = FPaths::GetCleanFilename(FileName);
+	int32 NumValidChars = CleanFileName.Len() - 4; // chop off the file extension
+
+	// Remove any digits at the end
+	NumValidChars = CleanFileName.FindLastCharByPredicate([&](TCHAR Letter){ return !FChar::IsDigit(Letter); }, NumValidChars) + 1;
+
+	// Optionally remove other unwanted chars like underscores and invalid object name chars that would later be replaced by underscores
+	if (!bDiscardNumbersOnly)
+	{
+		NumValidChars = CleanFileName.FindLastCharByPredicate([&](TCHAR Letter)
+			{
+				// INVALID_OBJECTNAME_CHARACTERS is defined in NameTypes.h and is a string literal containing all the invalid chars.
+				// The number at the end of a filename in a sequence is often separated by an underscore or other special character,
+				// so when we want to get the base filename for deriving the new asset name, we also discard these characters.
+				// Underscores are not part of the invalid chars string literal, so we append them here to also discard these chars.
+				const TCHAR* InvalidChars = TEXT("_") INVALID_OBJECTNAME_CHARACTERS;
+				while (*InvalidChars)
+				{
+					if (Letter == *InvalidChars)
+					{
+						return false;
+					}
+					++InvalidChars;
+				}
+				return true;
+			}, NumValidChars) + 1;
+	}
+
+	const FString CleanFileNameWithoutSuffix = CleanFileName.Left(NumValidChars);
+	return CleanFileNameWithoutSuffix;
+}
+
+TArray<FString> FindOpenVDBSequenceFileNames(const FString& Filename)
+{
+	TArray<FString> SequenceFilenames;
+
+	// The file is potentially a sequence if the character before the `.vdb` is a number.
+	const bool bIsFilePotentiallyPartOfASequence = FChar::IsDigit(Filename[Filename.Len() - 5]);
+
+	if (!bIsFilePotentiallyPartOfASequence)
+	{
+		SequenceFilenames.Add(Filename);
+	}
+	else
+	{
+		const FString Path = FPaths::GetPath(Filename);
+		const FString CleanFilename = FPaths::GetCleanFilename(Filename);
+		const FString CleanFilenameWithoutSuffix = GetVDBSequenceBaseFileName(Filename, true /*bDiscardNumbersOnly*/);
+
+		// Find all files potentially part of the sequence
+		TArray<FString> PotentialSequenceFilenames;
+		IFileManager::Get().FindFiles(PotentialSequenceFilenames, *Path, TEXT("*.vdb"));
+		PotentialSequenceFilenames = PotentialSequenceFilenames.FilterByPredicate([&CleanFilenameWithoutSuffix](const FString& Str)
+			{
+				if (!CleanFilenameWithoutSuffix.IsEmpty())
+				{
+					// Check for the same base filename
+					return Str.StartsWith(CleanFilenameWithoutSuffix);
+				}
+				else
+				{
+					// Removing the digits at the end of the input file resulted in an empty string, so we are looking for numeric filenames only (excluding the 4 chars file extension)
+					return Str.LeftChop(4).IsNumeric();
+				}
+			});
+
+		auto GetFilenameNumberSuffix = [](const FString& Filename) -> int32
+		{
+			const FString FilenameWithoutExt = Filename.LeftChop(4);
+			const int32 LastNonDigitIndex = FilenameWithoutExt.FindLastCharByPredicate([](TCHAR Letter) { return !FChar::IsDigit(Letter); }) + 1;
+			const FString NumberSuffixStr = FilenameWithoutExt.RightChop(LastNonDigitIndex);
+
+			int32 Number = INDEX_NONE;
+			if (NumberSuffixStr.IsNumeric())
+			{
+				TTypeFromString<int32>::FromString(Number, *NumberSuffixStr);
+			}
+			return Number;
+		};
+
+		// Find range of number suffixes
+		int32 LowestIndex = INT32_MAX;
+		int32 HighestIndex = INT32_MIN;
+		for (FString& ItemFilename : PotentialSequenceFilenames)
+		{
+			const int32 Index = GetFilenameNumberSuffix(ItemFilename);
+			if (Index == INDEX_NONE)
+			{
+				ItemFilename.Empty();
+				continue;
+			}
+			LowestIndex = FMath::Min(LowestIndex, Index);
+			HighestIndex = FMath::Max(HighestIndex, Index);
+		}
+
+		check(HighestIndex >= LowestIndex);
+
+		// Sort the filenames into the result array
+		SequenceFilenames.SetNum(HighestIndex - LowestIndex + 1);
+		for (const FString& ItemFilename : PotentialSequenceFilenames)
+		{
+			const int32 Index = ItemFilename.IsEmpty() ? INDEX_NONE : GetFilenameNumberSuffix(ItemFilename);
+			if (Index == INDEX_NONE)
+			{
+				continue;
+			}
+			SequenceFilenames[Index - LowestIndex] = Path / ItemFilename;
+		}
+
+		// Chop off any items after finding the first gap
+		for (int32 i = 0; i < SequenceFilenames.Num(); ++i)
+		{
+			if (SequenceFilenames[i].IsEmpty())
+			{
+				SequenceFilenames.SetNum(i);
+				break;
+			}
+		}
+	}
+
+	check(!SequenceFilenames.IsEmpty());
+
+	return SequenceFilenames;
+}
+
 #endif // WITH_EDITOR
\ No newline at end of file
diff --git a/Engine/Source/Editor/SparseVolumeTexture/Private/SparseVolumeTextureOpenVDBUtility.h b/Engine/Source/Editor/SparseVolumeTexture/Private/SparseVolumeTextureOpenVDBUtility.h
deleted file mode 100644
diff --git a/Engine/Source/Editor/SparseVolumeTexture/Public/SparseVolumeTextureOpenVDBUtility.h b/Engine/Source/Editor/SparseVolumeTexture/Public/SparseVolumeTextureOpenVDBUtility.h
new file mode 100644
--- /dev/null
+++ b/Engine/Source/Editor/SparseVolumeTexture/Public/SparseVolumeTextureOpenVDBUtility.h
@@ -0,0 +1,39 @@
+// Copyright Epic Games, Inc. All Rights Reserved.
+
+#pragma once
+
+#if WITH_EDITOR
+
+#include "CoreMinimal.h"
+
+namespace UE
+{
+	namespace SVT
+	{
+		struct FTextureData;
+	}
+}
+struct FOpenVDBGridInfo;
+enum class EOpenVDBGridType : uint8;
+
+struct FOpenVDBToSVTConversionResult
+{
+	struct FSparseVolumeAssetHeader* Header;
+	TArray<uint32>* PageTable;
+	TArray<uint8>* PhysicalTileDataA;
+	TArray<uint8>* PhysicalTileDataB;
+};
+
+bool IsOpenVDBGridValid(const FOpenVDBGridInfo& GridInfo, const FString& Filename);
+
+bool SPARSEVOLUMETEXTURE_API GetOpenVDBGridInfo(TArray64<uint8>& SourceFile, bool bCreateStrings, TArray<FOpenVDBGridInfo>* OutGridInfo);
+
+bool SPARSEVOLUMETEXTURE_API ConvertOpenVDBToSparseVolumeTexture(TArray64<uint8>& SourceFile, const struct FOpenVDBImportOptions& ImportOptions, const FIntVector3& VolumeBoundsMin, UE::SVT::FTextureData& OutResult, FTransform& OutFrameTransform);
+
+const TCHAR* OpenVDBGridTypeToString(EOpenVDBGridType Type);
+
+FString SPARSEVOLUMETEXTURE_API GetVDBSequenceBaseFileName(const FString& FileName, bool bDiscardNumbersOnly);
+
+TArray<FString> SPARSEVOLUMETEXTURE_API FindOpenVDBSequenceFileNames(const FString& Filename);
+
+#endif // WITH_EDITOR
\ No newline at end of file
diff --git a/Engine/Source/Editor/SparseVolumeTexture/SparseVolumeTexture.Build.cs b/Engine/Source/Editor/SparseVolumeTexture/SparseVolumeTexture.Build.cs
--- a/Engine/Source/Editor/SparseVolumeTexture/SparseVolumeTexture.Build.cs
+++ b/Engine/Source/Editor/SparseVolumeTexture/SparseVolumeTexture.Build.cs
@@ -55,6 +55,7 @@ public class SparseVolumeTexture : ModuleRules
 		{
 			bUseRTTI = true;
 			bEnableExceptions = true;
+			PublicDefinitions.Add("OPENVDB_AVAILABLE=1");
 
 			AddEngineThirdPartyPrivateStaticDependencies(Target,
 				"IntelTBB",
@@ -68,6 +69,7 @@ public class SparseVolumeTexture : ModuleRules
 		{
 			bUseRTTI = false;
 			bEnableExceptions = false;
+			PublicDefinitions.Add("OPENVDB_AVAILABLE=1");
 
 			AddEngineThirdPartyPrivateStaticDependencies(Target,
 				"IntelTBB",
@@ -77,5 +79,9 @@ public class SparseVolumeTexture : ModuleRules
 				"OpenVDB"
 			);
 		}
+		else
+		{
+			PublicDefinitions.Add("OPENVDB_AVAILABLE=0");
+		}
 	}
 }
