[Interchange Analytics] Add File Creator Metadata to Texture Importer
- Added a way for textures to send back the metadata used for Analytics.

#rnx

diff --git a/Engine/Plugins/Interchange/Runtime/Source/Import/Private/Texture/InterchangeImageWrapperTranslator.cpp b/Engine/Plugins/Interchange/Runtime/Source/Import/Private/Texture/InterchangeImageWrapperTranslator.cpp
--- a/Engine/Plugins/Interchange/Runtime/Source/Import/Private/Texture/InterchangeImageWrapperTranslator.cpp
+++ b/Engine/Plugins/Interchange/Runtime/Source/Import/Private/Texture/InterchangeImageWrapperTranslator.cpp
@@ -224,6 +224,20 @@ TOptional<UE::Interchange::FImportImage> UInterchangeImageWrapperTranslator::Get
 					PayloadData.NumMips = DecompressedImage.MipMapImage.GetMipCount();
 					PayloadData.MipGenSettings = PayloadData.NumMips > 1 ? TextureMipGenSettings::TMGS_LeaveExistingMips : TextureMipGenSettings::TMGS_FromTextureGroup;
 				}
+
+				// Meta Data
+				const bool bHasVendorString = !DecompressedImage.ApplicationVendor.IsEmpty();
+				const bool bHasNameString = !DecompressedImage.ApplicationName.IsEmpty();
+				const bool bHasVersionString = !DecompressedImage.ApplicationVersion.IsEmpty();
+
+				if (bHasVendorString || bHasNameString || bHasVersionString)
+				{
+					UE::Interchange::FTextureCreatorApplicationMetadata Metadata;
+					Metadata.ApplicationVendor = DecompressedImage.ApplicationVendor;
+					Metadata.ApplicationName = DecompressedImage.ApplicationName;
+					Metadata.ApplicationVersion = DecompressedImage.ApplicationVersion;
+					PayloadData.TextureCreatorApplicationMetadata = Metadata;
+				}
 			}
 		}
 		else
diff --git a/Engine/Plugins/Interchange/Runtime/Source/Import/Private/Texture/InterchangeTextureFactory.cpp b/Engine/Plugins/Interchange/Runtime/Source/Import/Private/Texture/InterchangeTextureFactory.cpp
--- a/Engine/Plugins/Interchange/Runtime/Source/Import/Private/Texture/InterchangeTextureFactory.cpp
+++ b/Engine/Plugins/Interchange/Runtime/Source/Import/Private/Texture/InterchangeTextureFactory.cpp
@@ -37,6 +37,7 @@
 #include "Misc/CoreStats.h"
 #include "Misc/Paths.h"
 #include "Nodes/InterchangeBaseNode.h"
+#include "Nodes/InterchangeSourceNode.h"
 #include "ProfilingDebugging/CpuProfilerTrace.h"
 #include "Serialization/EditorBulkData.h"
 #include "Texture/InterchangeBlockedTexturePayloadInterface.h"
@@ -1867,6 +1868,50 @@ namespace UE::Interchange::Private::InterchangeTextureFactory
 		// at this point the texture is mostly set up
 		// NormalMapIdentification is not run yet
 	}
+
+	void ExtractCreatorApplicationMetadata(const FTexturePayloadVariant& TexturePayloadVarient, FTextureCreatorApplicationMetadata& OutTextureCreatorApplicationMetadata)
+	{
+		if (TexturePayloadVarient.IsType<FEmptyVariantState>())
+		{
+			return;
+		}
+
+		if (TexturePayloadVarient.IsType<TOptional<FImportImage>>())
+		{
+			if (const TOptional<FImportImage>* Payload = TexturePayloadVarient.TryGet<TOptional<FImportImage>>(); Payload->IsSet())
+			{
+				const FImportImage& ImagePayload = Payload->GetValue();
+				if (ImagePayload.TextureCreatorApplicationMetadata.IsSet())
+				{
+					OutTextureCreatorApplicationMetadata = ImagePayload.TextureCreatorApplicationMetadata.GetValue();
+				}
+			}
+		}
+
+		if (TexturePayloadVarient.IsType<TOptional<FImportBlockedImage>>())
+		{
+			if (const TOptional<FImportBlockedImage>* Payload = TexturePayloadVarient.TryGet<TOptional<FImportBlockedImage>>(); Payload->IsSet())
+			{
+				const FImportBlockedImage& ImagePayload = Payload->GetValue();
+				if (ImagePayload.TextureCreatorApplicationMetadata.IsSet())
+				{
+					OutTextureCreatorApplicationMetadata = ImagePayload.TextureCreatorApplicationMetadata.GetValue();
+				}
+			}
+		}
+
+		if (TexturePayloadVarient.IsType<TOptional<FImportSlicedImage>>())
+		{
+			if (const TOptional<FImportSlicedImage>* Payload = TexturePayloadVarient.TryGet<TOptional<FImportSlicedImage>>(); Payload->IsSet())
+			{
+				const FImportSlicedImage& ImagePayload = Payload->GetValue();
+				if (ImagePayload.TextureCreatorApplicationMetadata.IsSet())
+				{
+					OutTextureCreatorApplicationMetadata = ImagePayload.TextureCreatorApplicationMetadata.GetValue();
+				}
+			}
+		}
+	}
  }
 
 UClass* UInterchangeTextureFactory::GetFactoryClass() const
@@ -2055,6 +2100,27 @@ UInterchangeFactoryBase::FImportAssetResult UInterchangeTextureFactory::ImportAs
 		return ImportAssetResult;
 	}
 
+	UE::Interchange::FTextureCreatorApplicationMetadata OutTextureCreatorApplicationMetadata;
+	ExtractCreatorApplicationMetadata(TexturePayload, OutTextureCreatorApplicationMetadata);
+	if (UInterchangeSourceNode* SourceNode = UInterchangeSourceNode::FindOrCreateUniqueInstance(const_cast<UInterchangeBaseNodeContainer*>(Arguments.NodeContainer)))
+	{
+		using namespace UE::Interchange;
+		if (!OutTextureCreatorApplicationMetadata.ApplicationVendor.IsEmpty())
+		{
+			SourceNode->SetExtraInformation(FSourceNodeExtraInfoStaticData::GetApplicationVendorExtraInfoKey(), OutTextureCreatorApplicationMetadata.ApplicationVendor);
+		}
+
+		if (!OutTextureCreatorApplicationMetadata.ApplicationName.IsEmpty())
+		{
+			SourceNode->SetExtraInformation(FSourceNodeExtraInfoStaticData::GetApplicationNameExtraInfoKey(), OutTextureCreatorApplicationMetadata.ApplicationName);
+		}
+
+		if (!OutTextureCreatorApplicationMetadata.ApplicationVersion.IsEmpty())
+		{
+			SourceNode->SetExtraInformation(FSourceNodeExtraInfoStaticData::GetApplicationVersionExtraInfoKey(), OutTextureCreatorApplicationMetadata.ApplicationVersion);
+		}
+	}
+
 	UTexture* Texture = nullptr;
 	// create a new texture or overwrite existing asset, if possible
 	if (!ExistingAsset)
diff --git a/Engine/Plugins/Interchange/Runtime/Source/Import/Public/Texture/InterchangeBlockedTexturePayloadData.h b/Engine/Plugins/Interchange/Runtime/Source/Import/Public/Texture/InterchangeBlockedTexturePayloadData.h
--- a/Engine/Plugins/Interchange/Runtime/Source/Import/Public/Texture/InterchangeBlockedTexturePayloadData.h
+++ b/Engine/Plugins/Interchange/Runtime/Source/Import/Public/Texture/InterchangeBlockedTexturePayloadData.h
@@ -22,7 +22,7 @@ namespace UE
 			TextureCompressionSettings CompressionSettings = TC_Default;
 			bool bSRGB = true;
 			TOptional<TextureMipGenSettings> MipGenSettings;
-
+			TOptional<FTextureCreatorApplicationMetadata> TextureCreatorApplicationMetadata;
 
 			int64 ComputeBufferSize() const;
 			bool IsValid() const;
diff --git a/Engine/Plugins/Interchange/Runtime/Source/Import/Public/Texture/InterchangeTexturePayloadData.h b/Engine/Plugins/Interchange/Runtime/Source/Import/Public/Texture/InterchangeTexturePayloadData.h
--- a/Engine/Plugins/Interchange/Runtime/Source/Import/Public/Texture/InterchangeTexturePayloadData.h
+++ b/Engine/Plugins/Interchange/Runtime/Source/Import/Public/Texture/InterchangeTexturePayloadData.h
@@ -11,6 +11,13 @@ namespace UE
 {
 	namespace Interchange
 	{
+		struct INTERCHANGEIMPORT_API FTextureCreatorApplicationMetadata
+		{
+			FString ApplicationVendor;
+			FString ApplicationName;
+			FString ApplicationVersion;
+		};
+
 		struct INTERCHANGEIMPORT_API FImportImage
 		{
 			virtual ~FImportImage() = default;
@@ -33,6 +40,7 @@ namespace UE
 			int32 SizeY = 0;
 			bool bSRGB = true;
 			TOptional<TextureMipGenSettings> MipGenSettings;
+			TOptional<FTextureCreatorApplicationMetadata> TextureCreatorApplicationMetadata;
 
 			void Init2DWithParams(int32 InSizeX, int32 InSizeY, ETextureSourceFormat InFormat, bool bInSRGB, bool bShouldAllocateRawData = true);
 			void Init2DWithParams(int32 InSizeX, int32 InSizeY, int32 InNumMips, ETextureSourceFormat InFormat, bool bInSRGB, bool bShouldAllocateRawData = true);
diff --git a/Engine/Source/Runtime/ImageWrapper/Private/Formats/TiffImageWrapper.cpp b/Engine/Source/Runtime/ImageWrapper/Private/Formats/TiffImageWrapper.cpp
--- a/Engine/Source/Runtime/ImageWrapper/Private/Formats/TiffImageWrapper.cpp
+++ b/Engine/Source/Runtime/ImageWrapper/Private/Formats/TiffImageWrapper.cpp
@@ -458,7 +458,8 @@ namespace UE::ImageWrapper::Private
 		}
 
 		// We just want mip 0
-		const bool bUncompressResult = Uncompress_Internal(InFormat, InBitDepth);
+		TMap<FString, FString> OutTagMetaData;
+		const bool bUncompressResult = Uncompress_Internal(InFormat, InBitDepth, OutTagMetaData);
 		if (bUncompressResult && SubImageBuffer.Num())
 		{
 			// Append the current sub-image to the base byte buffer that reads all the 
@@ -488,7 +489,13 @@ namespace UE::ImageWrapper::Private
 
 		do
 		{
-			const bool bUncompressResult = Uncompress_Internal(InFormat, InBitDepth);
+			TMap<FString, FString> OutTagMetaData;
+			const bool bUncompressResult = Uncompress_Internal(InFormat, InBitDepth, OutTagMetaData);
+			if (OutDecompressedImage.ApplicationName.IsEmpty())
+			{
+				OutDecompressedImage.ApplicationName = OutTagMetaData.FindRef(TEXT("Software"),TEXT(""));
+			}
+
 			if (bUncompressResult && SubImageBuffer.Num())
 			{
 				static bool bExactMatch = false;
@@ -510,12 +517,34 @@ namespace UE::ImageWrapper::Private
 		CurrSubImageHeight = 0;
 	}
 
-	bool FTiffImageWrapper::Uncompress_Internal(const ERGBFormat InFormat, int32 InBitDepth)
+	bool FTiffImageWrapper::Uncompress_Internal(const ERGBFormat InFormat, int32 InBitDepth, TMap<FString, FString>& OutTagMetaData)
 	{
 		// SubImageBuffer gets emptied in the respective calls to UnpackIntoRawBuffer
 		TIFFGetField(Tiff, TIFFTAG_IMAGEWIDTH, &CurrSubImageWidth);
 		TIFFGetField(Tiff, TIFFTAG_IMAGELENGTH, &CurrSubImageHeight);
 
+		const int32 TagCount = TIFFGetTagListCount(Tiff);
+		OutTagMetaData.Reserve(TagCount);
+		for (int32 TagCountIndex = 0; TagCountIndex < TagCount; ++TagCountIndex)
+		{
+			uint32 TagListEntry = TIFFGetTagListEntry(Tiff, TagCountIndex);
+			const TIFFField* TiffField = TIFFFieldWithTag(Tiff, TagListEntry);
+			if (!TiffField)
+			{
+				continue;
+			}
+
+			if (TIFFFieldDataType(TiffField) == TIFF_ASCII)
+			{
+				void* TagValueRawData = nullptr;
+				if (TIFFGetField(Tiff, TagListEntry, &TagValueRawData) != 1)
+				{
+					continue;
+				}
+				OutTagMetaData.Emplace(TIFFFieldName(TiffField), (const char*)TagValueRawData);
+			}
+		}
+
 		if (InFormat == Format && InBitDepth == BitDepth)
 		{
 			// Read using RGBA
diff --git a/Engine/Source/Runtime/ImageWrapper/Private/Formats/TiffImageWrapper.h b/Engine/Source/Runtime/ImageWrapper/Private/Formats/TiffImageWrapper.h
--- a/Engine/Source/Runtime/ImageWrapper/Private/Formats/TiffImageWrapper.h
+++ b/Engine/Source/Runtime/ImageWrapper/Private/Formats/TiffImageWrapper.h
@@ -33,7 +33,7 @@ namespace UE::ImageWrapper::Private
 		virtual ERawImageFormat::Type GetSupportedRawFormat(const ERawImageFormat::Type InFormat) const override;
 
 	private:
-		bool Uncompress_Internal(const ERGBFormat InFormat, int32 InBitDepth);
+		bool Uncompress_Internal(const ERGBFormat InFormat, int32 InBitDepth, TMap<FString, FString>& OutTagMetaData);
 		
 		void ReleaseTiffImage();
 
diff --git a/Engine/Source/Runtime/ImageWrapper/Public/ImageWrapperOutputTypes.h b/Engine/Source/Runtime/ImageWrapper/Public/ImageWrapperOutputTypes.h
--- a/Engine/Source/Runtime/ImageWrapper/Public/ImageWrapperOutputTypes.h
+++ b/Engine/Source/Runtime/ImageWrapper/Public/ImageWrapperOutputTypes.h
@@ -10,4 +10,9 @@
 struct FDecompressedImageOutput
 {
 	FMipMapImage MipMapImage;
+
+	/** Meta data */
+	FString ApplicationVendor;
+	FString ApplicationName;
+	FString ApplicationVersion;
 };
\ No newline at end of file
diff --git a/Engine/Source/Runtime/Interchange/Engine/Private/InterchangeManager.cpp b/Engine/Source/Runtime/Interchange/Engine/Private/InterchangeManager.cpp
--- a/Engine/Source/Runtime/Interchange/Engine/Private/InterchangeManager.cpp
+++ b/Engine/Source/Runtime/Interchange/Engine/Private/InterchangeManager.cpp
@@ -214,30 +214,27 @@ namespace UE::Interchange::Private
 		return FJsonFragment(MoveTemp(ReturnValue));
 	}
 
-	FJsonFragment AnalyticsPrimarySourceDataMetadataToJsonFragment(const UInterchangeSourceData* InSourceData, const UInterchangeBaseNodeContainer* InBaseNodeContainer)
+	FJsonFragment AnalyticsPrimarySourceDataMetadataToJsonFragment(const UInterchangeSourceData* InSourceData, const UInterchangeSourceNode* InSourceNode)
 	{
 		using namespace UE::Interchange;
 		FString ReturnValue;
 		FAnalyticsJsonWriter JsonWriter(&ReturnValue);
 		
-		if (const UInterchangeSourceNode* SourceNode = UInterchangeSourceNode::GetUniqueInstance(InBaseNodeContainer))
-		{
-			TMap<FString, FString> ExtraInfoMap;
-			SourceNode->GetExtraInformation(ExtraInfoMap);
+		TMap<FString, FString> ExtraInfoMap;
+		InSourceNode->GetExtraInformation(ExtraInfoMap);
 
-			const FString DefaultValue(TEXT("N/A"));
-				
-			const FString ApplicationVendor = ExtraInfoMap.FindRef(FSourceNodeExtraInfoStaticData::GetApplicationVendorExtraInfoKey(), DefaultValue);
-			const FString ApplicationName = ExtraInfoMap.FindRef(FSourceNodeExtraInfoStaticData::GetApplicationNameExtraInfoKey(), DefaultValue);
-			const FString ApplicationVersion = ExtraInfoMap.FindRef(FSourceNodeExtraInfoStaticData::GetApplicationVersionExtraInfoKey(), DefaultValue);
-				
-			JsonWriter.WriteObjectStart();
-			JsonWriter.WriteValue(TEXT("SourceExtension"), FPaths::GetExtension(InSourceData->GetFilename()));
-			JsonWriter.WriteValue(TEXT("ApplicationVendor"), ApplicationVendor);
-			JsonWriter.WriteValue(TEXT("ApplicationName"), ApplicationName);
-			JsonWriter.WriteValue(TEXT("ApplicationVersion"), ApplicationVersion);
-			JsonWriter.WriteObjectEnd();
-		}
+		const FString DefaultValue(TEXT("N/A"));
+
+		const FString ApplicationVendor = ExtraInfoMap.FindRef(FSourceNodeExtraInfoStaticData::GetApplicationVendorExtraInfoKey(), DefaultValue);
+		const FString ApplicationName = ExtraInfoMap.FindRef(FSourceNodeExtraInfoStaticData::GetApplicationNameExtraInfoKey(), DefaultValue);
+		const FString ApplicationVersion = ExtraInfoMap.FindRef(FSourceNodeExtraInfoStaticData::GetApplicationVersionExtraInfoKey(), DefaultValue);
+
+		JsonWriter.WriteObjectStart();
+		JsonWriter.WriteValue(TEXT("SourceExtension"), FPaths::GetExtension(InSourceData->GetFilename()));
+		JsonWriter.WriteValue(TEXT("ApplicationVendor"), ApplicationVendor);
+		JsonWriter.WriteValue(TEXT("ApplicationName"), ApplicationName);
+		JsonWriter.WriteValue(TEXT("ApplicationVersion"), ApplicationVersion);
+		JsonWriter.WriteObjectEnd();
 		JsonWriter.Close();
 		return FJsonFragment(MoveTemp(ReturnValue));
 	}
@@ -565,9 +562,11 @@ void UE::Interchange::FImportAsyncHelper::SendAnalyticImportEndData()
 		return;
 	}
 
+	const UInterchangeBaseNodeContainer* PrimaryBaseNodeContainer = BaseNodeContainers[0].Get();
+	if(const UInterchangeSourceNode* PrimarySourceNode = UInterchangeSourceNode::GetUniqueInstance(PrimaryBaseNodeContainer))
 	{
 		using namespace UE::Interchange::Private;
-		Attribs.Add(FAnalyticsEventAttribute(TEXT("PrimarySourceDataMetadata"), AnalyticsPrimarySourceDataMetadataToJsonFragment(SourceDatas[0], BaseNodeContainers[0].Get())));
+		Attribs.Add(FAnalyticsEventAttribute(TEXT("PrimarySourceDataMetadata"), AnalyticsPrimarySourceDataMetadataToJsonFragment(SourceDatas[0], PrimarySourceNode)));
 	}
 
 	int32 ImportedObjectCount = 0;
