USD Interchange: Added support for GeometryCache import.

diff --git a/Engine/Plugins/Interchange/Runtime/Source/FactoryNodes/InterchangeFactoryNodes.Build.cs b/Engine/Plugins/Interchange/Runtime/Source/FactoryNodes/InterchangeFactoryNodes.Build.cs
--- a/Engine/Plugins/Interchange/Runtime/Source/FactoryNodes/InterchangeFactoryNodes.Build.cs
+++ b/Engine/Plugins/Interchange/Runtime/Source/FactoryNodes/InterchangeFactoryNodes.Build.cs
@@ -25,6 +25,7 @@ namespace UnrealBuildTool.Rules
 			PrivateDependencyModuleNames.AddRange(
 				new string[]
 				{
+					"GeometryCache",
 					"LevelSequence",
 				}
 			);
diff --git a/Engine/Plugins/Interchange/Runtime/Source/FactoryNodes/Private/InterchangeGeometryCacheFactoryNode.cpp b/Engine/Plugins/Interchange/Runtime/Source/FactoryNodes/Private/InterchangeGeometryCacheFactoryNode.cpp
new file mode 100644
--- /dev/null
+++ b/Engine/Plugins/Interchange/Runtime/Source/FactoryNodes/Private/InterchangeGeometryCacheFactoryNode.cpp
@@ -0,0 +1,55 @@
+// Copyright Epic Games, Inc. All Rights Reserved.
+
+#include "InterchangeGeometryCacheFactoryNode.h"
+
+#include UE_INLINE_GENERATED_CPP_BY_NAME(InterchangeGeometryCacheFactoryNode)
+
+UInterchangeGeometryCacheFactoryNode::UInterchangeGeometryCacheFactoryNode()
+{
+	AssetClass = nullptr;
+}
+
+void UInterchangeGeometryCacheFactoryNode::InitializeGeometryCacheNode(const FString& UniqueID, const FString& DisplayLabel, const FString& InAssetClass)
+{
+	bIsNodeClassInitialized = false;
+	InitializeNode(UniqueID, DisplayLabel, EInterchangeNodeContainerType::FactoryData);
+
+	FString OperationName = GetTypeName() + TEXT(".SetAssetClassName");
+	InterchangePrivateNodeBase::SetCustomAttribute<FString>(*Attributes, ClassNameAttributeKey, OperationName, InAssetClass);
+	FillAssetClassFromAttribute();
+}
+
+FString UInterchangeGeometryCacheFactoryNode::GetTypeName() const
+{
+	const FString TypeName = TEXT("GeometryCacheNode");
+	return TypeName;
+}
+
+UClass* UInterchangeGeometryCacheFactoryNode::GetObjectClass() const
+{
+	ensure(bIsNodeClassInitialized);
+	return AssetClass.Get() != nullptr ? AssetClass.Get() : UGeometryCache::StaticClass();
+}
+
+void UInterchangeGeometryCacheFactoryNode::FillAssetClassFromAttribute()
+{
+#if WITH_ENGINE
+	FString OperationName = GetTypeName() + TEXT(".GetAssetClassName");
+	FString ClassName;
+	InterchangePrivateNodeBase::GetCustomAttribute<FString>(*Attributes, ClassNameAttributeKey, OperationName, ClassName);
+	if (ClassName.Equals(UGeometryCache::StaticClass()->GetName()))
+	{
+		AssetClass = UGeometryCache::StaticClass();
+		bIsNodeClassInitialized = true;
+	}
+#endif
+}
+
+bool UInterchangeGeometryCacheFactoryNode::SetNodeClassFromClassAttribute()
+{
+	if (!bIsNodeClassInitialized)
+	{
+		FillAssetClassFromAttribute();
+	}
+	return bIsNodeClassInitialized;
+}
diff --git a/Engine/Plugins/Interchange/Runtime/Source/FactoryNodes/Public/InterchangeGeometryCacheFactoryNode.h b/Engine/Plugins/Interchange/Runtime/Source/FactoryNodes/Public/InterchangeGeometryCacheFactoryNode.h
new file mode 100644
--- /dev/null
+++ b/Engine/Plugins/Interchange/Runtime/Source/FactoryNodes/Public/InterchangeGeometryCacheFactoryNode.h
@@ -0,0 +1,42 @@
+// Copyright Epic Games, Inc. All Rights Reserved.
+
+#pragma once
+
+#include "GeometryCache.h"
+#include "InterchangeMeshFactoryNode.h"
+#include "Nodes/InterchangeFactoryBaseNode.h"
+
+#include "InterchangeGeometryCacheFactoryNode.generated.h"
+
+UCLASS(BlueprintType)
+class INTERCHANGEFACTORYNODES_API UInterchangeGeometryCacheFactoryNode : public UInterchangeMeshFactoryNode
+{
+	GENERATED_BODY()
+
+public:
+	UInterchangeGeometryCacheFactoryNode();
+
+	/**
+	 * Initialize node data.
+	 * @param UniqueID - The unique ID for this node.
+	 * @param DisplayLabel - The name of the node.
+	 * @param InAssetClass - The class the GeometryCache factory will create for this node.
+	 */
+	UFUNCTION(BlueprintCallable, Category = "Interchange | Node | GeometryCache")
+	void InitializeGeometryCacheNode(const FString& UniqueID, const FString& DisplayLabel, const FString& InAssetClass);
+
+	/**
+	 * Return the node type name of the class. This is used when reporting errors.
+	 */
+	virtual FString GetTypeName() const override;
+
+	/** Get the class this node creates. */
+	virtual class UClass* GetObjectClass() const override;
+
+private:
+	virtual void FillAssetClassFromAttribute() override;
+	virtual bool SetNodeClassFromClassAttribute() override;
+
+protected:
+	TSubclassOf<UGeometryCache> AssetClass = nullptr;
+};
diff --git a/Engine/Plugins/Interchange/Runtime/Source/Import/InterchangeImport.Build.cs b/Engine/Plugins/Interchange/Runtime/Source/Import/InterchangeImport.Build.cs
--- a/Engine/Plugins/Interchange/Runtime/Source/Import/InterchangeImport.Build.cs
+++ b/Engine/Plugins/Interchange/Runtime/Source/Import/InterchangeImport.Build.cs
@@ -33,6 +33,7 @@ namespace UnrealBuildTool.Rules
 					"AssetRegistry",
 					"CinematicCamera",
 					"ClothingSystemRuntimeCommon",
+					"GeometryCache",
 					"GLTFCore",
 					"IESFile",
 					"ImageCore",
diff --git a/Engine/Plugins/Interchange/Runtime/Source/Import/Private/InterchangeImportModule.cpp b/Engine/Plugins/Interchange/Runtime/Source/Import/Private/InterchangeImportModule.cpp
--- a/Engine/Plugins/Interchange/Runtime/Source/Import/Private/InterchangeImportModule.cpp
+++ b/Engine/Plugins/Interchange/Runtime/Source/Import/Private/InterchangeImportModule.cpp
@@ -13,6 +13,7 @@
 #include "InterchangeManager.h"
 #include "Material/InterchangeMaterialFactory.h"
 #include "MaterialX/InterchangeMaterialXTranslator.h"
+#include "Mesh/InterchangeGeometryCacheFactory.h"
 #include "Mesh/InterchangeOBJTranslator.h"
 #include "Mesh/InterchangePhysicsAssetFactory.h"
 #include "Mesh/InterchangeSkeletalMeshFactory.h"
@@ -114,6 +115,7 @@ void FInterchangeImportModule::StartupModule()
 		InterchangeManager.RegisterFactory(UInterchangeSkeletonFactory::StaticClass());
 		InterchangeManager.RegisterFactory(UInterchangeSkeletalMeshFactory::StaticClass());
 		InterchangeManager.RegisterFactory(UInterchangeStaticMeshFactory::StaticClass());
+		InterchangeManager.RegisterFactory(UInterchangeGeometryCacheFactory::StaticClass());
 		InterchangeManager.RegisterFactory(UInterchangePhysicsAssetFactory::StaticClass());
 		InterchangeManager.RegisterFactory(UInterchangeActorFactory::StaticClass());
 		InterchangeManager.RegisterFactory(UInterchangeLevelFactory::StaticClass());
diff --git a/Engine/Plugins/Interchange/Runtime/Source/Import/Private/Mesh/InterchangeGeometryCacheFactory.cpp b/Engine/Plugins/Interchange/Runtime/Source/Import/Private/Mesh/InterchangeGeometryCacheFactory.cpp
new file mode 100644
--- /dev/null
+++ b/Engine/Plugins/Interchange/Runtime/Source/Import/Private/Mesh/InterchangeGeometryCacheFactory.cpp
@@ -0,0 +1,498 @@
+// Copyright Epic Games, Inc. All Rights Reserved. 
+
+#include "Mesh/InterchangeGeometryCacheFactory.h"
+
+#include "Async/ParallelFor.h"
+#include "GeometryCache.h"
+#include "GeometryCacheCodecV1.h"
+#include "GeometryCacheHelpers.h"
+#include "GeometryCacheMeshData.h"
+#include "GeometryCacheTrackStreamable.h"
+#include "InterchangeCommonPipelineDataFactoryNode.h"
+#include "InterchangeGeometryCacheFactoryNode.h"
+#include "InterchangeImportCommon.h"
+#include "InterchangeImportLog.h"
+#include "InterchangeManager.h"
+#include "InterchangeSceneNode.h"
+#include "InterchangeSourceData.h"
+#include "Mesh/InterchangeMeshHelper.h"
+#include "Mesh/InterchangeMeshPayloadInterface.h"
+#include "Misc/ScopedSlowTask.h"
+#include "Nodes/InterchangeBaseNode.h"
+#include "Nodes/InterchangeBaseNodeContainer.h"
+
+#include UE_INLINE_GENERATED_CPP_BY_NAME(InterchangeGeometryCacheFactory)
+
+#if WITH_EDITORONLY_DATA
+#include "EditorFramework/AssetImportData.h"
+#endif //WITH_EDITORONLY_DATA
+
+static int32 GGeometryCacheParallelFrameReads = 16;
+static FAutoConsoleVariableRef CVarGeometryCacheParallelFrameReads(
+	TEXT("Interchange.GeometryCache.ParallelFrameReads"),
+	GGeometryCacheParallelFrameReads,
+	TEXT("Maximum number of frames to read in parallel")
+);
+
+UClass* UInterchangeGeometryCacheFactory::GetFactoryClass() const
+{
+	return UGeometryCache::StaticClass();
+}
+
+UInterchangeFactoryBase::FImportAssetResult UInterchangeGeometryCacheFactory::BeginImportAsset_GameThread(const FImportAssetObjectParams& Arguments)
+{
+	TRACE_CPUPROFILER_EVENT_SCOPE(UInterchangeGeometryCacheFactory::BeginImportAsset_GameThread);
+
+	FImportAssetResult ImportAssetResult;
+
+	const IInterchangeMeshPayloadInterface* MeshTranslatorPayloadInterface = Cast<IInterchangeMeshPayloadInterface>(Arguments.Translator);
+	if (!MeshTranslatorPayloadInterface)
+	{
+		UE_LOG(LogInterchangeImport, Error, TEXT("Cannot import geometry cache. The translator does not implement IInterchangeMeshPayloadInterface."));
+		return ImportAssetResult;
+	}
+
+	if (!Arguments.AssetNode || !Arguments.AssetNode->GetObjectClass()->IsChildOf(GetFactoryClass()))
+	{
+		return ImportAssetResult;
+	}
+
+	const UInterchangeGeometryCacheFactoryNode* GeometryCacheFactoryNode = Cast<UInterchangeGeometryCacheFactoryNode>(Arguments.AssetNode);
+	if (GeometryCacheFactoryNode == nullptr)
+	{
+		return ImportAssetResult;
+	}
+
+	UObject* ExistingAsset = Arguments.ReimportObject;
+	if (!ExistingAsset)
+	{
+		FSoftObjectPath ReferenceObject;
+		if (GeometryCacheFactoryNode->GetCustomReferenceObject(ReferenceObject))
+		{
+			ExistingAsset = ReferenceObject.TryLoad();
+		}
+	}
+
+	// Create a new geometry cache or overwrite existing asset, if possible
+	UGeometryCache* GeometryCache = nullptr;
+	if (!ExistingAsset)
+	{
+		GeometryCache = NewObject<UGeometryCache>(Arguments.Parent, *Arguments.AssetName, RF_Public | RF_Standalone);
+	}
+	else
+	{
+		// This is a reimport, we are just re-updating the source data
+		GeometryCache = Cast<UGeometryCache>(ExistingAsset);
+		if (GeometryCache)
+		{
+			GeometryCache->ClearForReimporting();
+		}
+	}
+
+	if (!GeometryCache)
+	{
+		if (!Arguments.ReimportObject)
+		{
+			UE_LOG(LogInterchangeImport, Warning, TEXT("Could not create GeometryCache asset %s."), *Arguments.AssetName);
+		}
+		return ImportAssetResult;
+	}
+
+#if WITH_EDITOR
+	GeometryCache->PreEditChange(nullptr);
+#endif // WITH_EDITOR
+
+	TArray<FString> MeshUids;
+	GeometryCacheFactoryNode->GetTargetNodeUids(MeshUids);
+
+	// Create a track for each mesh to be processed and add it to the GeometryCache
+	for (int32 MeshIndex = 0; MeshIndex < MeshUids.Num(); ++MeshIndex)
+	{
+		const FString& MeshUid = MeshUids[MeshIndex];
+
+		// Extract the mesh name
+		// #ueent_todo: Make this more generic as this is currently based on the USD uids
+		TArray<FString> MeshUidTokens;
+		MeshUid.ParseIntoArray(MeshUidTokens, TEXT("/"));
+		FString MeshName = MeshUidTokens[MeshUidTokens.Num() - 1];
+
+		FName CodecName = MakeUniqueObjectName(GeometryCache, UGeometryCacheCodecV1::StaticClass(), FName(MeshName + FString(TEXT("_Codec"))));
+		UGeometryCacheCodecV1* Codec = NewObject<UGeometryCacheCodecV1>(GeometryCache, CodecName, RF_Public);
+
+#if WITH_EDITORONLY_DATA
+		// #ueent_todo: Expose geometry cache settings
+		// Compression settings for good quality
+		const float VertexQuantizationPrecision = 0.0005f;
+		const int32 UVBits = 16;
+		Codec->InitializeEncoder(VertexQuantizationPrecision, UVBits);
+#endif
+
+		FName TrackName = MakeUniqueObjectName(GeometryCache, UGeometryCacheTrackStreamable::StaticClass(), FName(MeshName));
+		UGeometryCacheTrackStreamable* Track = NewObject<UGeometryCacheTrackStreamable>(GeometryCache, TrackName, RF_Public);
+
+#if WITH_EDITORONLY_DATA
+		const bool bForceSingleOptimization = false;
+		const bool bCalculateMotionVectors = false;
+		const bool bOptimizeIndexBuffers = false;
+		Track->BeginCoding(Codec, bForceSingleOptimization, bCalculateMotionVectors, bOptimizeIndexBuffers);
+		// EndCoding has to be called from the main thread once all the frame data have been added to the track
+#endif
+
+		GeometryCache->AddTrack(Track);
+
+		TArray<FMatrix> Mats;
+		Mats.Add(FMatrix::Identity);
+		Mats.Add(FMatrix::Identity);
+
+		TArray<float> MatTimes;
+		MatTimes.Add(0.0f);
+		MatTimes.Add(0.0f);
+		Track->SetMatrixSamples(Mats, MatTimes);
+	}
+
+	const bool bReimport = Arguments.ReimportObject && GeometryCache;
+
+	//Call the mesh helper to create the missing material and to use the unmatched existing slot with the unmatched import slot
+	{
+		using namespace UE::Interchange::Private::MeshHelper;
+		TMap<FString, FString> SlotMaterialDependencies;
+		GeometryCacheFactoryNode->GetSlotMaterialDependencies(SlotMaterialDependencies);
+		GeometryCacheFactorySetupAssetMaterialArray(GeometryCache->Materials, GeometryCache->MaterialSlotNames, SlotMaterialDependencies, Arguments.NodeContainer, bReimport);
+	}
+
+	FTransform GlobalOffsetTransform = FTransform::Identity;
+	if (UInterchangeCommonPipelineDataFactoryNode* CommonPipelineDataFactoryNode = UInterchangeCommonPipelineDataFactoryNode::GetUniqueInstance(Arguments.NodeContainer))
+	{
+		CommonPipelineDataFactoryNode->GetCustomGlobalOffsetTransform(GlobalOffsetTransform);
+	}
+
+	const int32 MaxWorkerThreads = FTaskGraphInterface::Get().GetNumWorkerThreads();
+	const int32 NumFrameThreads = FMath::Clamp(MaxWorkerThreads, 1, GGeometryCacheParallelFrameReads);
+
+	// For each mesh, retrieve their payloads (frames) and process them into their respective tracks.
+	// The GeometryCache data is processed in the game thread since the payloads are queried in
+	// parallel but need to be processed in order. This allows to minimize memory usage since
+	// only NumFrameThreads amount of data will be in memory at a time.
+
+	int32 StartFrame = MIN_int32;
+	int32 EndFrame = MAX_int32;
+	int32 MaterialOffset = 0;
+	int32 NumMeshes = MeshUids.Num();
+	for (int32 MeshIndex = 0; MeshIndex < NumMeshes; ++MeshIndex)
+	{
+		const FString& MeshUid = MeshUids[MeshIndex];
+
+		FTransform GlobalMeshTransform;
+		const UInterchangeBaseNode* Node = Arguments.NodeContainer->GetNode(MeshUid);
+		const UInterchangeGeometryCacheNode* GeometryCacheNode = Cast<const UInterchangeGeometryCacheNode>(Node);
+		if (GeometryCacheNode == nullptr)
+		{
+			// MeshUid must refer to a scene node
+			const UInterchangeSceneNode* SceneNode = Cast<const UInterchangeSceneNode>(Node);
+			if (!ensure(SceneNode))
+			{
+				UE_LOG(LogInterchangeImport, Warning, TEXT("Invalid mesh reference when importing GeometryCache asset %s."), *Arguments.AssetName);
+				continue;
+			}
+
+			// Get the transform from the scene node
+			FTransform SceneNodeGlobalTransform;
+			if (SceneNode->GetCustomGlobalTransform(Arguments.NodeContainer, GlobalOffsetTransform, SceneNodeGlobalTransform))
+			{
+				GlobalMeshTransform = SceneNodeGlobalTransform;
+			}
+
+			constexpr bool bBakeMeshes = true;
+			constexpr bool bBakePivotMeshes = false;
+			UE::Interchange::Private::MeshHelper::AddSceneNodeGeometricAndPivotToGlobalTransform(GlobalMeshTransform, SceneNode, bBakeMeshes, bBakePivotMeshes);
+
+			// And get the mesh node which it references
+			FString MeshDependencyUid;
+			SceneNode->GetCustomAssetInstanceUid(MeshDependencyUid);
+			GeometryCacheNode = Cast<UInterchangeGeometryCacheNode>(Arguments.NodeContainer->GetNode(MeshDependencyUid));
+		}
+		else
+		{
+			// If we have a mesh that is not reference by a scene node, we must apply the global offset.
+			GlobalMeshTransform = GlobalOffsetTransform;
+		}
+
+		if (!ensure(GeometryCacheNode))
+		{
+			UE_LOG(LogInterchangeImport, Warning, TEXT("Invalid mesh reference when importing GeometryCache asset %s."), *Arguments.AssetName);
+			continue;
+		}
+
+		TOptional<FInterchangeMeshPayLoadKey> OptionalPayLoadKey = GeometryCacheNode->GetPayLoadKey();
+		if (!ensure(OptionalPayLoadKey.IsSet()))
+		{
+			UE_LOG(LogInterchangeImport, Warning, TEXT("Empty mesh reference payload when importing GeometryCache asset %s."), *Arguments.AssetName);
+			continue;
+		}
+
+		int32 MeshStartFrame = 0;
+		GeometryCacheNode->GetCustomStartFrame(MeshStartFrame);
+		StartFrame = FMath::Max(StartFrame, MeshStartFrame);
+
+		int32 MeshEndFrame = 0;
+		GeometryCacheNode->GetCustomEndFrame(MeshEndFrame);
+		EndFrame = FMath::Min(EndFrame, MeshEndFrame);
+
+		double FrameRate = 24.0;
+		GeometryCacheNode->GetCustomFrameRate(FrameRate);
+
+		bool bConstantTopology = true;
+		GeometryCacheNode->GetCustomHasConstantTopology(bConstantTopology);
+
+		FString MeshName = GeometryCacheNode->GetDisplayLabel();
+		FString Title = FString::Format(TEXT("Importing frames for {0} (Mesh {1}/{2})"), { MeshName, MeshIndex + 1,  NumMeshes });
+		FScopedSlowTask SlowTask(static_cast<float>((MeshEndFrame + 1) - MeshStartFrame), FText::FromString(Title));
+		SlowTask.MakeDialog(true);
+
+		FEvent* FrameWrittenEvent = FPlatformProcess::GetSynchEventFromPool();
+		std::atomic<int32> WriteFrameIndex = MeshStartFrame;
+		FCriticalSection Mutex;
+		float CompletedFrames = 0.0f;
+
+		// Parallel frame read: frame data can be read concurrently but have to be processed in order for AddMeshSample
+		ParallelFor(
+			NumFrameThreads,
+			[MeshStartFrame,
+			MeshEndFrame,
+			FrameRate,
+			&OptionalPayLoadKey,
+			&MeshTranslatorPayloadInterface,
+			&GlobalMeshTransform,
+			MaterialOffset,
+			&WriteFrameIndex,
+			FrameWrittenEvent,
+			&Mutex,
+			StreamableTrack = Cast<UGeometryCacheTrackStreamable>(GeometryCache->Tracks[MeshIndex]),
+			bConstantTopology,
+			&SlowTask,
+			&CompletedFrames,
+			NumFrameThreads](int32 FrameThreadIndex)
+			{
+				int32 FrameIndex = MeshStartFrame + FrameThreadIndex;
+
+				while (FrameIndex < MeshEndFrame)
+				{
+					FInterchangeMeshPayLoadKey PayLoadKey(OptionalPayLoadKey->UniqueId, FrameIndex);
+					TOptional<UE::Interchange::FMeshPayloadData> PayloadData = MeshTranslatorPayloadInterface->GetMeshPayloadData(PayLoadKey, GlobalMeshTransform);
+
+					if (PayloadData.IsSet() && !PayloadData->MeshDescription.IsEmpty())
+					{
+						// Compute the normals and tangents for the mesh
+						const float ComparisonThreshold = THRESH_POINTS_ARE_SAME;
+
+						// This function make sure the Polygon Normals Tangents Binormals are computed and also remove degenerated triangle from the render mesh
+						// description.
+						FStaticMeshOperations::ComputeTriangleTangentsAndNormals(PayloadData->MeshDescription, ComparisonThreshold);
+
+						// Compute any missing normals or tangents.
+						// Static meshes always blend normals of overlapping corners.
+						EComputeNTBsFlags ComputeNTBsOptions = EComputeNTBsFlags::BlendOverlappingNormals;
+						ComputeNTBsOptions |= EComputeNTBsFlags::IgnoreDegenerateTriangles;
+						ComputeNTBsOptions |= EComputeNTBsFlags::UseMikkTSpace;
+
+						FStaticMeshOperations::ComputeTangentsAndNormals(PayloadData->MeshDescription, ComputeNTBsOptions);
+
+						// Convert the MeshDescription to MeshData
+						FGeometryCacheMeshData MeshData;
+						UE::GeometryCache::Utils::GetGeometryCacheMeshDataFromMeshDescription(MeshData, PayloadData->MeshDescription, MaterialOffset, FrameRate);
+
+						// Wait until it's our turn to process this frame.
+						while (WriteFrameIndex < FrameIndex)
+						{
+							const uint32 WaitTimeInMs = 10;
+							FrameWrittenEvent->Wait(WaitTimeInMs);
+						}
+
+						{
+							FScopeLock WriteLock(&Mutex);
+
+#if WITH_EDITORONLY_DATA
+							// Add it to the track
+							StreamableTrack->AddMeshSample(MeshData, (FrameIndex - MeshStartFrame) / FrameRate, bConstantTopology);
+#endif
+
+							// Mark the next frame index as ready for processing.
+							++WriteFrameIndex;
+
+							++CompletedFrames;
+							if (IsInGameThread())
+							{
+								SlowTask.EnterProgressFrame(CompletedFrames);
+								CompletedFrames = 0.0f;
+							}
+
+							FrameWrittenEvent->Trigger();
+						}
+					}
+					// Get new frame index to read for next run cycle
+					FrameIndex += NumFrameThreads;
+				}
+			});
+
+		FPlatformProcess::ReturnSynchEventToPool(FrameWrittenEvent);
+
+		++MaterialOffset;
+	}
+
+	GeometryCache->SetFrameStartEnd(StartFrame, EndFrame);
+
+	ImportAssetResult.ImportedObject = GeometryCache;
+	return ImportAssetResult;
+}
+
+UInterchangeFactoryBase::FImportAssetResult UInterchangeGeometryCacheFactory::ImportAsset_Async(const FImportAssetObjectParams& Arguments)
+{
+	TRACE_CPUPROFILER_EVENT_SCOPE(UInterchangeGeometryCacheFactory::ImportAsset_Async);
+
+	FImportAssetResult ImportAssetResult;
+	if (!Arguments.AssetNode || !Arguments.AssetNode->GetObjectClass()->IsChildOf(GetFactoryClass()))
+	{
+		return ImportAssetResult;
+	}
+
+	UInterchangeGeometryCacheFactoryNode* GeometryCacheFactoryNode = Cast<UInterchangeGeometryCacheFactoryNode>(Arguments.AssetNode);
+	if (GeometryCacheFactoryNode == nullptr)
+	{
+		return ImportAssetResult;
+	}
+
+	UObject* GeometryCacheObject = UE::Interchange::FFactoryCommon::AsyncFindObject(GeometryCacheFactoryNode, GetFactoryClass(), Arguments.Parent, Arguments.AssetName);
+	if (!GeometryCacheObject)
+	{
+		UE_LOG(LogInterchangeImport, Error, TEXT("Could not import the GeometryCache asset %s because the asset does not exist."), *Arguments.AssetName);
+		return ImportAssetResult;
+	}
+
+	UGeometryCache* GeometryCache = Cast<UGeometryCache>(GeometryCacheObject);
+	if (!ensure(GeometryCache))
+	{
+		UE_LOG(LogInterchangeImport, Error, TEXT("Could not cast to GeometryCache asset %s."), *Arguments.AssetName);
+		return ImportAssetResult;
+	}
+
+	// Getting the file Hash will cache it into the source data
+	Arguments.SourceData->GetFileContentHash();
+
+	ImportAssetResult.ImportedObject = GeometryCacheObject;
+	return ImportAssetResult;
+}
+
+UInterchangeFactoryBase::FImportAssetResult UInterchangeGeometryCacheFactory::EndImportAsset_GameThread(const FImportAssetObjectParams& Arguments)
+{
+	TRACE_CPUPROFILER_EVENT_SCOPE(UInterchangeGeometryCacheFactory::EndImportAsset_GameThread);
+
+	FImportAssetResult ImportAssetResult;
+	if (!Arguments.AssetNode || !Arguments.AssetNode->GetObjectClass()->IsChildOf(GetFactoryClass()))
+	{
+		return ImportAssetResult;
+	}
+
+	UInterchangeGeometryCacheFactoryNode* GeometryCacheFactoryNode = Cast<UInterchangeGeometryCacheFactoryNode>(Arguments.AssetNode);
+	if (GeometryCacheFactoryNode == nullptr)
+	{
+		return ImportAssetResult;
+	}
+
+	const UClass* GeometryCacheClass = GeometryCacheFactoryNode->GetObjectClass();
+	check(GeometryCacheClass && GeometryCacheClass->IsChildOf(GetFactoryClass()));
+
+	UObject* ExistingAsset = StaticFindObject(nullptr, Arguments.Parent, *Arguments.AssetName);
+
+	UGeometryCache* GeometryCache = Cast<UGeometryCache>(ExistingAsset);
+	if (!ensure(GeometryCache))
+	{
+		UE_LOG(LogInterchangeImport, Error, TEXT("Could not create GeometryCache asset %s."), *Arguments.AssetName);
+		return ImportAssetResult;
+	}
+
+#if WITH_EDITORONLY_DATA
+	// Finalize the coding for all tracks
+	for (int32 Index = 0; Index < GeometryCache->Tracks.Num(); ++Index)
+	{
+		if (UGeometryCacheTrackStreamable* StreamableTrack = Cast<UGeometryCacheTrackStreamable>(GeometryCache->Tracks[Index]))
+		{
+			StreamableTrack->EndCoding();
+		}
+	}
+#endif
+
+	ImportAssetResult.ImportedObject = GeometryCache;
+	return ImportAssetResult;
+}
+
+void UInterchangeGeometryCacheFactory::SetupObject_GameThread(const FSetupObjectParams& Arguments)
+{
+	TRACE_CPUPROFILER_EVENT_SCOPE(UInterchangeGeometryCacheFactory::SetupObject_GameThread);
+
+	Super::SetupObject_GameThread(Arguments);
+
+#if WITH_EDITORONLY_DATA
+	if (ensure(Arguments.ImportedObject && Arguments.SourceData))
+	{
+		// We must call the Update of the asset source file in the main thread because UAssetImportData::Update execute some delegate we do not control
+		UGeometryCache* GeometryCache = CastChecked<UGeometryCache>(Arguments.ImportedObject);
+
+		UE::Interchange::FFactoryCommon::FUpdateImportAssetDataParameters UpdateImportAssetDataParameters(GeometryCache, GeometryCache->AssetImportData, Arguments.SourceData, Arguments.NodeUniqueID, Arguments.NodeContainer, Arguments.OriginalPipelines, Arguments.Translator);
+		GeometryCache->AssetImportData = UE::Interchange::FFactoryCommon::UpdateImportAssetData(UpdateImportAssetDataParameters);
+	}
+#endif
+}
+
+bool UInterchangeGeometryCacheFactory::GetSourceFilenames(const UObject* Object, TArray<FString>& OutSourceFilenames) const
+{
+#if WITH_EDITORONLY_DATA
+	if (const UGeometryCache* GeometryCache = Cast<UGeometryCache>(Object))
+	{
+		return UE::Interchange::FFactoryCommon::GetSourceFilenames(GeometryCache->AssetImportData, OutSourceFilenames);
+	}
+#endif
+
+	return false;
+}
+
+bool UInterchangeGeometryCacheFactory::SetSourceFilename(const UObject* Object, const FString& SourceFilename, int32 SourceIndex) const
+{
+#if WITH_EDITORONLY_DATA
+	if (const UGeometryCache* GeometryCache = Cast<UGeometryCache>(Object))
+	{
+		return UE::Interchange::FFactoryCommon::SetSourceFilename(GeometryCache->AssetImportData, SourceFilename, SourceIndex);
+	}
+#endif
+
+	return false;
+}
+
+void UInterchangeGeometryCacheFactory::BackupSourceData(const UObject* Object) const
+{
+#if WITH_EDITORONLY_DATA
+	if (const UGeometryCache* GeometryCache = Cast<UGeometryCache>(Object))
+	{
+		UE::Interchange::FFactoryCommon::BackupSourceData(GeometryCache->AssetImportData);
+	}
+#endif
+}
+
+void UInterchangeGeometryCacheFactory::ReinstateSourceData(const UObject* Object) const
+{
+#if WITH_EDITORONLY_DATA
+	if (const UGeometryCache* GeometryCache = Cast<UGeometryCache>(Object))
+	{
+		UE::Interchange::FFactoryCommon::ReinstateSourceData(GeometryCache->AssetImportData);
+	}
+#endif
+}
+
+void UInterchangeGeometryCacheFactory::ClearBackupSourceData(const UObject* Object) const
+{
+#if WITH_EDITORONLY_DATA
+	if (const UGeometryCache* GeometryCache = Cast<UGeometryCache>(Object))
+	{
+		UE::Interchange::FFactoryCommon::ClearBackupSourceData(GeometryCache->AssetImportData);
+	}
+#endif
+}
diff --git a/Engine/Plugins/Interchange/Runtime/Source/Import/Private/Mesh/InterchangeMeshHelper.cpp b/Engine/Plugins/Interchange/Runtime/Source/Import/Private/Mesh/InterchangeMeshHelper.cpp
--- a/Engine/Plugins/Interchange/Runtime/Source/Import/Private/Mesh/InterchangeMeshHelper.cpp
+++ b/Engine/Plugins/Interchange/Runtime/Source/Import/Private/Mesh/InterchangeMeshHelper.cpp
@@ -147,6 +147,64 @@ namespace UE::Interchange::Private::MeshHelper
 		int32 MaterialIndex = INDEX_NONE;
 	};
 
+	template<>
+	class FMeshMaterialViewer<TObjectPtr<UMaterialInterface>>
+	{
+	public:
+		FMeshMaterialViewer(TArray<TObjectPtr<UMaterialInterface>>& InMaterials, TArray<FName>& InMaterialSlotNames, int32 InMaterialIndex)
+			: Materials(InMaterials)
+			, MaterialSlotNames(InMaterialSlotNames)
+			, MaterialIndex(InMaterialIndex)
+		{ }
+
+		FName GetMaterialSlotName() const
+		{
+			if (Materials.IsValidIndex(MaterialIndex))
+			{
+				return MaterialSlotNames[MaterialIndex];
+			}
+			return NAME_None;
+		}
+
+		FName GetImportedMaterialSlotName() const
+		{
+			return GetMaterialSlotName();
+		}
+
+		void SetMaterialSlotName(const FName MaterialSlotName)
+		{
+			if (Materials.IsValidIndex(MaterialIndex))
+			{
+				MaterialSlotNames[MaterialIndex] = MaterialSlotName;
+			}
+		}
+
+		void SetImportedMaterialSlotName(const FName ImportedMaterialSlotName)
+		{
+		}
+
+		UMaterialInterface* GetMaterialInterface() const
+		{
+			if (Materials.IsValidIndex(MaterialIndex))
+			{
+				return Materials[MaterialIndex];
+			}
+			return nullptr;
+		}
+
+		void SetMaterialInterface(UMaterialInterface* MaterialInterface)
+		{
+			if (Materials.IsValidIndex(MaterialIndex))
+			{
+				Materials[MaterialIndex] = MaterialInterface;
+			}
+		}
+	private:
+		TArray<TObjectPtr<UMaterialInterface>>& Materials;
+		TArray<FName>& MaterialSlotNames;
+		int32 MaterialIndex = INDEX_NONE;
+	};
+
 	template<typename MaterialType>
 	class INTERCHANGEIMPORT_API FMeshMaterialArrayViewer
 	{
@@ -210,6 +268,71 @@ namespace UE::Interchange::Private::MeshHelper
 		TArray<FMeshMaterialViewer<MaterialType>> MeshMaterialArrayViewer;
 	};
 
+	template<>
+	class FMeshMaterialArrayViewer<TObjectPtr<UMaterialInterface>>
+	{
+	public:
+		FMeshMaterialArrayViewer(TArray<TObjectPtr<UMaterialInterface>>& InMaterials, TArray<FName>& InMaterialSlotNames, TFunction<void(TObjectPtr<UMaterialInterface>& Material)>& InEmplaceMaterialFunctor)
+			: Materials(InMaterials)
+			, MaterialSlotNames(InMaterialSlotNames)
+			, EmplaceMaterialFunctor(InEmplaceMaterialFunctor)
+		{
+			RebuildViewer();
+		}
+
+		void RebuildViewer()
+		{
+			MeshMaterialArrayViewer.Reset(Materials.Num());
+			for (int32 MaterialIndex = 0; MaterialIndex < Materials.Num(); ++MaterialIndex)
+			{
+				MeshMaterialArrayViewer.Emplace(Materials, MaterialSlotNames, MaterialIndex);
+			}
+		}
+
+		int32 Num() const
+		{
+			return MeshMaterialArrayViewer.Num();
+		}
+
+		FMeshMaterialViewer<TObjectPtr<UMaterialInterface>>& operator[](int32 MaterialIndex)
+		{
+			check(MeshMaterialArrayViewer.IsValidIndex(MaterialIndex));
+			return MeshMaterialArrayViewer[MaterialIndex];
+		}
+
+		FMeshMaterialViewer<TObjectPtr<UMaterialInterface>>* FindByPredicate(TFunction<bool(const FMeshMaterialViewer<TObjectPtr<UMaterialInterface>>& MaterialViewer)> Predicate)
+		{
+			return MeshMaterialArrayViewer.FindByPredicate(Predicate);
+		}
+
+		void Emplace(UMaterialInterface* NewMaterial, FName InMaterialSlotName, FName ImportedMaterialSlotName)
+		{
+			TObjectPtr<UMaterialInterface>& Material = Materials.AddDefaulted_GetRef();
+			Material = NewMaterial;
+
+			FName& MaterialSlotName = MaterialSlotNames.AddDefaulted_GetRef();
+			MaterialSlotName = InMaterialSlotName;
+
+			EmplaceMaterialFunctor(Material);
+
+			RebuildViewer();
+		}
+
+		void Reserve(int32 Count)
+		{
+			Materials.Reserve(Count);
+			MaterialSlotNames.Reserve(Count);
+
+			RebuildViewer();
+		}
+
+	private:
+		TArray<TObjectPtr<UMaterialInterface>>& Materials;
+		TArray<FName>& MaterialSlotNames;
+		TFunction<void(TObjectPtr<UMaterialInterface>& Material)>& EmplaceMaterialFunctor;
+		TArray<FMeshMaterialViewer<TObjectPtr<UMaterialInterface>>> MeshMaterialArrayViewer;
+	};
+
 	template<typename MaterialType>
 	void InternalMeshFactorySetupAssetMaterialArray(FMeshMaterialArrayViewer<MaterialType>& ExistMaterialArrayViewer
 		, TMap<FString, FString> ImportedSlotMaterialDependencies
@@ -367,4 +490,15 @@ namespace UE::Interchange::Private::MeshHelper
 		InternalMeshFactorySetupAssetMaterialArray<FStaticMaterial>(MeshMaterialArrayViewer, ImportedSlotMaterialDependencies, NodeContainer, bIsReimport);
 	}
 
+	void GeometryCacheFactorySetupAssetMaterialArray(TArray<TObjectPtr<UMaterialInterface>>& ExistMaterials
+		, TArray<FName>& InMaterialSlotNames
+		, TMap<FString, FString> ImportedSlotMaterialDependencies
+		, const UInterchangeBaseNodeContainer* NodeContainer
+		, const bool bIsReimport)
+	{
+		TFunction<void(TObjectPtr<UMaterialInterface>&)> EmplaceMaterialFunctor = [](TObjectPtr<UMaterialInterface>& Material) {};
+		FMeshMaterialArrayViewer<TObjectPtr<UMaterialInterface>> MeshMaterialArrayViewer(ExistMaterials, InMaterialSlotNames, EmplaceMaterialFunctor);
+		InternalMeshFactorySetupAssetMaterialArray<TObjectPtr<UMaterialInterface>>(MeshMaterialArrayViewer, ImportedSlotMaterialDependencies, NodeContainer, bIsReimport);
+	}
+
 } //ns UE::Interchange::Private::MeshHelper
\ No newline at end of file
diff --git a/Engine/Plugins/Interchange/Runtime/Source/Import/Private/Usd/InterchangeUsdTranslator.cpp b/Engine/Plugins/Interchange/Runtime/Source/Import/Private/Usd/InterchangeUsdTranslator.cpp
--- a/Engine/Plugins/Interchange/Runtime/Source/Import/Private/Usd/InterchangeUsdTranslator.cpp
+++ b/Engine/Plugins/Interchange/Runtime/Source/Import/Private/Usd/InterchangeUsdTranslator.cpp
@@ -1046,10 +1046,8 @@ namespace UE::InterchangeUsdTranslator::Private
 					*Prim.GetPrimPath().GetString(),
 					*UsdUtils::StringifyMap(PrimvarToUVIndex),
 					Value.Primvar.IsEmpty()
-						? TEXT(
-							  " Is your UsdUVTexture Shader missing the 'inputs:st' attribute? (It specifies which UV set to sample the texture with)"
-						  )
-						: TEXT("")
+					? TEXT(" Is your UsdUVTexture Shader missing the 'inputs:st' attribute? (It specifies which UV set to sample the texture with)")
+					: TEXT("")
 				);
 			}
 
@@ -1668,7 +1666,7 @@ namespace UE::InterchangeUsdTranslator::Private
 		const FString MeshPrimPath = MeshPrim.GetPrimPath().GetString();
 
 		TFunction<void(const FString&, int32, const FString&)> AddMorphTargetNode =
-			[&MeshNode, &MeshPrimPath, &NodeContainer, &Info](const FString& MorphTargetName, int32 BlendShapeIndex, const FString& InbetweenName)
+		[&MeshNode, &MeshPrimPath, &NodeContainer, &Info](const FString& MorphTargetName, int32 BlendShapeIndex, const FString& InbetweenName)
 		{
 			// Note: We identify a blend shape by its Mesh prim path and the blend shape index, even though
 			// the blend shape itself is a full standalone prim. This is for two reasons:
@@ -1731,7 +1729,8 @@ namespace UE::InterchangeUsdTranslator::Private
 		}
 
 		// Fill in the MeshNode itself
-		UInterchangeMeshNode* MeshNode = NewObject<UInterchangeMeshNode>(&NodeContainer);
+		const bool bIsAnimated = UsdUtils::IsAnimatedMesh(Prim);
+		UInterchangeMeshNode* MeshNode = bIsAnimated ? NewObject<UInterchangeGeometryCacheNode>(&NodeContainer) : NewObject<UInterchangeMeshNode>(&NodeContainer);
 		MeshNode->InitializeNode(NodeUid, NodeName, EInterchangeNodeContainerType::TranslatedAsset);
 		MeshNode->SetAssetName(NodeName);
 		const bool bIsSkinned = static_cast<bool>(Info.ClosestParentSkelRootPath) && Prim.HasAPI(TEXT("SkelBindingAPI"));
@@ -1746,7 +1745,7 @@ namespace UE::InterchangeUsdTranslator::Private
 			// We are currently not supporting Skinned Primitive Shapes.
 			// In theory a Skinned Mesh needs Joint influences and weights information provided,
 			// However, there does not seem to be a ruleset against a Primitive Shape having SkelBindingAPI set.
-			// Which means, that on a theoretical level there could be such a scenario.
+			// Which means, that on a theoretical level there could be such a scenario. 
 			ensureMsgf(!bIsSkinned, TEXT("Unexpected scenario: Primitive Shape is skinned."));
 
 			// For Primitive Shapes we add PrimitiveShapePrefix for the PayLoadKey,
@@ -1755,7 +1754,36 @@ namespace UE::InterchangeUsdTranslator::Private
 			PayloadKey = PrimitiveShapePrefix + PayloadKey;
 		}
 
-		if (bIsSkinned && !bPrimitiveShape)
+		double TimeCode = UsdUtils::GetDefaultTimeCode();
+		if (UInterchangeGeometryCacheNode* GeometryCacheNode = Cast<UInterchangeGeometryCacheNode>(MeshNode))
+		{
+			GeometryCacheNode->SetPayLoadKey(PrimPath, EInterchangeMeshPayLoadType::ANIMATED);
+
+			UE::FUsdStage Stage = Prim.GetStage();
+			int32 StartFrame = FMath::FloorToInt(Stage.GetStartTimeCode());
+			int32 EndFrame = FMath::CeilToInt(Stage.GetEndTimeCode());
+			UsdUtils::GetAnimatedMeshTimeCodes(Stage, PrimPath, StartFrame, EndFrame);
+
+			double FramesPerSecond = Stage.GetTimeCodesPerSecond();
+			if (FramesPerSecond == 0)
+			{
+				UE_LOG(LogUsd, Warning, TEXT("Invalid USD GeometryCache FPS detected. Falling back to 1 FPS"));
+				FramesPerSecond = 1;
+			}
+
+			// The GeometryCache module expects the end frame to be one past the last animation frame
+			EndFrame += 1;
+
+			GeometryCacheNode->SetCustomStartFrame(StartFrame);
+			GeometryCacheNode->SetCustomEndFrame(EndFrame);
+			GeometryCacheNode->SetCustomFrameRate(FramesPerSecond);
+
+			const bool bConstantTopology = UsdUtils::GetMeshTopologyVariance(Prim) != UsdUtils::EMeshTopologyVariance::Heterogenous;
+			GeometryCacheNode->SetCustomHasConstantTopology(bConstantTopology);
+
+			TimeCode = StartFrame;
+		}
+		else if (bIsSkinned && !bPrimitiveShape)
 		{
 			MeshNode->SetSkinnedMesh(true);
 			MeshNode->SetPayLoadKey(PayloadKey, EInterchangeMeshPayLoadType::SKELETAL);
@@ -1780,7 +1808,6 @@ namespace UE::InterchangeUsdTranslator::Private
 
 		// Material assignments
 		{
-			const double TimeCode = UsdUtils::GetDefaultTimeCode();
 			const bool bProvideMaterialIndices = false;
 			UsdUtils::FUsdPrimMaterialAssignmentInfo Assignments = UsdUtils::GetPrimMaterialAssignments(
 				Prim,
@@ -2551,8 +2578,8 @@ namespace UE::InterchangeUsdTranslator::Private
 		}
 		else if (Prim.IsA(TEXT("Gprim")) && !Prim.IsA(TEXT("PointBased")) && !Prim.IsA(TEXT("Volume")))
 		{
-			// PointBased prims are currently not supported apart from Meshes (which is taken care of in a previous branch).
-			// Volumes are also not currently supported.
+			//PointBased prims are currently not supported apart from Meshes (which is taken care of in a previous branch).
+			//Volumes are also not currently supported.
 
 			Prefix = &PrimitiveShapePrefix;
 			constexpr bool bPrimitiveShape_True = true;
@@ -2845,7 +2872,7 @@ namespace UE::InterchangeUsdTranslator::Private
 			// to resolve all the bindings. The cached assignments come from AddMeshNode step, where
 			// we switch the active variant on the current stage and so get nice material bindings that
 			// resolve normally.
-			//
+		
 			// Note that we can't even use the info cache here, because it wouldn't have cached info
 			// about the inactive LOD variants
 			if (UsdUtils::FUsdPrimMaterialAssignmentInfo* CachedInfo = Impl.CachedMaterialAssignments.Find(PrimPath))
@@ -3591,6 +3618,18 @@ TOptional<UE::Interchange::FMeshPayloadData> UInterchangeUSDTranslator::GetMeshP
 			);
 			break;
 		}
+		case EInterchangeMeshPayLoadType::ANIMATED:
+		{
+			OptionsCopy.TimeCode = PayloadKey.FrameNumber;
+
+			bSuccess = UE::InterchangeUsdTranslator::Private::GetStaticMeshPayloadData(
+				PayloadKey.UniqueId,
+				*ImplPtr,
+				OptionsCopy,
+				MeshPayloadData.MeshDescription
+			);
+			break;
+		}
 		case EInterchangeMeshPayLoadType::NONE:	   // Fallthrough
 		default:
 			break;
diff --git a/Engine/Plugins/Interchange/Runtime/Source/Import/Public/Mesh/InterchangeGeometryCacheFactory.h b/Engine/Plugins/Interchange/Runtime/Source/Import/Public/Mesh/InterchangeGeometryCacheFactory.h
new file mode 100644
--- /dev/null
+++ b/Engine/Plugins/Interchange/Runtime/Source/Import/Public/Mesh/InterchangeGeometryCacheFactory.h
@@ -0,0 +1,32 @@
+// Copyright Epic Games, Inc. All Rights Reserved.
+
+#pragma once
+
+#include "InterchangeFactoryBase.h"
+
+#include "InterchangeGeometryCacheFactory.generated.h"
+
+UCLASS(BlueprintType)
+class INTERCHANGEIMPORT_API UInterchangeGeometryCacheFactory : public UInterchangeFactoryBase
+{
+	GENERATED_BODY()
+public:
+
+	//////////////////////////////////////////////////////////////////////////
+	// Interchange factory base interface begin
+
+	virtual UClass* GetFactoryClass() const override;
+	virtual EInterchangeFactoryAssetType GetFactoryAssetType() override { return EInterchangeFactoryAssetType::Meshes; }
+	virtual FImportAssetResult BeginImportAsset_GameThread(const FImportAssetObjectParams& Arguments) override;
+	virtual FImportAssetResult ImportAsset_Async(const FImportAssetObjectParams& Arguments) override;
+	virtual FImportAssetResult EndImportAsset_GameThread(const FImportAssetObjectParams& Arguments) override;
+	virtual void SetupObject_GameThread(const FSetupObjectParams& Arguments) override;
+	virtual bool GetSourceFilenames(const UObject* Object, TArray<FString>& OutSourceFilenames) const override;
+	virtual bool SetSourceFilename(const UObject* Object, const FString& SourceFilename, int32 SourceIndex) const override;
+	virtual void BackupSourceData(const UObject* Object) const override;
+	virtual void ReinstateSourceData(const UObject* Object) const override;
+	virtual void ClearBackupSourceData(const UObject* Object) const override;
+
+	// Interchange factory base interface end
+	//////////////////////////////////////////////////////////////////////////
+};
diff --git a/Engine/Plugins/Interchange/Runtime/Source/Import/Public/Mesh/InterchangeMeshHelper.h b/Engine/Plugins/Interchange/Runtime/Source/Import/Public/Mesh/InterchangeMeshHelper.h
--- a/Engine/Plugins/Interchange/Runtime/Source/Import/Public/Mesh/InterchangeMeshHelper.h
+++ b/Engine/Plugins/Interchange/Runtime/Source/Import/Public/Mesh/InterchangeMeshHelper.h
@@ -6,6 +6,7 @@ struct FSkeletalMaterial;
 struct FStaticMaterial;
 class UInterchangeBaseNodeContainer;
 class UInterchangeSceneNode;
+class UMaterialInterface;
 
 namespace UE::Interchange::Private::MeshHelper
 {
@@ -24,4 +25,10 @@ namespace UE::Interchange::Private::MeshHelper
 		, const UInterchangeBaseNodeContainer* NodeContainer
 		, const bool bIsReimport);
 
+	INTERCHANGEIMPORT_API void GeometryCacheFactorySetupAssetMaterialArray(TArray<TObjectPtr<UMaterialInterface>>& ExistMaterials
+		, TArray<FName>& MaterialSlotNames
+		, TMap<FString, FString> ImportedSlotMaterialDependencies
+		, const UInterchangeBaseNodeContainer* NodeContainer
+		, const bool bIsReimport);
+
 } //ns UE::Interchange::Private::MeshHelper
\ No newline at end of file
diff --git a/Engine/Plugins/Interchange/Runtime/Source/Nodes/Private/InterchangeMeshNode.cpp b/Engine/Plugins/Interchange/Runtime/Source/Nodes/Private/InterchangeMeshNode.cpp
--- a/Engine/Plugins/Interchange/Runtime/Source/Nodes/Private/InterchangeMeshNode.cpp
+++ b/Engine/Plugins/Interchange/Runtime/Source/Nodes/Private/InterchangeMeshNode.cpp
@@ -538,3 +538,46 @@ bool UInterchangeMeshNode::RemoveSlotMaterialDependencyUid(const FString& SlotNa
 	return SlotMaterialDependencies.RemoveKey(SlotName);
 }
 
+UInterchangeGeometryCacheNode::UInterchangeGeometryCacheNode()
+{
+}
+
+bool UInterchangeGeometryCacheNode::GetCustomStartFrame(int32& AttributeValue) const
+{
+	IMPLEMENT_NODE_ATTRIBUTE_GETTER(StartFrame, int32);
+}
+
+bool UInterchangeGeometryCacheNode::SetCustomStartFrame(const int32& AttributeValue)
+{
+	IMPLEMENT_NODE_ATTRIBUTE_SETTER_NODELEGATE(StartFrame, int32);
+}
+
+bool UInterchangeGeometryCacheNode::GetCustomEndFrame(int32& AttributeValue) const
+{
+	IMPLEMENT_NODE_ATTRIBUTE_GETTER(EndFrame, int32);
+}
+
+bool UInterchangeGeometryCacheNode::SetCustomEndFrame(const int32& AttributeValue)
+{
+	IMPLEMENT_NODE_ATTRIBUTE_SETTER_NODELEGATE(EndFrame, int32);
+}
+
+bool UInterchangeGeometryCacheNode::GetCustomFrameRate(double& AttributeValue) const
+{
+	IMPLEMENT_NODE_ATTRIBUTE_GETTER(FrameRate, double);
+}
+
+bool UInterchangeGeometryCacheNode::SetCustomFrameRate(const double& AttributeValue)
+{
+	IMPLEMENT_NODE_ATTRIBUTE_SETTER_NODELEGATE(FrameRate, double);
+}
+
+bool UInterchangeGeometryCacheNode::GetCustomHasConstantTopology(bool& AttributeValue) const
+{
+	IMPLEMENT_NODE_ATTRIBUTE_GETTER(HasConstantTopology, bool);
+}
+
+bool UInterchangeGeometryCacheNode::SetCustomHasConstantTopology(const bool& AttributeValue)
+{
+	IMPLEMENT_NODE_ATTRIBUTE_SETTER_NODELEGATE(HasConstantTopology, bool);
+}
diff --git a/Engine/Plugins/Interchange/Runtime/Source/Nodes/Public/InterchangeMeshNode.h b/Engine/Plugins/Interchange/Runtime/Source/Nodes/Public/InterchangeMeshNode.h
--- a/Engine/Plugins/Interchange/Runtime/Source/Nodes/Public/InterchangeMeshNode.h
+++ b/Engine/Plugins/Interchange/Runtime/Source/Nodes/Public/InterchangeMeshNode.h
@@ -35,7 +35,8 @@ enum class EInterchangeMeshPayLoadType : uint8
 	NONE = 0,
 	STATIC,
 	SKELETAL,
-	MORPHTARGET
+	MORPHTARGET,
+	ANIMATED
 };
 
 
@@ -52,6 +53,9 @@ public:
 	UPROPERTY(EditAnywhere, BlueprintReadWrite, Category = "Interchange | Mesh")
 	EInterchangeMeshPayLoadType Type = EInterchangeMeshPayLoadType::NONE;
 
+	UPROPERTY(EditAnywhere, BlueprintReadWrite, Category = "Interchange | GeometryCache")
+	int32 FrameNumber = 0;
+
 	FInterchangeMeshPayLoadKey() {}
 
 	FInterchangeMeshPayLoadKey(const FString& InUniqueId, const EInterchangeMeshPayLoadType& InType)
@@ -60,9 +64,16 @@ public:
 	{
 	}
 
+	FInterchangeMeshPayLoadKey(const FString& InUniqueId, int32 InFrameNumber)
+		: UniqueId(InUniqueId)
+		, Type(EInterchangeMeshPayLoadType::ANIMATED)
+		, FrameNumber(InFrameNumber)
+	{
+	}
+
 	bool operator==(const FInterchangeMeshPayLoadKey& Other) const
 	{
-		return UniqueId.Equals(Other.UniqueId) && Type == Other.Type;
+		return UniqueId.Equals(Other.UniqueId) && Type == Other.Type && FrameNumber == Other.FrameNumber;
 	}
 
 	//Return the translator key merge with the transform
@@ -76,7 +87,7 @@ public:
 
 	friend uint32 GetTypeHash(const FInterchangeMeshPayLoadKey& InterchangeMeshPayLoadKey)
 	{
-		return GetTypeHash(InterchangeMeshPayLoadKey.UniqueId + FString::FromInt(static_cast<int32>(InterchangeMeshPayLoadKey.Type)));
+		return GetTypeHash(InterchangeMeshPayLoadKey.UniqueId + FString::FromInt(static_cast<int32>(InterchangeMeshPayLoadKey.Type)) + FString::FromInt(InterchangeMeshPayLoadKey.FrameNumber));
 	}
 };
 
@@ -364,3 +375,50 @@ private:
 	
 	UE::Interchange::TMapAttributeHelper<FString, FString> SlotMaterialDependencies;
 };
+
+UCLASS(BlueprintType)
+class INTERCHANGENODES_API UInterchangeGeometryCacheNode : public UInterchangeMeshNode
+{
+	GENERATED_BODY()
+
+public:
+	UInterchangeGeometryCacheNode();
+
+	/** Query the start frame of the animated mesh. Return false if the attribute was not set.*/
+	UFUNCTION(BlueprintCallable, Category = "Interchange | Node | GeometryCache")
+	bool GetCustomStartFrame(int32& AttributeValue) const;
+
+	/** Set the start frame of the animated mesh. Return false if the attribute could not be set.*/
+	UFUNCTION(BlueprintCallable, Category = "Interchange | Node | GeometryCache")
+	bool SetCustomStartFrame(const int32& AttributeValue);
+
+	/** Query the end frame of the animated mesh. Return false if the attribute was not set.*/
+	UFUNCTION(BlueprintCallable, Category = "Interchange | Node | GeometryCache")
+	bool GetCustomEndFrame(int32& AttributeValue) const;
+
+	/** Set the end frame of the animated mesh. Return false if the attribute could not be set.*/
+	UFUNCTION(BlueprintCallable, Category = "Interchange | Node | GeometryCache")
+	bool SetCustomEndFrame(const int32& AttributeValue);
+
+	/** Query the frame rate of the animated mesh. Return false if the attribute was not set.*/
+	UFUNCTION(BlueprintCallable, Category = "Interchange | Node | GeometryCache")
+	bool GetCustomFrameRate(double& AttributeValue) const;
+
+	/** Set the frame rate of the animated mesh. Return false if the attribute could not be set.*/
+	UFUNCTION(BlueprintCallable, Category = "Interchange | Node | GeometryCache")
+	bool SetCustomFrameRate(const double& AttributeValue);
+
+	/** Query whether this animated mesh has constant topoplogy. Return false if the attribute was not set.*/
+	UFUNCTION(BlueprintCallable, Category = "Interchange | Node | GeometryCache")
+	bool GetCustomHasConstantTopology(bool& AttributeValue) const;
+
+	/** Set the constant topology attribute of this animated mesh. Return false if the attribute could not be set.*/
+	UFUNCTION(BlueprintCallable, Category = "Interchange | Node | GeometryCache")
+	bool SetCustomHasConstantTopology(const bool& AttributeValue);
+
+private:
+	IMPLEMENT_NODE_ATTRIBUTE_KEY(StartFrame);
+	IMPLEMENT_NODE_ATTRIBUTE_KEY(EndFrame);
+	IMPLEMENT_NODE_ATTRIBUTE_KEY(FrameRate);
+	IMPLEMENT_NODE_ATTRIBUTE_KEY(HasConstantTopology);
+};
\ No newline at end of file
diff --git a/Engine/Plugins/Interchange/Runtime/Source/Pipelines/InterchangePipelines.Build.cs b/Engine/Plugins/Interchange/Runtime/Source/Pipelines/InterchangePipelines.Build.cs
--- a/Engine/Plugins/Interchange/Runtime/Source/Pipelines/InterchangePipelines.Build.cs
+++ b/Engine/Plugins/Interchange/Runtime/Source/Pipelines/InterchangePipelines.Build.cs
@@ -26,6 +26,7 @@ namespace UnrealBuildTool.Rules
 					"ApplicationCore",
 					"CinematicCamera",
 					"DeveloperSettings",
+					"GeometryCache",
 					"GLTFCore",
 					"InputCore",
 					"InterchangeImport",
diff --git a/Engine/Plugins/Interchange/Runtime/Source/Pipelines/Private/InterchangeGenericGeometryCachePipeline.cpp b/Engine/Plugins/Interchange/Runtime/Source/Pipelines/Private/InterchangeGenericGeometryCachePipeline.cpp
new file mode 100644
--- /dev/null
+++ b/Engine/Plugins/Interchange/Runtime/Source/Pipelines/Private/InterchangeGenericGeometryCachePipeline.cpp
@@ -0,0 +1,156 @@
+// Copyright Epic Games, Inc. All Rights Reserved. 
+
+#include "InterchangeGenericMeshPipeline.h"
+
+#include "InterchangeGeometryCacheFactoryNode.h"
+#include "InterchangeMeshNode.h"
+#include "InterchangePipelineMeshesUtilities.h"
+#include "InterchangeSceneNode.h"
+#include "Nodes/InterchangeBaseNode.h"
+#include "Nodes/InterchangeUserDefinedAttribute.h"
+
+void UInterchangeGenericMeshPipeline::ExecutePreImportPipelineGeometryCache()
+{
+	check(CommonMeshesProperties.IsValid());
+
+	if (bImportGeometryCaches && CommonMeshesProperties->ForceAllMeshAsType == EInterchangeForceMeshType::IFMT_None)
+	{
+		// If there's an animated mesh, combine all meshes into a single Geometry Cache
+		TArray<FString> MeshInstanceUids;
+		PipelineMeshesUtilities->GetAllGeometryCacheInstance(MeshInstanceUids);
+
+		TArray<FString> MeshUids;
+		for (const FString& MeshUid : MeshInstanceUids)
+		{
+			const FInterchangeMeshInstance& MeshInstance = PipelineMeshesUtilities->GetMeshInstanceByUid(MeshUid);
+
+			// Only look at lod 0 since GeometryCache don't support LODs
+			const int32 LodIndex = 0;
+			FInterchangeLodSceneNodeContainer EmptySceneNodeContainer;
+			const FInterchangeLodSceneNodeContainer& SceneNodeContainer = MeshInstance.SceneNodePerLodIndex.Contains(LodIndex) ?
+				MeshInstance.SceneNodePerLodIndex.FindChecked(LodIndex) : EmptySceneNodeContainer;
+
+			for (const UInterchangeSceneNode* SceneNode : SceneNodeContainer.SceneNodes)
+			{
+				MeshUids.Add(SceneNode->GetUniqueID());
+			}
+		}
+
+		// If we got some instances, create a geometry cache factory node
+		if (MeshUids.Num() > 0)
+		{
+			UInterchangeGeometryCacheFactoryNode* GeometryCacheFactoryNode = CreateGeometryCacheFactoryNode(MeshUids);
+			GeometryCacheFactoryNodes.Add(GeometryCacheFactoryNode);
+		}
+	}
+}
+
+UInterchangeGeometryCacheFactoryNode* UInterchangeGenericMeshPipeline::CreateGeometryCacheFactoryNode(const TArray<FString>& MeshUids)
+{
+	check(CommonMeshesProperties.IsValid());
+	if (MeshUids.Num() == 0)
+	{
+		return nullptr;
+	}
+
+	// Name the geometry cache node according to the first non-root node
+	FString GeometryCacheUid;
+	FString DisplayLabel;
+	{
+		// Starting from the first mesh, we go up through its ancestors until we find the first node after the root
+		const UInterchangeBaseNode* Node = BaseNodeContainer->GetNode(MeshUids[0]);
+		if (const UInterchangeSceneNode* SceneNode = Cast<const UInterchangeSceneNode>(Node))
+		{
+			const UInterchangeSceneNode* ParentSceneNode = Cast<const UInterchangeSceneNode>(BaseNodeContainer->GetNode(SceneNode->GetParentUid()));
+			if (!ParentSceneNode || ParentSceneNode->GetParentUid() == UInterchangeBaseNode::InvalidNodeUid())
+			{
+				// The SceneNode itself or its parent is the root so just get the SceneNode's info
+				DisplayLabel = SceneNode->GetDisplayLabel();
+				GeometryCacheUid = SceneNode->GetUniqueID();
+			}
+			else
+			{
+				while (SceneNode && ParentSceneNode)
+				{
+					SceneNode = Cast<const UInterchangeSceneNode>(BaseNodeContainer->GetNode(SceneNode->GetParentUid()));
+					ParentSceneNode = Cast<const UInterchangeSceneNode>(BaseNodeContainer->GetNode(SceneNode->GetParentUid()));
+					if (!ParentSceneNode || ParentSceneNode->GetParentUid() == UInterchangeBaseNode::InvalidNodeUid())
+					{
+						// The root has no parent so retrieve the needed info and break 
+						DisplayLabel = SceneNode->GetDisplayLabel();
+						GeometryCacheUid = SceneNode->GetUniqueID();
+						break;
+					}
+				}
+			}
+		}
+	}
+
+	const FString StaticMeshUid = UInterchangeFactoryBaseNode::BuildFactoryNodeUid(GeometryCacheUid);
+	UInterchangeGeometryCacheFactoryNode* GeometryCacheFactoryNode = NewObject<UInterchangeGeometryCacheFactoryNode>(BaseNodeContainer, NAME_None);
+	if (!ensure(GeometryCacheFactoryNode))
+	{
+		return nullptr;
+	}
+
+	GeometryCacheFactoryNode->InitializeGeometryCacheNode(StaticMeshUid, DisplayLabel, UGeometryCache::StaticClass()->GetName());
+	BaseNodeContainer->AddNode(GeometryCacheFactoryNode);
+
+	// For now, keep all the mesh sections separate since each section will be in its own track in the GeometryCache
+	GeometryCacheFactoryNode->SetCustomKeepSectionsSeparate(true);
+
+	AddMeshesToGeometryCache(GeometryCacheFactoryNode, MeshUids);
+
+	// #ueent_todo: Support the common meshes build options (not all are supported)
+	GeometryCacheFactoryNode->SetCustomRecomputeNormals(CommonMeshesProperties->bRecomputeNormals);
+	GeometryCacheFactoryNode->SetCustomRecomputeTangents(CommonMeshesProperties->bRecomputeTangents);
+	GeometryCacheFactoryNode->SetCustomUseMikkTSpace(CommonMeshesProperties->bUseMikkTSpace);
+	GeometryCacheFactoryNode->SetCustomComputeWeightedNormals(CommonMeshesProperties->bComputeWeightedNormals);
+	GeometryCacheFactoryNode->SetCustomUseHighPrecisionTangentBasis(CommonMeshesProperties->bUseHighPrecisionTangentBasis);
+	GeometryCacheFactoryNode->SetCustomUseFullPrecisionUVs(CommonMeshesProperties->bUseFullPrecisionUVs);
+	GeometryCacheFactoryNode->SetCustomUseBackwardsCompatibleF16TruncUVs(CommonMeshesProperties->bUseBackwardsCompatibleF16TruncUVs);
+	GeometryCacheFactoryNode->SetCustomRemoveDegenerates(CommonMeshesProperties->bRemoveDegenerates);
+
+	return GeometryCacheFactoryNode;
+}
+
+void UInterchangeGenericMeshPipeline::AddMeshesToGeometryCache(UInterchangeGeometryCacheFactoryNode* GeometryCacheFactoryNode, const TArray<FString>& NodeUids)
+{
+	TMap<FString, FString> ExistingLodSlotMaterialDependencies;
+	constexpr bool bAddSourceNodeName = true;
+	for (const FString& NodeUid : NodeUids)
+	{
+		TMap<FString, FString> SlotMaterialDependencies;
+		if (const UInterchangeSceneNode* SceneNode = Cast<UInterchangeSceneNode>(BaseNodeContainer->GetNode(NodeUid)))
+		{
+			FString MeshDependency;
+			SceneNode->GetCustomAssetInstanceUid(MeshDependency);
+			if (const UInterchangeMeshNode* MeshDependencyNode = Cast<UInterchangeMeshNode>(BaseNodeContainer->GetNode(MeshDependency)))
+			{
+				UInterchangeUserDefinedAttributesAPI::DuplicateAllUserDefinedAttribute(MeshDependencyNode, GeometryCacheFactoryNode, bAddSourceNodeName);
+
+				// Note here we add the SceneNode because we'll bake its transform into the mesh
+				GeometryCacheFactoryNode->AddTargetNodeUid(NodeUid);
+				MeshDependencyNode->AddTargetNodeUid(GeometryCacheFactoryNode->GetUniqueID());
+
+				MeshDependencyNode->GetSlotMaterialDependencies(SlotMaterialDependencies);
+			}
+			else
+			{
+				SceneNode->GetSlotMaterialDependencies(SlotMaterialDependencies);
+			}
+
+			UInterchangeUserDefinedAttributesAPI::DuplicateAllUserDefinedAttribute(SceneNode, GeometryCacheFactoryNode, bAddSourceNodeName);
+		}
+		else if (const UInterchangeMeshNode* MeshNode = Cast<UInterchangeMeshNode>(BaseNodeContainer->GetNode(NodeUid)))
+		{
+			UInterchangeUserDefinedAttributesAPI::DuplicateAllUserDefinedAttribute(MeshNode, GeometryCacheFactoryNode, bAddSourceNodeName);
+			GeometryCacheFactoryNode->AddTargetNodeUid(NodeUid);
+			MeshNode->AddTargetNodeUid(GeometryCacheFactoryNode->GetUniqueID());
+
+			MeshNode->GetSlotMaterialDependencies(SlotMaterialDependencies);
+		}
+
+		UE::Interchange::MeshesUtilities::ApplySlotMaterialDependencies(*GeometryCacheFactoryNode, SlotMaterialDependencies, *BaseNodeContainer, &ExistingLodSlotMaterialDependencies);
+	}
+}
diff --git a/Engine/Plugins/Interchange/Runtime/Source/Pipelines/Private/InterchangeGenericMeshPipeline.cpp b/Engine/Plugins/Interchange/Runtime/Source/Pipelines/Private/InterchangeGenericMeshPipeline.cpp
--- a/Engine/Plugins/Interchange/Runtime/Source/Pipelines/Private/InterchangeGenericMeshPipeline.cpp
+++ b/Engine/Plugins/Interchange/Runtime/Source/Pipelines/Private/InterchangeGenericMeshPipeline.cpp
@@ -415,6 +415,9 @@ void UInterchangeGenericMeshPipeline::ExecutePipeline(UInterchangeBaseNodeContai
 
 	//Create staticmesh factory nodes
 	ExecutePreImportPipelineStaticMesh();
+
+	//Create geometry cache factory nodes
+	ExecutePreImportPipelineGeometryCache();
 }
 
 void UInterchangeGenericMeshPipeline::ExecutePostImportPipeline(const UInterchangeBaseNodeContainer* InBaseNodeContainer, const FString& FactoryNodeKey, UObject* CreatedAsset, bool bIsAReimport)
diff --git a/Engine/Plugins/Interchange/Runtime/Source/Pipelines/Private/InterchangePipelineMeshesUtilities.cpp b/Engine/Plugins/Interchange/Runtime/Source/Pipelines/Private/InterchangePipelineMeshesUtilities.cpp
--- a/Engine/Plugins/Interchange/Runtime/Source/Pipelines/Private/InterchangePipelineMeshesUtilities.cpp
+++ b/Engine/Plugins/Interchange/Runtime/Source/Pipelines/Private/InterchangePipelineMeshesUtilities.cpp
@@ -185,7 +185,7 @@ bool FInterchangePipelineMeshesUtilitiesContext::IsStaticMeshInstance(const FInt
 	{
 		return false;
 	}
-	return !IsSkeletalMeshInstance(MeshInstance, BaseNodeContainer);
+	return !MeshInstance.bIsAnimated && !IsSkeletalMeshInstance(MeshInstance, BaseNodeContainer);
 }
 
 bool FInterchangePipelineMeshesUtilitiesContext::IsSkeletalMeshInstance(const FInterchangeMeshInstance& MeshInstance, UInterchangeBaseNodeContainer* BaseNodeContainer)
@@ -197,6 +197,10 @@ bool FInterchangePipelineMeshesUtilitiesContext::IsSkeletalMeshInstance(const FI
 bool FInterchangePipelineMeshesUtilitiesContext::IsSkeletalMeshInstance(const FInterchangeMeshInstance& MeshInstance, UInterchangeBaseNodeContainer* BaseNodeContainer, bool& bOutIsStaticMeshNestedInSkeleton)
 {
 	bOutIsStaticMeshNestedInSkeleton = false;
+	if (MeshInstance.bIsAnimated)
+	{
+		return false;
+	}
 	if (bConvertSkeletalMeshToStaticMesh)
 	{
 		return false;
@@ -223,6 +227,15 @@ bool FInterchangePipelineMeshesUtilitiesContext::IsSkeletalMeshInstance(const FI
 	return false;
 }
 
+bool FInterchangePipelineMeshesUtilitiesContext::IsGeometryCacheInstance(const FInterchangeMeshInstance& MeshInstance)
+{
+	if (bIgnoreGeometryCaches)
+	{
+		return false;
+	}
+	return MeshInstance.bIsAnimated;
+}
+
 bool FInterchangePipelineMeshesUtilitiesContext::IsStaticMeshGeometry(const FInterchangeMeshGeometry& MeshGeometry)
 {
 	if (bIgnoreStaticMeshes)
@@ -237,7 +250,7 @@ bool FInterchangePipelineMeshesUtilitiesContext::IsStaticMeshGeometry(const FInt
 	{
 		return false;
 	}
-	return !IsSkeletalMeshGeometry(MeshGeometry);
+	return !MeshGeometry.MeshNode->IsA<UInterchangeGeometryCacheNode>() && !IsSkeletalMeshGeometry(MeshGeometry);
 }
 
 bool FInterchangePipelineMeshesUtilitiesContext::IsSkeletalMeshGeometry(const FInterchangeMeshGeometry& MeshGeometry)
@@ -257,15 +270,31 @@ bool FInterchangePipelineMeshesUtilitiesContext::IsSkeletalMeshGeometry(const FI
 		return false;
 	}
 
-	if (bConvertStaticMeshToSkeletalMesh
+	if (!MeshGeometry.MeshNode->IsA<UInterchangeGeometryCacheNode>()
+		&& (bConvertStaticMeshToSkeletalMesh
 		|| MeshGeometry.MeshNode->IsSkinnedMesh()
-		|| (bConvertStaticsWithMorphTargetsToSkeletals && (MeshGeometry.MeshNode->GetMorphTargetDependeciesCount() > 0)))
+		|| (bConvertStaticsWithMorphTargetsToSkeletals && (MeshGeometry.MeshNode->GetMorphTargetDependeciesCount() > 0))))
 	{
 		return true;
 	}
 	return false;
 }
 
+bool FInterchangePipelineMeshesUtilitiesContext::IsGeometryCacheGeometry(const FInterchangeMeshGeometry& MeshGeometry)
+{
+	if (bIgnoreGeometryCaches)
+	{
+		return false;
+	}
+
+	if (bQueryGeometryOnlyIfNoInstance && MeshGeometry.ReferencingMeshInstanceUids.Num() > 0)
+	{
+		return false;
+	}
+
+	return MeshGeometry.MeshNode->IsA<UInterchangeGeometryCacheNode>();
+}
+
 UInterchangePipelineMeshesUtilities* UInterchangePipelineMeshesUtilities::CreateInterchangePipelineMeshesUtilities(UInterchangeBaseNodeContainer* BaseNodeContainer)
 {
 	check(BaseNodeContainer);
@@ -372,7 +401,21 @@ UInterchangePipelineMeshesUtilities* UInterchangePipelineMeshesUtilities::Create
 								}
 							}
 
-							if (LodGroupNode)
+							if (MeshGeometry.MeshNode->IsA<UInterchangeGeometryCacheNode>())
+							{
+								FInterchangeMeshInstance& MeshInstance = PipelineMeshesUtilities->MeshInstancesPerMeshInstanceUid.FindOrAdd(NodeUid);
+								MeshInstance.LodGroupNode = nullptr;
+								MeshInstance.MeshInstanceUid = NodeUid;
+								FInterchangeLodSceneNodeContainer& InstancedSceneNodes = MeshInstance.SceneNodePerLodIndex.FindOrAdd(LodIndex);
+								InstancedSceneNodes.SceneNodes.AddUnique(SceneNode);
+								MeshGeometry.ReferencingMeshInstanceUids.Add(MeshInstance.MeshInstanceUid);
+								MeshInstance.ReferencingMeshGeometryUids.Add(MeshUid);
+								MeshInstance.bReferenceSkinnedMesh = false;
+								MeshInstance.bReferenceMorphTarget = false;
+								MeshInstance.bHasMorphTargets = false;
+								MeshInstance.bIsAnimated = true;
+							}
+							else if (LodGroupNode)
 							{
 								//We have a LOD
 								FInterchangeMeshInstance& MeshInstance = PipelineMeshesUtilities->MeshInstancesPerMeshInstanceUid.FindOrAdd(LodGroupNode->GetUniqueID());
@@ -572,6 +615,31 @@ void UInterchangePipelineMeshesUtilities::IterateAllStaticMeshInstance(TFunction
 	}
 }
 
+void UInterchangePipelineMeshesUtilities::GetAllGeometryCacheInstance(TArray<FString>& MeshInstanceUids) const
+{
+	MeshInstanceUids.Empty(MeshInstancesPerMeshInstanceUid.Num());
+	for (const TPair<FString, FInterchangeMeshInstance>& MeshInstanceUidAndMeshInstance : MeshInstancesPerMeshInstanceUid)
+	{
+		const FInterchangeMeshInstance& MeshInstance = MeshInstanceUidAndMeshInstance.Value;
+		if (CurrentDataContext.IsGeometryCacheInstance(MeshInstance))
+		{
+			MeshInstanceUids.Add(MeshInstance.MeshInstanceUid);
+		}
+	}
+}
+
+void UInterchangePipelineMeshesUtilities::IterateAllGeometryCacheInstance(TFunctionRef<void(const FInterchangeMeshInstance&)> IterationLambda) const
+{
+	for (const TPair<FString, FInterchangeMeshInstance>& MeshInstanceUidAndMeshInstance : MeshInstancesPerMeshInstanceUid)
+	{
+		const FInterchangeMeshInstance& MeshInstance = MeshInstanceUidAndMeshInstance.Value;
+		if (CurrentDataContext.IsGeometryCacheInstance(MeshInstance))
+		{
+			IterationLambda(MeshInstance);
+		}
+	}
+}
+
 void UInterchangePipelineMeshesUtilities::GetAllMeshGeometry(TArray<FString>& MeshGeometryUids) const
 {
 	MeshGeometriesPerMeshUid.GetKeys(MeshGeometryUids);
@@ -636,6 +704,31 @@ void UInterchangePipelineMeshesUtilities::IterateAllStaticMeshGeometry(TFunction
 	}
 }
 
+void UInterchangePipelineMeshesUtilities::GetAllGeometryCacheGeometry(TArray<FString>& MeshGeometryUids) const
+{
+	MeshGeometryUids.Empty(MeshGeometriesPerMeshUid.Num());
+	for (const TPair<FString, FInterchangeMeshGeometry>& MeshGeometryUidAndMeshGeometry : MeshGeometriesPerMeshUid)
+	{
+		const FInterchangeMeshGeometry& MeshGeometry = MeshGeometryUidAndMeshGeometry.Value;
+		if (CurrentDataContext.IsGeometryCacheGeometry(MeshGeometry))
+		{
+			MeshGeometryUids.Add(MeshGeometry.MeshUid);
+		}
+	}
+}
+
+void UInterchangePipelineMeshesUtilities::IterateAllGeometryCacheGeometry(TFunctionRef<void(const FInterchangeMeshGeometry&)> IterationLambda) const
+{
+	for (const TPair<FString, FInterchangeMeshGeometry>& MeshGeometryUidAndMeshGeometry : MeshGeometriesPerMeshUid)
+	{
+		const FInterchangeMeshGeometry& MeshGeometry = MeshGeometryUidAndMeshGeometry.Value;
+		if (CurrentDataContext.IsGeometryCacheGeometry(MeshGeometry))
+		{
+			IterationLambda(MeshGeometry);
+		}
+	}
+}
+
 void UInterchangePipelineMeshesUtilities::GetAllMeshGeometryNotInstanced(TArray<FString>& MeshGeometryUids) const
 {
 	MeshGeometryUids.Empty(MeshGeometriesPerMeshUid.Num());
diff --git a/Engine/Plugins/Interchange/Runtime/Source/Pipelines/Public/InterchangeGenericMeshPipeline.h b/Engine/Plugins/Interchange/Runtime/Source/Pipelines/Public/InterchangeGenericMeshPipeline.h
--- a/Engine/Plugins/Interchange/Runtime/Source/Pipelines/Public/InterchangeGenericMeshPipeline.h
+++ b/Engine/Plugins/Interchange/Runtime/Source/Pipelines/Public/InterchangeGenericMeshPipeline.h
@@ -14,6 +14,7 @@
 #include "InterchangeGenericMeshPipeline.generated.h"
 
 class UInterchangeGenericAssetsPipeline;
+class UInterchangeGeometryCacheFactoryNode;
 class UInterchangeMeshNode;
 class UInterchangePipelineMeshesUtilities;
 class UInterchangeSceneNode;
@@ -241,6 +242,12 @@ public:
 	UPROPERTY(EditAnywhere, BlueprintReadWrite, Category = "Skeletal Meshes", meta = (SubCategory = "Build"))
 	int32 BoneInfluenceLimit = 0;
 
+	//////	GEOMETRY_CACHES_CATEGORY Properties //////
+
+	/** If enabled, imports all geometry cache assets found in the sources. */
+	UPROPERTY(EditAnywhere, BlueprintReadWrite, Category = "Geometry Caches")
+	bool bImportGeometryCaches = true;
+
 	virtual void AdjustSettingsForContext(const FInterchangePipelineContextParams& ContextParams) override;
 
 	virtual void PreDialogCleanup(const FName PipelineStackName) override;
@@ -381,6 +388,32 @@ private:
 		, bool& bContainStaticMeshAnimationNode
 		, bool& bIgnoreStaticMeshes) const;
 
+	/************************************************************************/
+	/* Geometry cache API BEGIN                                             */
+	
+	/**
+	 * This function will create any geometry cache we need to create according to the pipeline options
+	 */
+	void ExecutePreImportPipelineGeometryCache();
+
+	/** Geometry cache factory assets nodes */
+	TArray<UInterchangeGeometryCacheFactoryNode*> GeometryCacheFactoryNodes;
+
+	/**
+	 * This function can create a UInterchangeGeometryCacheFactoryNode
+	 * @param MeshUids - The MeshUids can represent a SceneNode pointing on a MeshNode or directly a MeshNode
+	 */
+	UInterchangeGeometryCacheFactoryNode* CreateGeometryCacheFactoryNode(const TArray<FString>& MeshUids);
+
+	/**
+	 * This function add all mesh nodes to the geometry cache.
+	 * @param MeshUids - The MeshUids can be a UInterchangeSceneNode or a UInterchangeMeshNode.
+	 */
+	void AddMeshesToGeometryCache(UInterchangeGeometryCacheFactoryNode* GeometryCacheFactoryNode, const TArray<FString>& NodeUids);
+
+	/* Geometry cache API END                                               */
+	/************************************************************************/
+
 public:
 
 		virtual void PostLoad() override;
diff --git a/Engine/Plugins/Interchange/Runtime/Source/Pipelines/Public/InterchangePipelineMeshesUtilities.h b/Engine/Plugins/Interchange/Runtime/Source/Pipelines/Public/InterchangePipelineMeshesUtilities.h
--- a/Engine/Plugins/Interchange/Runtime/Source/Pipelines/Public/InterchangePipelineMeshesUtilities.h
+++ b/Engine/Plugins/Interchange/Runtime/Source/Pipelines/Public/InterchangePipelineMeshesUtilities.h
@@ -47,6 +47,7 @@ struct FInterchangeMeshInstance
 		bReferenceSkinnedMesh = false;
 		bReferenceMorphTarget = false;
 		bHasMorphTargets = false;
+		bIsAnimated = false;
 	}
 	/**
 	 * This ID represents either a LOD group scene node UID or a mesh scene node UID.
@@ -69,6 +70,9 @@ struct FInterchangeMeshInstance
 	UPROPERTY(EditAnywhere, Category = "Interchange | Pipeline | MeshesInstance")
 	bool bHasMorphTargets;
 
+	UPROPERTY(EditAnywhere, Category = "Interchange | Pipeline | MeshesInstance")
+	bool bIsAnimated;
+
 	/**
 	 * Each scene node here represents a mesh scene node. If it represents a LOD group, there may be more then one mesh scene node for a specific LOD index.
 	 */
@@ -165,11 +169,19 @@ struct FInterchangePipelineMeshesUtilitiesContext
 	UPROPERTY(EditAnywhere, BlueprintReadWrite, Category = "Interchange | Pipeline | MeshesContext")
 	bool bIgnoreStaticMeshes = false;
 
+	/**
+	 * If enabled, all geometry caches will be ignored. The mesh utility will not return any geometry caches.
+	 */
+	UPROPERTY(EditAnywhere, BlueprintReadWrite, Category = "Interchange | Pipeline | MeshesContext")
+	bool bIgnoreGeometryCaches = false;
+
 	bool IsStaticMeshInstance(const FInterchangeMeshInstance& MeshInstance, UInterchangeBaseNodeContainer* BaseNodeContainer);
 	bool IsSkeletalMeshInstance(const FInterchangeMeshInstance& MeshInstance, UInterchangeBaseNodeContainer* BaseNodeContainer);
 	bool IsSkeletalMeshInstance(const FInterchangeMeshInstance& MeshInstance, UInterchangeBaseNodeContainer* BaseNodeContainer, bool& bOutIsStaticMeshNestedInSkeleton);
+	bool IsGeometryCacheInstance(const FInterchangeMeshInstance& MeshInstance);
 	bool IsStaticMeshGeometry(const FInterchangeMeshGeometry& MeshGeometry);
 	bool IsSkeletalMeshGeometry(const FInterchangeMeshGeometry& MeshGeometry);
+	bool IsGeometryCacheGeometry(const FInterchangeMeshGeometry& MeshGeometry);
 };
 
 /**/
@@ -217,6 +229,17 @@ public:
 	*/
 	void IterateAllStaticMeshInstance(TFunctionRef<void(const FInterchangeMeshInstance&)> IterationLambda) const;
 
+	/**
+	* Get the unique IDs of all geometry cache instances.
+	*/
+	UFUNCTION(BlueprintCallable, Category = "Interchange | Pipeline | Meshes")
+	void GetAllGeometryCacheInstance(TArray<FString>& MeshInstanceUids) const;
+
+	/**
+	* Iterate over all geometry cache instances.
+	*/
+	void IterateAllGeometryCacheInstance(TFunctionRef<void(const FInterchangeMeshInstance&)> IterationLambda) const;
+
 	/**
 	* Get the unique IDs of all mesh geometry.
 	*/
@@ -250,6 +273,17 @@ public:
 	*/
 	void IterateAllStaticMeshGeometry(TFunctionRef<void(const FInterchangeMeshGeometry&)> IterationLambda) const;
 
+	/**
+	* Get the unique IDs of all geometry cache geometry.
+	*/
+	UFUNCTION(BlueprintCallable, Category = "Interchange | Pipeline | Meshes")
+	void GetAllGeometryCacheGeometry(TArray<FString>& MeshGeometryUids) const;
+
+	/**
+	* Iterate over all geometry cache geometry.
+	*/
+	void IterateAllGeometryCacheGeometry(TFunctionRef<void(const FInterchangeMeshGeometry&)> IterationLambda) const;
+
 	/**
 	* Get the unique IDs of all non-instanced mesh geometry.
 	*/
