[Interchange-USD/glTF] Fixing Euler flipping on imported animation curves.

#rnx

diff --git a/Engine/Plugins/Interchange/Runtime/Source/Import/Private/Gltf/InterchangeGltfAnimation.cpp b/Engine/Plugins/Interchange/Runtime/Source/Import/Private/Gltf/InterchangeGltfAnimation.cpp
--- a/Engine/Plugins/Interchange/Runtime/Source/Import/Private/Gltf/InterchangeGltfAnimation.cpp
+++ b/Engine/Plugins/Interchange/Runtime/Source/Import/Private/Gltf/InterchangeGltfAnimation.cpp
@@ -313,13 +313,13 @@ namespace UE::Interchange::Gltf::Private
 				FrameDataBuffer.SetNumUninitialized(Sampler.Output.Count * 4);
 				Sampler.Output.GetQuatArray(reinterpret_cast<FVector4f*>(FrameDataBuffer.GetData()));
 
-				struct HeadingHelper
+				struct FHeadingHelper
 				{
 					FRotator3f Heading;
 					bool bHeadingSet = false;
 				};
-				HeadingHelper Helper;
-				ConvertCurve([&Helper](const float* QuatData)
+				FHeadingHelper HeadingHelper;
+				ConvertCurve([&HeadingHelper](const float* QuatData)
 					{
 						//GLTF uses quaternions while Uneal uses Eulers for level animations
 						// Quaternions can represent multiple Orientations in Eulers.
@@ -330,11 +330,15 @@ namespace UE::Interchange::Gltf::Private
 
 						FQuat4f Quat(QuatData[0], QuatData[1], QuatData[2], QuatData[3]);
 						FRotator3f Rotator = Quat.Rotator();
-						if (Helper.bHeadingSet)
+						if (HeadingHelper.bHeadingSet)
 						{
+							FMath::WindRelativeAnglesDegrees(HeadingHelper.Heading.Roll, Rotator.Roll);
+							FMath::WindRelativeAnglesDegrees(HeadingHelper.Heading.Pitch, Rotator.Pitch);
+							FMath::WindRelativeAnglesDegrees(HeadingHelper.Heading.Yaw, Rotator.Yaw);
+
 							FRotator3f OtherChoice = Rotator.GetEquivalentRotator().GetNormalized();
-							float FirstDiff = Helper.Heading.GetManhattanDistance(Rotator);
-							float SecondDiff = Helper.Heading.GetManhattanDistance(OtherChoice);
+							float FirstDiff = HeadingHelper.Heading.GetManhattanDistance(Rotator);
+							float SecondDiff = HeadingHelper.Heading.GetManhattanDistance(OtherChoice);
 							if (SecondDiff < FirstDiff)
 							{
 								Rotator = OtherChoice;
@@ -342,10 +346,10 @@ namespace UE::Interchange::Gltf::Private
 						}
 						else
 						{
-							Helper.bHeadingSet = true;
+							HeadingHelper.bHeadingSet = true;
 						}
 
-						Helper.Heading = Rotator;
+						HeadingHelper.Heading = Rotator;
 
 						return Rotator.Euler();
 					}, 3, 4, true);
diff --git a/Engine/Plugins/Runtime/USDCore/Source/USDUtilities/Private/USDPrimConversion.cpp b/Engine/Plugins/Runtime/USDCore/Source/USDUtilities/Private/USDPrimConversion.cpp
--- a/Engine/Plugins/Runtime/USDCore/Source/USDUtilities/Private/USDPrimConversion.cpp
+++ b/Engine/Plugins/Runtime/USDCore/Source/USDUtilities/Private/USDPrimConversion.cpp
@@ -1005,6 +1005,14 @@ bool UsdToUnreal::ConvertTransformTimeSamples(
 												: ERichCurveInterpMode::RCIM_Constant;
 
 	double LastTimeSample = TNumericLimits<double>::Lowest();
+
+	struct FHeadingHelper
+	{
+		FRotator Heading;
+		bool bHeadingSet = false;
+	};
+	FHeadingHelper HeadingHelper;
+
 	for (const double UsdTimeSample : UsdTimeSamples)
 	{
 		// We never want to evaluate the same time twice
@@ -1028,6 +1036,27 @@ bool UsdToUnreal::ConvertTransformTimeSamples(
 		FRotator Rotator = UEValue.Rotator();
 		FVector Scale = UEValue.GetScale3D();
 
+		if (HeadingHelper.bHeadingSet)
+		{
+			FMath::WindRelativeAnglesDegrees(HeadingHelper.Heading.Roll, Rotator.Roll);
+			FMath::WindRelativeAnglesDegrees(HeadingHelper.Heading.Pitch, Rotator.Pitch);
+			FMath::WindRelativeAnglesDegrees(HeadingHelper.Heading.Yaw, Rotator.Yaw);
+
+			FRotator OtherChoice = Rotator.GetEquivalentRotator().GetNormalized();
+			float FirstDiff = HeadingHelper.Heading.GetManhattanDistance(Rotator);
+			float SecondDiff = HeadingHelper.Heading.GetManhattanDistance(OtherChoice);
+			if (SecondDiff < FirstDiff)
+			{
+				Rotator = OtherChoice;
+			}
+		}
+		else
+		{
+			HeadingHelper.bHeadingSet = true;
+		}
+
+		HeadingHelper.Heading = Rotator;
+
 		LocationXValues.Emplace_GetRef(Location.X).InterpMode = InterpMode;
 		LocationYValues.Emplace_GetRef(Location.Y).InterpMode = InterpMode;
 		LocationZValues.Emplace_GetRef(Location.Z).InterpMode = InterpMode;
