Interchange:
- Add support of SpecularProfiles

diff --git a/Engine/Plugins/Interchange/Runtime/Source/FactoryNodes/Private/InterchangeSpecularProfileFactoryNode.cpp b/Engine/Plugins/Interchange/Runtime/Source/FactoryNodes/Private/InterchangeSpecularProfileFactoryNode.cpp
new file mode 100644
--- /dev/null
+++ b/Engine/Plugins/Interchange/Runtime/Source/FactoryNodes/Private/InterchangeSpecularProfileFactoryNode.cpp
@@ -0,0 +1,36 @@
+// Copyright Epic Games, Inc. All Rights Reserved.
+
+#include "InterchangeSpecularProfileFactoryNode.h"
+#include "Engine/SpecularProfile.h"
+
+#include UE_INLINE_GENERATED_CPP_BY_NAME(InterchangeSpecularProfileFactoryNode)
+
+FString UInterchangeSpecularProfileFactoryNode::GetTypeName() const
+{
+	return TEXT("SpecularProfileFactoryNode");
+}
+
+UClass* UInterchangeSpecularProfileFactoryNode::GetObjectClass() const
+{
+	return USpecularProfile::StaticClass();
+}
+
+bool UInterchangeSpecularProfileFactoryNode::SetCustomFormat(ESpecularProfileFormat AttributeValue)
+{
+	IMPLEMENT_NODE_ATTRIBUTE_SETTER_NODELEGATE(Format, ESpecularProfileFormat);
+}
+
+bool UInterchangeSpecularProfileFactoryNode::GetCustomFormat(ESpecularProfileFormat& AttributeValue) const
+{
+	IMPLEMENT_NODE_ATTRIBUTE_GETTER(Format, ESpecularProfileFormat);
+}
+
+bool UInterchangeSpecularProfileFactoryNode::GetCustomTexture(FString& AttributeValue) const
+{
+	IMPLEMENT_NODE_ATTRIBUTE_GETTER(Texture, FString);
+}
+
+bool UInterchangeSpecularProfileFactoryNode::SetCustomTexture(const FString& AttributeValue)
+{
+	IMPLEMENT_NODE_ATTRIBUTE_SETTER_NODELEGATE(Texture, FString)
+}
diff --git a/Engine/Plugins/Interchange/Runtime/Source/FactoryNodes/Public/InterchangeSpecularProfileFactoryNode.h b/Engine/Plugins/Interchange/Runtime/Source/FactoryNodes/Public/InterchangeSpecularProfileFactoryNode.h
new file mode 100644
--- /dev/null
+++ b/Engine/Plugins/Interchange/Runtime/Source/FactoryNodes/Public/InterchangeSpecularProfileFactoryNode.h
@@ -0,0 +1,41 @@
+// Copyright Epic Games, Inc. All Rights Reserved.
+
+#pragma once
+
+#include "Nodes/InterchangeFactoryBaseNode.h"
+
+#include "InterchangeSpecularProfileFactoryNode.generated.h"
+
+enum class ESpecularProfileFormat : uint8;
+
+UCLASS(BlueprintType)
+class INTERCHANGEFACTORYNODES_API UInterchangeSpecularProfileFactoryNode : public UInterchangeFactoryBaseNode
+{
+	GENERATED_BODY()
+
+public:
+
+	/**
+	 * Return the node type name of the class. This is used when reporting errors.
+	 */
+	virtual FString GetTypeName() const override;
+
+	virtual class UClass* GetObjectClass() const override;
+
+	UFUNCTION(BlueprintCallable, Category = "Interchange | FactoryNode | SpecularProfile")
+	bool SetCustomFormat(ESpecularProfileFormat Format);
+
+	UFUNCTION(BlueprintCallable, Category = "Interchange | FactoryNode | SpecularProfile")
+	bool GetCustomFormat(ESpecularProfileFormat& Format) const;
+
+	UFUNCTION(BlueprintCallable, Category = "Interchange | FactoryNode | SpecularProfile")
+	bool GetCustomTexture(FString& TextureUid) const;
+
+	UFUNCTION(BlueprintCallable, Category = "Interchange | FactoryNode | SpecularProfile")
+	bool SetCustomTexture(const FString& TextureUid);
+
+private:
+
+	IMPLEMENT_NODE_ATTRIBUTE_KEY(Format)
+	IMPLEMENT_NODE_ATTRIBUTE_KEY(Texture)
+};
diff --git a/Engine/Plugins/Interchange/Runtime/Source/Import/Private/InterchangeImportModule.cpp b/Engine/Plugins/Interchange/Runtime/Source/Import/Private/InterchangeImportModule.cpp
--- a/Engine/Plugins/Interchange/Runtime/Source/Import/Private/InterchangeImportModule.cpp
+++ b/Engine/Plugins/Interchange/Runtime/Source/Import/Private/InterchangeImportModule.cpp
@@ -32,6 +32,7 @@
 #include "Scene/InterchangeSceneVariantSetsFactory.h"
 #include "Scene/InterchangeSkeletalMeshActorFactory.h"
 #include "Scene/InterchangeStaticMeshActorFactory.h"
+#include "SpecularProfile/InterchangeSpecularProfileFactory.h"
 #include "Texture/InterchangeDDSTranslator.h"
 #include "Texture/InterchangeIESTranslator.h"
 #include "Texture/InterchangeImageWrapperTranslator.h"
@@ -130,6 +131,7 @@ void FInterchangeImportModule::StartupModule()
 		InterchangeManager.RegisterFactory(UInterchangeSceneVariantSetsFactory::StaticClass());
 		InterchangeManager.RegisterFactory(UInterchangeLightActorFactory::StaticClass());
 		InterchangeManager.RegisterFactory(UInterchangeSceneImportAssetFactory::StaticClass());
+		InterchangeManager.RegisterFactory(UInterchangeSpecularProfileFactory::StaticClass());
 #if WITH_EDITOR
 		// Editor-only because SVT internals like UStreamableSparseVolumeTexture::AppendFrame are editor-only
 		InterchangeManager.RegisterFactory(UInterchangeSparseVolumeTextureFactory::StaticClass());
diff --git a/Engine/Plugins/Interchange/Runtime/Source/Import/Private/Material/InterchangeMaterialFactory.cpp b/Engine/Plugins/Interchange/Runtime/Source/Import/Private/Material/InterchangeMaterialFactory.cpp
--- a/Engine/Plugins/Interchange/Runtime/Source/Import/Private/Material/InterchangeMaterialFactory.cpp
+++ b/Engine/Plugins/Interchange/Runtime/Source/Import/Private/Material/InterchangeMaterialFactory.cpp
@@ -8,6 +8,7 @@
 #include "InterchangeDecalMaterialFactoryNode.h"
 #include "InterchangeShaderGraphNode.h"
 #include "InterchangeSourceData.h"
+#include "InterchangeSpecularProfileFactoryNode.h"
 #include "InterchangeTextureNode.h"
 #include "InterchangeTextureFactoryNode.h"
 #include "InterchangeTranslatorBase.h"
@@ -23,6 +24,7 @@
 #include "Materials/MaterialExpressionTextureSample.h"
 #include "Materials/MaterialExpressionThinTranslucentMaterialOutput.h"
 #include "Materials/MaterialExpressionVectorParameter.h"
+#include "Materials/MaterialExpressionSubstrate.h"
 #include "Materials/MaterialExpressionSwitch.h"
 #include "Materials/MaterialInterface.h"
 #include "Materials/MaterialInstanceConstant.h"
@@ -33,6 +35,7 @@
 #include "UObject/ObjectRedirector.h"
 #include "Misc/PackageName.h"
 #include "Engine/RendererSettings.h"
+#include "Engine/SpecularProfile.h"
 
 #include UE_INLINE_GENERATED_CPP_BY_NAME(InterchangeMaterialFactory)
 
@@ -49,6 +52,7 @@
 
 namespace UE::Interchange::MaterialFactory::Internal
 {
+	void SetupSpecularProfile(const UInterchangeFactoryBase::FImportAssetObjectParams& Arguments, const UInterchangeFactoryBaseNode* FactoryNode, UObject* MaterialInstanceOrSlabBSDFExpression);
 
 #if WITH_EDITOR
 	/**
@@ -355,6 +359,10 @@ namespace UE::Interchange::MaterialFactory::Internal
 			{
 				SetupFunctionCallExpression(Arguments, ExpressionNode, FunctionCallExpression);
 			}
+			else if (UMaterialExpressionSubstrateSlabBSDF* SlabBSDFExpression = Cast<UMaterialExpressionSubstrateSlabBSDF>(MaterialExpression))
+			{
+				SetupSpecularProfile(Arguments, &ExpressionNode, SlabBSDFExpression);
+			}
 
 			return MaterialExpression;
 		}
@@ -720,6 +728,35 @@ namespace UE::Interchange::MaterialFactory::Internal
 
 		return ToBeConvertedTexture;
 	}
+
+	void SetupSpecularProfile(const UInterchangeFactoryBase::FImportAssetObjectParams& Arguments, const UInterchangeFactoryBaseNode* FactoryNode, UObject* MaterialInstanceOrSlabBSDFExpression)
+	{
+		using namespace UE::Interchange::Materials;
+
+		if (FString SpecularProfileUid; FactoryNode->GetStringAttribute(SubstrateMaterial::SpecularProfile.ToString(), SpecularProfileUid))
+		{
+			FString SpecularProfileFactoryNodeUid = UInterchangeFactoryBaseNode::BuildFactoryNodeUid(SpecularProfileUid);
+			if (const UInterchangeSpecularProfileFactoryNode* SpecularProfileFactoryNode = Cast<UInterchangeSpecularProfileFactoryNode>(Arguments.NodeContainer->GetNode(SpecularProfileFactoryNodeUid)))
+			{
+				if (FSoftObjectPath ReferenceObject; SpecularProfileFactoryNode->GetCustomReferenceObject(ReferenceObject))
+				{
+					if (USpecularProfile* SpecularProfile = Cast<USpecularProfile>(ReferenceObject.TryLoad()))
+					{
+						if (UMaterialExpressionSubstrateSlabBSDF* SlabBSDFExpression = Cast<UMaterialExpressionSubstrateSlabBSDF>(MaterialInstanceOrSlabBSDFExpression))
+						{
+							SlabBSDFExpression->SpecularProfile = SpecularProfile;
+						}
+						else if (UMaterialInstance* MaterialInstance = Cast<UMaterialInstance>(MaterialInstanceOrSlabBSDFExpression))
+						{
+							MaterialInstance->SpecularProfileOverride = SpecularProfile;
+							MaterialInstance->bOverrideSpecularProfile = true;
+						}
+
+					}
+				}
+			}
+		}
+	}
 }
 		
 namespace UE::Interchange::Materials::HashUtils
@@ -905,6 +942,8 @@ UInterchangeFactoryBase::FImportAssetResult UInterchangeMaterialFactory::BeginIm
 			EReimportStrategyFlags ReimportStrategyFlags = MaterialFactoryNode->GetReimportStrategyFlags();
 			bool bApplyPipelineProperties = !Arguments.ReimportObject || ReimportStrategyFlags != EReimportStrategyFlags::ApplyNoProperties;
 
+			UE::Interchange::MaterialFactory::Internal::SetupSpecularProfile(Arguments, MaterialFactoryNode, MaterialInstance);
+
 #if WITH_EDITORONLY_DATA
 			// For the time being, reimport policies are only enforced on UMaterialInstanceConstant
 			if (Arguments.ReimportObject && MaterialInstance->IsA<UMaterialInstanceConstant>())
diff --git a/Engine/Plugins/Interchange/Runtime/Source/Import/Private/SpecularProfile/InterchangeSpecularProfileFactory.cpp b/Engine/Plugins/Interchange/Runtime/Source/Import/Private/SpecularProfile/InterchangeSpecularProfileFactory.cpp
new file mode 100644
--- /dev/null
+++ b/Engine/Plugins/Interchange/Runtime/Source/Import/Private/SpecularProfile/InterchangeSpecularProfileFactory.cpp
@@ -0,0 +1,137 @@
+// Copyright Epic Games, Inc. All Rights Reserved.
+
+#include "SpecularProfile/InterchangeSpecularProfileFactory.h"
+#include "InterchangeSourceData.h"
+#include "InterchangeSpecularProfileFactoryNode.h"
+#include "InterchangeTexture2DFactoryNode.h"
+#include "Nodes/InterchangeBaseNodeContainer.h"
+
+#include "Engine/SpecularProfile.h"
+
+#include UE_INLINE_GENERATED_CPP_BY_NAME(InterchangeSpecularProfileFactory)
+
+#define LOCTEXT_NAMESPACE "InterchangeSpecularProfileFactory"
+
+UClass* UInterchangeSpecularProfileFactory::GetFactoryClass() const
+{
+	return USpecularProfile::StaticClass();
+}
+
+EInterchangeFactoryAssetType UInterchangeSpecularProfileFactory::GetFactoryAssetType()
+{
+	return EInterchangeFactoryAssetType::Materials;
+}
+
+UInterchangeFactoryBase::FImportAssetResult UInterchangeSpecularProfileFactory::BeginImportAsset_GameThread(const FImportAssetObjectParams& Arguments)
+{
+	TRACE_CPUPROFILER_EVENT_SCOPE(UInterchangeSpecularProfileFactory::BeginImportAsset_GameThread);
+
+	Super::BeginImportAsset_GameThread(Arguments);
+
+	FImportAssetResult ImportAssetResult;
+
+	auto CannotReimportMessage = [&Arguments, this]()
+	{
+		UInterchangeResultError_Generic* Message = AddMessage<UInterchangeResultError_Generic>();
+		Message->SourceAssetName = Arguments.SourceData->GetFilename();
+		Message->DestinationAssetName = Arguments.AssetName;
+		Message->AssetType = GetFactoryClass();
+		Message->Text = LOCTEXT("CreateEmptyAssetUnsupportedReimport", "Re-import of USpecularProfile not supported yet.");
+		Arguments.AssetNode->SetSkipNodeImport();
+	};
+
+	if (Arguments.ReimportObject)
+	{
+		CannotReimportMessage();
+		return ImportAssetResult;
+	}
+
+	auto CouldNotCreateSpecularProfileLog = [this, &Arguments, &ImportAssetResult](const FText& Info)
+	{
+		UInterchangeResultError_Generic* Message = AddMessage<UInterchangeResultError_Generic>();
+		Message->SourceAssetName = Arguments.SourceData->GetFilename();
+		Message->DestinationAssetName = Arguments.AssetName;
+		Message->AssetType = GetFactoryClass();
+		Message->Text = FText::Format(LOCTEXT("SpecProfileFactory_CouldNotCreateSpecProfile", "Could not create Specular Profile asset %s. Reason: %s"), FText::FromString(Arguments.AssetName), Info);
+		ImportAssetResult.bIsFactorySkipAsset = true;
+	};
+
+	const FText MissMatchClassText = LOCTEXT("SpecProfileFactory_CouldNotCreateSpecProfile_MissMatchClass", "Mismatch between Interchange specular profile factory node class and factory class.");
+	
+	if (!Arguments.AssetNode || !Arguments.AssetNode->GetObjectClass()->IsChildOf(GetFactoryClass()))
+	{
+		CouldNotCreateSpecularProfileLog(MissMatchClassText);
+		return ImportAssetResult;
+	}
+
+	const UInterchangeSpecularProfileFactoryNode* FactoryNode = Cast<UInterchangeSpecularProfileFactoryNode>(Arguments.AssetNode);
+	if (!FactoryNode)
+	{
+		CouldNotCreateSpecularProfileLog(LOCTEXT("SpecProfileFactory_CouldNotCreateSpecProfile_CannotCastFactoryNode", "Cannot cast Interchange factory node to UInterchangeSpecularProfileFactoryNode."));
+		return ImportAssetResult;
+	}
+
+	const UClass* SpecularProfileClass = FactoryNode->GetObjectClass();
+	if (!ensure(SpecularProfileClass && SpecularProfileClass->IsChildOf(GetFactoryClass())))
+	{
+		CouldNotCreateSpecularProfileLog(MissMatchClassText);
+		return ImportAssetResult;
+	}
+
+	UObject* ExistingAsset = Arguments.ReimportObject;
+	if (!ExistingAsset)
+	{
+		FSoftObjectPath ReferenceObject;
+		if (FactoryNode->GetCustomReferenceObject(ReferenceObject))
+		{
+			ExistingAsset = ReferenceObject.TryLoad();
+		}
+	}
+
+	// create a new specular profile or overwrite existing asset, if possible
+	USpecularProfile* SpecularProfile = nullptr;
+	if (!ExistingAsset)
+	{
+		SpecularProfile = NewObject<USpecularProfile>(Arguments.Parent, *Arguments.AssetName, RF_Public | RF_Standalone);
+	}
+	else
+	{
+		// This is a reimport or an override, we are just re-updating the source data
+
+		//TODO: put back the Cast when the SpecularProfile will support re-import
+		//SpecularProfile = Cast<USpecularProfile>(ExistingAsset);
+		CannotReimportMessage();
+		return ImportAssetResult;
+	}
+
+	if (!SpecularProfile)
+	{
+		CouldNotCreateSpecularProfileLog(LOCTEXT("SpecProfileFactory_CouldNotCreateSpecProfile_SpecularProfileCreationFail", "Specular Profile creation failed"));
+		return ImportAssetResult;
+	}
+
+	ImportAssetResult.ImportedObject = SpecularProfile;
+
+	if(ESpecularProfileFormat Format; FactoryNode->GetCustomFormat(Format))
+	{
+		SpecularProfile->Settings.Format = Format;
+	}
+
+	if (FString TextureUid; FactoryNode->GetCustomTexture(TextureUid))
+	{
+		FString TextureFactoryNodeUid = UInterchangeFactoryBaseNode::BuildFactoryNodeUid(TextureUid);
+		if (const UInterchangeTexture2DFactoryNode* TextureFactoryNode = Cast<UInterchangeTexture2DFactoryNode>(Arguments.NodeContainer->GetNode(TextureFactoryNodeUid)))
+		{
+			if(FSoftObjectPath ReferenceObject; TextureFactoryNode->GetCustomReferenceObject(ReferenceObject))
+			{
+				if (UTexture2D* Texture = Cast<UTexture2D>(ReferenceObject.TryLoad()))
+				{
+					SpecularProfile->Settings.Texture = Texture;
+				}
+			}
+		}
+	}
+
+	return ImportAssetResult;
+}
+#undef LOCTEXT_NAMESPACE
\ No newline at end of file
diff --git a/Engine/Plugins/Interchange/Runtime/Source/Import/Public/SpecularProfile/InterchangeSpecularProfileFactory.h b/Engine/Plugins/Interchange/Runtime/Source/Import/Public/SpecularProfile/InterchangeSpecularProfileFactory.h
new file mode 100644
--- /dev/null
+++ b/Engine/Plugins/Interchange/Runtime/Source/Import/Public/SpecularProfile/InterchangeSpecularProfileFactory.h
@@ -0,0 +1,19 @@
+// Copyright Epic Games, Inc. All Rights Reserved.
+
+#pragma once
+
+#include "InterchangeFactoryBase.h"
+
+#include "InterchangeSpecularProfileFactory.generated.h"
+
+UCLASS(BlueprintType)
+class INTERCHANGEIMPORT_API UInterchangeSpecularProfileFactory : public UInterchangeFactoryBase
+{
+	GENERATED_BODY()
+
+public:
+
+	virtual UClass* GetFactoryClass() const override;
+	virtual EInterchangeFactoryAssetType GetFactoryAssetType() override;
+	virtual FImportAssetResult BeginImportAsset_GameThread(const FImportAssetObjectParams& Arguments) override;
+};
\ No newline at end of file
diff --git a/Engine/Plugins/Interchange/Runtime/Source/Nodes/Private/InterchangeSpecularProfileNode.cpp b/Engine/Plugins/Interchange/Runtime/Source/Nodes/Private/InterchangeSpecularProfileNode.cpp
new file mode 100644
--- /dev/null
+++ b/Engine/Plugins/Interchange/Runtime/Source/Nodes/Private/InterchangeSpecularProfileNode.cpp
@@ -0,0 +1,50 @@
+// Copyright Epic Games, Inc. All Rights Reserved.
+
+#include "InterchangeSpecularProfileNode.h"
+#include "Nodes/InterchangeBaseNodeContainer.h"
+
+#include UE_INLINE_GENERATED_CPP_BY_NAME(InterchangeSpecularProfileNode)
+
+FString UInterchangeSpecularProfileNode::MakeNodeUid(const FStringView& NodeName)
+{
+	return FString(UInterchangeBaseNode::HierarchySeparator) + TEXT("SpecularProfile") + FString(UInterchangeBaseNode::HierarchySeparator) + NodeName;
+}
+
+UInterchangeSpecularProfileNode* UInterchangeSpecularProfileNode::Create(UInterchangeBaseNodeContainer* NodeContainer, const FStringView& NodeName)
+{
+	check(NodeContainer);
+
+	const FString SpecularProfileNodeUid = MakeNodeUid(NodeName);
+
+	UInterchangeSpecularProfileNode* SpecularProfileNode = NewObject<UInterchangeSpecularProfileNode>(NodeContainer);
+	SpecularProfileNode->InitializeNode(SpecularProfileNodeUid, FString{ NodeName }, EInterchangeNodeContainerType::TranslatedAsset);
+
+	NodeContainer->AddNode(SpecularProfileNode);
+
+	return SpecularProfileNode;
+}
+
+FString UInterchangeSpecularProfileNode::GetTypeName() const
+{
+	return TEXT("SpecularProfileNode");
+}
+
+bool UInterchangeSpecularProfileNode::SetCustomFormat(uint8 AttributeValue)
+{
+	IMPLEMENT_NODE_ATTRIBUTE_SETTER_NODELEGATE(Format, uint8);
+}
+
+bool UInterchangeSpecularProfileNode::GetCustomFormat(uint8& AttributeValue) const
+{
+	IMPLEMENT_NODE_ATTRIBUTE_GETTER(Format, uint8);
+}
+
+bool UInterchangeSpecularProfileNode::GetCustomTexture(FString& AttributeValue) const
+{
+	IMPLEMENT_NODE_ATTRIBUTE_GETTER(Texture, FString)
+}
+
+bool UInterchangeSpecularProfileNode::SetCustomTexture(const FString& AttributeValue)
+{
+	IMPLEMENT_NODE_ATTRIBUTE_SETTER_NODELEGATE(Texture, FString)
+}
diff --git a/Engine/Plugins/Interchange/Runtime/Source/Nodes/Public/InterchangeMaterialDefinitions.h b/Engine/Plugins/Interchange/Runtime/Source/Nodes/Public/InterchangeMaterialDefinitions.h
--- a/Engine/Plugins/Interchange/Runtime/Source/Nodes/Public/InterchangeMaterialDefinitions.h
+++ b/Engine/Plugins/Interchange/Runtime/Source/Nodes/Public/InterchangeMaterialDefinitions.h
@@ -390,6 +390,38 @@ namespace Interchange
 						const FName DefaultValue = TEXT("DefaultValue");
 					}
 				}
+
+				namespace SlabBSDF
+				{
+					const FName Name = TEXT("SubstrateSlabBSDF");
+
+					namespace Attributes
+					{
+						const FName SpecularProfile = TEXT("SpecularProfile");
+					}
+
+					namespace Inputs
+					{
+						const FName DiffuseAlbedo = TEXT("DiffuseAlbedo");
+						const FName F0 = TEXT("F0");
+						const FName F90 = TEXT("F90");
+						const FName Roughness = TEXT("Roughness");
+						const FName Anisotropy = TEXT("Anisotropy");
+						const FName Normal = TEXT("Normal");
+						const FName Tangent = TEXT("Tangent");
+						const FName SSSMFP = TEXT("SSSMFP");
+						const FName SSSMFPScale = TEXT("SSSMFPScale");
+						const FName SSSPhaseAnisotropy = TEXT("SSSPhaseAnisotropy");
+						const FName EmissiveColor = TEXT("EmissiveColor");
+						const FName SecondRoughness = TEXT("SecondRoughness");
+						const FName SecondRoughnessWeight = TEXT("SecondRoughnessWeight");
+						const FName FuzzRoughness = TEXT("FuzzRoughness");
+						const FName FuzzAmount = TEXT("FuzzAmount");
+						const FName FuzzColor = TEXT("FuzzColor");
+						const FName GlintValue = TEXT("GlintValue");
+						const FName GlintUV = TEXT("GlintUV");
+					}
+				}
 			}
 		}
 
@@ -715,6 +747,8 @@ namespace Interchange
 
 		namespace SubstrateMaterial
 		{
+			const FName SpecularProfile = TEXT("SpecularProfile"); // Used by MaterialInstances otherwise look at SubstrateSlabBSDF
+
 			namespace Parameters
 			{
 				const FName FrontMaterial = TEXT("Front Material");
diff --git a/Engine/Plugins/Interchange/Runtime/Source/Nodes/Public/InterchangeSpecularProfileNode.h b/Engine/Plugins/Interchange/Runtime/Source/Nodes/Public/InterchangeSpecularProfileNode.h
new file mode 100644
--- /dev/null
+++ b/Engine/Plugins/Interchange/Runtime/Source/Nodes/Public/InterchangeSpecularProfileNode.h
@@ -0,0 +1,50 @@
+// Copyright Epic Games, Inc. All Rights Reserved.
+
+#pragma once
+
+#include "Nodes/InterchangeBaseNode.h"
+
+#include "InterchangeSpecularProfileNode.generated.h"
+
+UCLASS(BlueprintType)
+class INTERCHANGENODES_API UInterchangeSpecularProfileNode : public UInterchangeBaseNode
+{
+	GENERATED_BODY()
+
+public:
+
+	/**
+	 * Build and return a UID name for a specular profile node.
+	 */
+	static FString MakeNodeUid(const FStringView& NodeName);
+
+	/**
+	 * Create a new UInterchangeSpecularProfileNode and add it to NodeContainer as a translated node.
+	 */
+	static UInterchangeSpecularProfileNode* Create(UInterchangeBaseNodeContainer* NodeContainer, const FStringView& NodeName);
+
+public:
+
+	/**
+	 * Return the node type name of the class. This is used when reporting errors.
+	 */
+	virtual FString GetTypeName() const override;
+
+	/** Set/Get the format based on ESpecularProfileFormat */
+	UFUNCTION(BlueprintCallable, Category = "Interchange | Node | SpecularProfile")
+	bool SetCustomFormat(uint8 Format);
+
+	UFUNCTION(BlueprintCallable, Category = "Interchange | Node | SpecularProfile")
+	bool GetCustomFormat(uint8& Format) const;
+
+	UFUNCTION(BlueprintCallable, Category = "Interchange | Node | SpecularProfile")
+	bool GetCustomTexture(FString& TextureUid) const;
+
+	UFUNCTION(BlueprintCallable, Category = "Interchange | Node | SpecularProfile")
+	bool SetCustomTexture(const FString& TextureUid);
+
+private:
+
+	IMPLEMENT_NODE_ATTRIBUTE_KEY(Format)
+	IMPLEMENT_NODE_ATTRIBUTE_KEY(Texture)
+};
diff --git a/Engine/Plugins/Interchange/Runtime/Source/Pipelines/Private/InterchangeGenericMaterialPipeline.cpp b/Engine/Plugins/Interchange/Runtime/Source/Pipelines/Private/InterchangeGenericMaterialPipeline.cpp
--- a/Engine/Plugins/Interchange/Runtime/Source/Pipelines/Private/InterchangeGenericMaterialPipeline.cpp
+++ b/Engine/Plugins/Interchange/Runtime/Source/Pipelines/Private/InterchangeGenericMaterialPipeline.cpp
@@ -4,7 +4,6 @@
 
 #include "AssetRegistry/AssetData.h"
 #include "AssetRegistry/AssetRegistryModule.h"
-#include "CoreMinimal.h"
 #include "InterchangeGenericTexturePipeline.h"
 #include "InterchangeManager.h"
 #include "InterchangeMaterialDefinitions.h"
@@ -14,6 +13,8 @@
 #include "InterchangeShaderGraphNode.h"
 #include "InterchangeSourceData.h"
 #include "InterchangeSparseVolumeTexturePipeline.h"
+#include "InterchangeSpecularProfileNode.h"
+#include "InterchangeSpecularProfileFactoryNode.h"
 #include "InterchangeTexture2DArrayNode.h"
 #include "InterchangeTexture2DNode.h"
 #include "InterchangeTextureBlurNode.h"
@@ -26,7 +27,6 @@
 #include "Materials/MaterialExpressionConstant.h"
 #include "Materials/MaterialExpressionConstant3Vector.h"
 #include "Materials/MaterialExpressionCosine.h"
-#include "Materials/MaterialExpressionDivide.h"
 #include "Materials/MaterialExpressionFresnel.h"
 #include "Materials/MaterialExpressionLinearInterpolate.h"
 #include "Materials/MaterialExpressionMaterialFunctionCall.h"
@@ -38,6 +38,7 @@
 #include "Materials/MaterialExpressionScalarParameter.h"
 #include "Materials/MaterialExpressionSine.h"
 #include "Materials/MaterialExpressionStaticBoolParameter.h"
+#include "Materials/MaterialExpressionSubstrate.h"
 #include "Materials/MaterialExpressionSwitch.h"
 #include "Materials/MaterialExpressionTangent.h"
 #include "Materials/MaterialExpressionTextureCoordinate.h"
@@ -51,7 +52,6 @@
 #include "Materials/MaterialExpressionTransformPosition.h"
 #include "Materials/MaterialExpressionVectorNoise.h"
 #include "Materials/MaterialExpressionVectorParameter.h"
-#include "Materials/MaterialFunction.h"
 #include "Materials/MaterialInstanceConstant.h"
 #include "Materials/MaterialInstanceDynamic.h"
 #include "MaterialX/MaterialExpressions/MaterialExpressionSwizzle.h"
@@ -60,7 +60,6 @@
 #include "Misc/PackageName.h"
 #include "Misc/Paths.h"
 #include "Nodes/InterchangeBaseNode.h"
-#include "Nodes/InterchangeBaseNode.h"
 #include "Nodes/InterchangeBaseNodeContainer.h"
 #include "Nodes/InterchangeSourceNode.h"
 #include "Nodes/InterchangeUserDefinedAttribute.h"
@@ -335,6 +334,19 @@ namespace UE::Interchange::InterchangeGenericMaterialPipeline::Private
 
 		return FoundMaterial;
 	}
+
+	void AddSpecularProfileToFactoryNode(const UInterchangeBaseNode* ShaderNode, UInterchangeFactoryBaseNode* FactoryNode, UInterchangeBaseNodeContainer* BaseNodeContainer)
+	{
+		using namespace UE::Interchange::Materials;
+		if (FString SpecularProfileUid; ShaderNode->GetStringAttribute(SubstrateMaterial::SpecularProfile.ToString(), SpecularProfileUid))
+		{
+			if (BaseNodeContainer->GetNode(SpecularProfileUid))
+			{
+				FactoryNode->AddStringAttribute(SubstrateMaterial::SpecularProfile.ToString(), SpecularProfileUid);
+				FactoryNode->AddFactoryDependencyUid(UInterchangeFactoryBaseNode::BuildFactoryNodeUid(SpecularProfileUid));
+			}
+		}
+	}
 }
 
 namespace UE::Interchange::Materials::HashUtils
@@ -963,6 +975,8 @@ void UInterchangeGenericMaterialPipeline::ExecutePipeline(UInterchangeBaseNodeCo
 		const UClass* MaterialClass = (EditorUtilities && EditorUtilities->IsRuntimeOrPIE()) ? UMaterialInstanceDynamic::StaticClass() : UMaterialInstanceConstant::StaticClass();
 		MaterialFactoryNode->SetCustomInstanceClassName(MaterialClass->GetPathName());
 
+		UE::Interchange::InterchangeGenericMaterialPipeline::Private::AddSpecularProfileToFactoryNode(MaterialNode, MaterialFactoryNode, BaseNodeContainer);
+		
 		TArray<FString> Inputs;
 		UInterchangeShaderPortsAPI::GatherInputs(MaterialNode, Inputs);
 
@@ -1037,6 +1051,18 @@ void UInterchangeGenericMaterialPipeline::ExecutePipeline(UInterchangeBaseNodeCo
 			BaseMaterialNodes[0]->SetDisplayLabel(OverrideAssetName);
 		}
 	}
+
+
+	TArray<UInterchangeSpecularProfileNode*> SpecularProfileNodes;
+	BaseNodeContainer->IterateNodesOfType<UInterchangeSpecularProfileNode>([&](const FString& NodeUid, UInterchangeSpecularProfileNode* Node)
+																			 {
+																			   SpecularProfileNodes.Add(Node);
+																			 });
+
+	for (const UInterchangeSpecularProfileNode* SpecularProfileNode : SpecularProfileNodes)
+	{
+		CreateSpecularProfileFactoryNode(SpecularProfileNode);
+	}
 }
 
 void UInterchangeGenericMaterialPipeline::ExecutePostFactoryPipeline(const UInterchangeBaseNodeContainer* InBaseNodeContainer, const FString& NodeKey, UObject* CreatedAsset, bool bIsAReimport)
@@ -2856,6 +2882,58 @@ void UInterchangeGenericMaterialPipeline::HandleSwitchNode(const UInterchangeSha
 	}
 }
 
+void UInterchangeGenericMaterialPipeline::HandleSlabBSDFNode(const UInterchangeShaderNode* ShaderNode, UInterchangeBaseMaterialFactoryNode* MaterialFactoryNode, UInterchangeMaterialExpressionFactoryNode* SlabBSDFFactoryNode)
+{
+	using namespace UE::Interchange::Materials::Standard::Nodes;
+	using namespace UE::Interchange::InterchangeGenericMaterialPipeline::Private;
+
+	SlabBSDFFactoryNode->SetCustomExpressionClassName(UMaterialExpressionSubstrateSlabBSDF::StaticClass()->GetName());
+
+	// ensure that SlabBSDF expression have the same members define as Interchange;
+#define SubstrateSlabBSDF_MEMBER_CHECKED(X) GET_MEMBER_NAME_STRING_CHECKED(UMaterialExpressionSubstrateSlabBSDF, X), SlabBSDF::Inputs::X.ToString()
+	const TSet<FString> SlabInputs{
+		SubstrateSlabBSDF_MEMBER_CHECKED(Anisotropy),
+		SubstrateSlabBSDF_MEMBER_CHECKED(DiffuseAlbedo),
+		SubstrateSlabBSDF_MEMBER_CHECKED(EmissiveColor),
+		SubstrateSlabBSDF_MEMBER_CHECKED(F0),
+		SubstrateSlabBSDF_MEMBER_CHECKED(F90),
+		SubstrateSlabBSDF_MEMBER_CHECKED(FuzzAmount),
+		SubstrateSlabBSDF_MEMBER_CHECKED(FuzzColor),
+		SubstrateSlabBSDF_MEMBER_CHECKED(FuzzRoughness),
+		SubstrateSlabBSDF_MEMBER_CHECKED(GlintUV),
+		SubstrateSlabBSDF_MEMBER_CHECKED(GlintValue),
+		SubstrateSlabBSDF_MEMBER_CHECKED(Normal),
+		SubstrateSlabBSDF_MEMBER_CHECKED(Roughness),
+		SubstrateSlabBSDF_MEMBER_CHECKED(SecondRoughness),
+		SubstrateSlabBSDF_MEMBER_CHECKED(SecondRoughnessWeight),
+		SubstrateSlabBSDF_MEMBER_CHECKED(SSSMFP),
+		SubstrateSlabBSDF_MEMBER_CHECKED(SSSMFPScale),
+		SubstrateSlabBSDF_MEMBER_CHECKED(SSSPhaseAnisotropy),
+		SubstrateSlabBSDF_MEMBER_CHECKED(Tangent),
+	};
+#undef SubstrateSlabBSDF_MEMBER_CHECKED
+
+	TArray<FString> Inputs;
+	UInterchangeShaderPortsAPI::GatherInputs(ShaderNode, Inputs);
+
+	for (const FString& InputName : Inputs)
+	{
+		if(SlabInputs.Contains(InputName))
+		{
+			TTuple<UInterchangeMaterialExpressionFactoryNode*, FString> InputExpression =
+				CreateMaterialExpressionForInput(MaterialFactoryNode, ShaderNode, InputName, SlabBSDFFactoryNode->GetUniqueID());
+
+			if (InputExpression.Get<0>())
+			{
+				UInterchangeShaderPortsAPI::ConnectOuputToInputByName(SlabBSDFFactoryNode, InputName, InputExpression.Get<0>()->GetUniqueID(), InputExpression.Get<1>());
+			}
+		}
+	}
+
+	// Specular Profile
+	AddSpecularProfileToFactoryNode(ShaderNode, SlabBSDFFactoryNode, BaseNodeContainer);
+}
+
 /** Trigonometry expressions are unitless and sets the period to 1 by default to let the user sets their period*/
 void UInterchangeGenericMaterialPipeline::HandleTrigonometryNode(const UInterchangeShaderNode* ShaderNode, UClass* StaticClass, UInterchangeBaseMaterialFactoryNode* MaterialFactoryNode, UInterchangeMaterialExpressionFactoryNode* TrigonometryFactoryNode)
 {
@@ -3049,6 +3127,10 @@ UInterchangeMaterialExpressionFactoryNode* UInterchangeGenericMaterialPipeline::
 	{
 		HandleVectorNoiseNode(ShaderNode, MaterialFactoryNode, MaterialExpression);
 	}
+	else if (*ShaderType == Nodes::SlabBSDF::Name)
+	{
+		HandleSlabBSDFNode(ShaderNode, MaterialFactoryNode, MaterialExpression);
+	}
 	else if(*ShaderType == Nodes::Cosine::Name)
 	{
 		HandleTrigonometryNode(ShaderNode, UMaterialExpressionCosine::StaticClass(), MaterialFactoryNode, MaterialExpression);
@@ -3416,6 +3498,8 @@ UInterchangeMaterialFactoryNode* UInterchangeGenericMaterialPipeline::CreateMate
 
 UInterchangeMaterialInstanceFactoryNode* UInterchangeGenericMaterialPipeline::CreateMaterialInstanceFactoryNode(const UInterchangeShaderGraphNode* ShaderGraphNode)
 {
+	using namespace UE::Interchange::InterchangeGenericMaterialPipeline::Private;
+
 	UInterchangeMaterialInstanceFactoryNode* MaterialInstanceFactoryNode =
 		Cast<UInterchangeMaterialInstanceFactoryNode>( CreateBaseMaterialFactoryNode(ShaderGraphNode, UInterchangeMaterialInstanceFactoryNode::StaticClass()) );
 
@@ -3488,10 +3572,39 @@ UInterchangeMaterialInstanceFactoryNode* UInterchangeGenericMaterialPipeline::Cr
 #endif
 
 	VisitShaderGraphNode(ShaderGraphNode, MaterialInstanceFactoryNode);
+	
+	AddSpecularProfileToFactoryNode(ShaderGraphNode, MaterialInstanceFactoryNode, BaseNodeContainer);
 
 	return MaterialInstanceFactoryNode;
 }
 
+void UInterchangeGenericMaterialPipeline::CreateSpecularProfileFactoryNode(const UInterchangeSpecularProfileNode* SpecularProfileNode)
+{
+	const FString FactoryNodeUid = UInterchangeFactoryBaseNode::BuildFactoryNodeUid(SpecularProfileNode->GetUniqueID());
+
+	UInterchangeSpecularProfileFactoryNode* FactoryNode = NewObject<UInterchangeSpecularProfileFactoryNode>(BaseNodeContainer, NAME_None);
+
+	FactoryNode->InitializeNode(FactoryNodeUid, SpecularProfileNode->GetDisplayLabel(), EInterchangeNodeContainerType::FactoryData);
+	FactoryNode->SetEnabled(true);
+	
+	if (uint8 Format; SpecularProfileNode->GetCustomFormat(Format))
+	{
+		FactoryNode->SetCustomFormat(ESpecularProfileFormat{ Format });
+	}
+
+	if (FString TextureUid; SpecularProfileNode->GetCustomTexture(TextureUid))
+	{
+		if(BaseNodeContainer->GetNode(TextureUid))
+		{
+			FactoryNode->SetCustomTexture(TextureUid);
+			FactoryNode->AddFactoryDependencyUid(UInterchangeFactoryBaseNode::BuildFactoryNodeUid(TextureUid));
+		}
+	}
+	BaseNodeContainer->AddNode(FactoryNode);
+	FactoryNode->AddTargetNodeUid(SpecularProfileNode->GetUniqueID());
+	SpecularProfileNode->AddTargetNodeUid(FactoryNode->GetUniqueID());
+}
+
 void UInterchangeGenericMaterialPipeline::VisitShaderGraphNode(const UInterchangeShaderGraphNode* ShaderGraphNode, UInterchangeMaterialInstanceFactoryNode* MaterialInstanceFactoryNode) const
 {
 	TArray<FString> Inputs;
diff --git a/Engine/Plugins/Interchange/Runtime/Source/Pipelines/Public/InterchangeGenericMaterialPipeline.h b/Engine/Plugins/Interchange/Runtime/Source/Pipelines/Public/InterchangeGenericMaterialPipeline.h
--- a/Engine/Plugins/Interchange/Runtime/Source/Pipelines/Public/InterchangeGenericMaterialPipeline.h
+++ b/Engine/Plugins/Interchange/Runtime/Source/Pipelines/Public/InterchangeGenericMaterialPipeline.h
@@ -22,6 +22,7 @@ class UInterchangeResult;
 class UInterchangeShaderGraphNode;
 class UInterchangeShaderNode;
 class UInterchangeSparseVolumeTexturePipeline;
+class UInterchangeSpecularProfileNode;
 class UInterchangeTextureNode;
 class UMaterialFunction;
 
@@ -136,6 +137,7 @@ private:
 	UInterchangeMaterialFactoryNode* CreateMaterialFactoryNode(const UInterchangeShaderGraphNode* ShaderGraphNode);
 	UInterchangeMaterialFunctionFactoryNode* CreateMaterialFunctionFactoryNode(const UInterchangeShaderGraphNode* FunctionCallShaderNode);
 	UInterchangeMaterialInstanceFactoryNode* CreateMaterialInstanceFactoryNode(const UInterchangeShaderGraphNode* ShaderGraphNode);
+	void CreateSpecularProfileFactoryNode(const UInterchangeSpecularProfileNode* SpecularProfileNode);
 	
 	/** True if the shader graph has a clear coat input. */
 	bool HasClearCoat(const UInterchangeShaderGraphNode* ShaderGraphNode) const;
@@ -224,6 +226,7 @@ private:
 	void HandleVectorNoiseNode(const UInterchangeShaderNode* ShaderNode, UInterchangeBaseMaterialFactoryNode* MaterialFactoryNode, UInterchangeMaterialExpressionFactoryNode* NoiseFactoryNode);
 	void HandleSwizzleNode(const UInterchangeShaderNode* ShaderNode, UInterchangeBaseMaterialFactoryNode* MaterialFactoryNode, UInterchangeMaterialExpressionFactoryNode* SwizzleFactoryNode);
 	void HandleSwitchNode(const UInterchangeShaderNode* ShaderNode, UInterchangeBaseMaterialFactoryNode* MaterialFactoryNode, UInterchangeMaterialExpressionFactoryNode* SwitchFactoryNode);
+	void HandleSlabBSDFNode(const UInterchangeShaderNode* ShaderNode, UInterchangeBaseMaterialFactoryNode* MaterialFactoryNode, UInterchangeMaterialExpressionFactoryNode* SlabBSDFFactoryNode);
 	void HandleTrigonometryNode(const UInterchangeShaderNode* ShaderNode, UClass* StaticClass, UInterchangeBaseMaterialFactoryNode* MaterialFactoryNode, UInterchangeMaterialExpressionFactoryNode* TrigonometryFactoryNode);
 
 	void HandleScalarParameterNode(const UInterchangeShaderNode* ShaderNode, UInterchangeMaterialExpressionFactoryNode* ScalarParameterFactoryNode);
