USD:
- Store the primvars in the mesh description

MaterialX:
- Rename input of MX_ColorCorrect from In to Input
- Look for a custom name of the input (the one from the Material Expression), if not found just take the name of the input
- Add an attribute to the Image ShaderNode if it's a geompropvalue converted to an image
- Put the static variables in the manager as different data structures that operate on the Node/Input/NodeDef of the graph

MaterialFactory:
- Store the path of the Material Expression in the factory node to be retrieved later by the pipeline in the ExecutePostFactoryImport

USDTranslator:
- Only add the attribute of the geomprop on the mesh if it wasn't already processed

USDPipeline:
- Add the baking of the primvars to textures, plus update the TextureSample material expressions with the newly created textures and update as well the correct UV channel

#virtualized

diff --git a/Engine/Plugins/Interchange/Runtime/Interchange.uplugin b/Engine/Plugins/Interchange/Runtime/Interchange.uplugin
--- a/Engine/Plugins/Interchange/Runtime/Interchange.uplugin
+++ b/Engine/Plugins/Interchange/Runtime/Interchange.uplugin
@@ -102,6 +102,11 @@
 		{
 			"Name": "InterchangeAssets",
 			"Enabled": true
+		},
+		{
+			"Name": "MeshModelingToolset",
+			"Enabled": true,
+			"TargetAllowList": [ "Editor" ]
 		}
 	],
 	"IsExperimentalVersion": false,
diff --git a/Engine/Plugins/Interchange/Runtime/Source/Common/Public/MaterialX/InterchangeMaterialXDefinitions.h b/Engine/Plugins/Interchange/Runtime/Source/Common/Public/MaterialX/InterchangeMaterialXDefinitions.h
--- a/Engine/Plugins/Interchange/Runtime/Source/Common/Public/MaterialX/InterchangeMaterialXDefinitions.h
+++ b/Engine/Plugins/Interchange/Runtime/Source/Common/Public/MaterialX/InterchangeMaterialXDefinitions.h
@@ -282,6 +282,30 @@ namespace UE
 			static constexpr const TCHAR* Y = TEXT("Y");
 		}
 	}
+
+	namespace MaterialFunctions
+	{
+		namespace Path
+		{
+			constexpr const TCHAR* HeightToNormalSmooth = TEXT("/Engine/Functions/Engine_MaterialFunctions03/Procedurals/HeightToNormalSmooth.HeightToNormalSmooth");
+			constexpr const TCHAR* MxACEScgToRec709 = TEXT("/InterchangeAssets/Functions/MX_ACEScgToRec709.MX_ACEScgToRec709");
+			constexpr const TCHAR* MxAdobeRGBToRec709 = TEXT("/InterchangeAssets/Functions/MX_AdobeRGBToRec709.MX_AdobeRGBToRec709");
+			constexpr const TCHAR* MxArtisticIOR = TEXT("/InterchangeAssets/Functions/MX_Artistic_IOR.MX_Artistic_IOR");
+			constexpr const TCHAR* MxColorCorrect = TEXT("/InterchangeAssets/Functions/MX_ColorCorrect.MX_ColorCorrect");
+			constexpr const TCHAR* MxHsvAdjust = TEXT("/InterchangeAssets/Functions/MX_HSVAdjust.MX_HSVAdjust");
+			constexpr const TCHAR* MxLinearAdobeRGBToRec709 = TEXT("/InterchangeAssets/Functions/MX_LinearAdobeRGBToRec709.MX_LinearAdobeRGBToRec709");
+			constexpr const TCHAR* MxLinearDisplayP3ToRec709 = TEXT("/InterchangeAssets/Functions/MX_LinearDisplayP3ToRec709.MX_LinearDisplayP3ToRec709");
+			constexpr const TCHAR* MxPlace2D = TEXT("/InterchangeAssets/Functions/MX_Place2D.MX_Place2D");
+			constexpr const TCHAR* MxRandomFloat = TEXT("/InterchangeAssets/Functions/MX_RandomFloat.MX_RandomFloat");
+			constexpr const TCHAR* MxRandomColor = TEXT("/InterchangeAssets/Functions/MX_RandomColor.MX_RandomColor");
+			constexpr const TCHAR* MxRoughnesAnisotropy = TEXT("/InterchangeAssets/Functions/MX_Roughness_Anisotropy.MX_Roughness_Anisotropy");
+			constexpr const TCHAR* MxRoughnessDual = TEXT("/InterchangeAssets/Functions/MX_Roughness_Dual.MX_Roughness_Dual");
+			constexpr const TCHAR* MxSrgbDisplayP3ToRec709 = TEXT("/InterchangeAssets/Functions/MX_SrgbDisplayP3ToRec709.MX_SrgbDisplayP3ToRec709");
+			constexpr const TCHAR* MxSrgbToRec709 = TEXT("/InterchangeAssets/Functions/MX_SrgbToRec709.MX_SrgbToRec709");
+			constexpr const TCHAR* NormalFromHeightMap = TEXT("/Engine/Functions/Engine_MaterialFunctions03/Procedurals/NormalFromHeightmap.NormalFromHeightmap");
+			constexpr const TCHAR* Refract = TEXT("/Engine/Functions/Engine_MaterialFunctions01/Vectors/Refract.Refract");
+		}
+	}
 }
 
 MATERIALX_NAMESPACE_BEGIN
diff --git a/Engine/Plugins/Interchange/Runtime/Source/Common/Public/Usd/InterchangeUsdDefinitions.h b/Engine/Plugins/Interchange/Runtime/Source/Common/Public/Usd/InterchangeUsdDefinitions.h
--- a/Engine/Plugins/Interchange/Runtime/Source/Common/Public/Usd/InterchangeUsdDefinitions.h
+++ b/Engine/Plugins/Interchange/Runtime/Source/Common/Public/Usd/InterchangeUsdDefinitions.h
@@ -35,6 +35,10 @@ namespace UE::Interchange::USD
 	namespace Primvar
 	{
 		const FString Number = TEXT("USD_PrimvarNumber");
-		const FString Name = TEXT("USD_PrimvarName"); // In case of "Number" primvars, the user should concatenate the Index to this attribute
+		// In case of "Number" primvars, the user should concatenate the Index to this attribute
+		const FString Name = TEXT("USD_PrimvarName");
+		const FString TangentSpace = TEXT("USD_PrimvarTangentSpace");
+		// Attribute that informs about the UID of a ShaderNode of type TextureSample
+		const FString ShaderNodeTextureSample = TEXT("USD_ShaderNodeTextureSample");
 	}
 }
diff --git a/Engine/Plugins/Interchange/Runtime/Source/Import/Private/Material/InterchangeMaterialFactory.cpp b/Engine/Plugins/Interchange/Runtime/Source/Import/Private/Material/InterchangeMaterialFactory.cpp
--- a/Engine/Plugins/Interchange/Runtime/Source/Import/Private/Material/InterchangeMaterialFactory.cpp
+++ b/Engine/Plugins/Interchange/Runtime/Source/Import/Private/Material/InterchangeMaterialFactory.cpp
@@ -247,6 +247,10 @@ namespace UE::Interchange::MaterialFactory::Internal
 				}
 			}
 
+			// We store the path of the Material Expression, in order for Pipelines to retrieve them later in the ExecutePostFactoryImport
+			using namespace UE::Interchange::Materials;
+			const_cast<UInterchangeMaterialExpressionFactoryNode*>(&ExpressionNode)->AddStringAttribute(Factory::Expression::Path.ToString(), MaterialExpression->GetPathName());
+
 			Expressions.Add(ExpressionNode.GetUniqueID()) = MaterialExpression;
 
 			TArray<FString> Inputs;
diff --git a/Engine/Plugins/Interchange/Runtime/Source/Import/Private/MaterialX/MaterialXUtils/MaterialXBase.cpp b/Engine/Plugins/Interchange/Runtime/Source/Import/Private/MaterialX/MaterialXUtils/MaterialXBase.cpp
--- a/Engine/Plugins/Interchange/Runtime/Source/Import/Private/MaterialX/MaterialXUtils/MaterialXBase.cpp
+++ b/Engine/Plugins/Interchange/Runtime/Source/Import/Private/MaterialX/MaterialXUtils/MaterialXBase.cpp
@@ -17,47 +17,10 @@ void FMaterialXBase::UpdateDocumentRecursively(MaterialX::GraphElementPtr Graph)
 {
 	AddTexCoordToTiledImageNodes(Graph);
 
-	static std::set<std::string> SkippedCategories{
-		mx::Category::ACEScgToLinRec709,
-		mx::Category::AdobeRgbToLinRec709,
-		mx::Category::ColorCorrect,
-		mx::Category::Contrast,
-		mx::Category::Distance,
-		mx::Category::Extract,
-		mx::Category::Fractal3D,
-		mx::Category::HsvAdjust,
-		mx::Category::LinAdobeRgbToLinRec709,
-		mx::Category::LinDisplayP3ToLinRec709,
-		mx::Category::Noise2D,
-		mx::Category::Noise3D,
-		mx::Category::OpenPBRSurface,
-		mx::Category::Overlay,
-		mx::Category::Place2D,
-		mx::Category::Ramp4,
-		mx::Category::RandomFloat,
-		mx::Category::RandomColor,
-		mx::Category::Range,
-		mx::Category::Refract,
-		mx::Category::Saturate,
-		mx::Category::Separate2,
-		mx::Category::Separate3,
-		mx::Category::Separate4,
-		mx::Category::SrgbDisplayP3ToLinRec709,
-		mx::Category::SrgbTextureToLinRec709,
-		mx::Category::StandardSurface,
-		mx::Category::UsdPreviewSurface
-	};
-
-	auto FilterNode = [](mx::NodePtr Node)
-	{
-		const std::string& Category = Node->getCategory();
-		return SkippedCategories.find(Category) == SkippedCategories.end();
-	};
-
 	//This function will replace all the nodes that are defined with a nodegraph with their equivalent node network
 	//This allows us to handle nodes defined with a nodegraph, with nodes that we already support while traversing the current node graph
 	//For example the node tiledimage is just a succession of the nodes image, multiply, add, etc.
-	Graph->flattenSubgraphs(mx::EMPTY_STRING, FilterNode);
+	Graph->flattenSubgraphs(mx::EMPTY_STRING, [](mx::NodePtr Node){return FMaterialXManager::GetInstance().FilterNodeGraph(Node);} );
 
 	ConvertNeutralNodesToDot(Graph);
 
diff --git a/Engine/Plugins/Interchange/Runtime/Source/Import/Private/MaterialX/MaterialXUtils/MaterialXManager.cpp b/Engine/Plugins/Interchange/Runtime/Source/Import/Private/MaterialX/MaterialXUtils/MaterialXManager.cpp
--- a/Engine/Plugins/Interchange/Runtime/Source/Import/Private/MaterialX/MaterialXUtils/MaterialXManager.cpp
+++ b/Engine/Plugins/Interchange/Runtime/Source/Import/Private/MaterialX/MaterialXUtils/MaterialXManager.cpp
@@ -27,8 +27,6 @@ namespace mx = MaterialX;
 //not a good solution to use semicolon because of drive disk on Windows
 const TCHAR FMaterialXManager::TexturePayloadSeparator = TEXT('{');
 
-#define MX_MATERIALFUNCTION(AssetName) TEXT("/InterchangeAssets/") + FString{bIsSubstrateEnabled ? TEXT("Substrate/") : TEXT("Functions/")} + TEXT(AssetName) TEXT(".") TEXT(AssetName)
-
 FMaterialXManager::FMaterialXManager()
 	: MatchingInputNames{
 	   {{TEXT(""),                            TEXT("amplitude")},   ExpressionInput(UE::Expressions::Inputs::Amplitude)},
@@ -223,25 +221,67 @@ FMaterialXManager::FMaterialXManager()
 		{mx::Category::AbsorptionVDF,			FMaterialXMaterialFunction{TInPlaceType<EInterchangeMaterialXVDF>{}, EInterchangeMaterialXVDF::Absorption}},
 		{mx::Category::AnisotropicVDF,			FMaterialXMaterialFunction{TInPlaceType<EInterchangeMaterialXVDF>{}, EInterchangeMaterialXVDF::Anisotropic}},
 		// Utility nodes
-		{mx::Category::ArtisticIOR,				FMaterialXMaterialFunction{TInPlaceType<FString>{}, TEXT("/InterchangeAssets/Functions/MX_Artistic_IOR.MX_Artistic_IOR")}},
-		{mx::Category::RoughnessAnisotropy,		FMaterialXMaterialFunction{TInPlaceType<FString>{}, TEXT("/InterchangeAssets/Functions/MX_Roughness_Anisotropy.MX_Roughness_Anisotropy")}},
-		{mx::Category::RoughnessDual,			FMaterialXMaterialFunction{TInPlaceType<FString>{}, TEXT("/InterchangeAssets/Functions/MX_Roughness_Dual.MX_Roughness_Dual")}},
+		{mx::Category::ArtisticIOR,				FMaterialXMaterialFunction{TInPlaceType<FString>{}, UE::MaterialFunctions::Path::MxArtisticIOR}},
+		{mx::Category::RoughnessAnisotropy,		FMaterialXMaterialFunction{TInPlaceType<FString>{}, UE::MaterialFunctions::Path::MxRoughnesAnisotropy}},
+		{mx::Category::RoughnessDual,			FMaterialXMaterialFunction{TInPlaceType<FString>{}, UE::MaterialFunctions::Path::MxRoughnessDual}},
 		// Math
-		{mx::Category::Place2D,					FMaterialXMaterialFunction{TInPlaceType<FString>{}, TEXT("/InterchangeAssets/Functions/MX_Place2D.MX_Place2D")}},
-		{mx::Category::Refract,					FMaterialXMaterialFunction{TInPlaceType<FString>{}, TEXT("/Engine/Functions/Engine_MaterialFunctions01/Vectors/Refract.Refract")}},
+		{mx::Category::Place2D,					FMaterialXMaterialFunction{TInPlaceType<FString>{}, UE::MaterialFunctions::Path::MxPlace2D}},
+		{mx::Category::Refract,					FMaterialXMaterialFunction{TInPlaceType<FString>{}, UE::MaterialFunctions::Path::Refract}},
 		// ColorTransform
-		{mx::Category::ACEScgToLinRec709,		FMaterialXMaterialFunction{TInPlaceType<FString>{}, TEXT("/InterchangeAssets/Functions/MX_ACEScgToRec709.MX_ACEScgToRec709")}},
-		{mx::Category::AdobeRgbToLinRec709,		FMaterialXMaterialFunction{TInPlaceType<FString>{}, TEXT("/InterchangeAssets/Functions/MX_AdobeRGBToRec709.MX_AdobeRGBToRec709")}},
-		{mx::Category::LinAdobeRgbToLinRec709,	FMaterialXMaterialFunction{TInPlaceType<FString>{}, TEXT("/InterchangeAssets/Functions/MX_LinearAdobeRGBToRec709.MX_LinearAdobeRGBToRec709")}},
-		{mx::Category::LinDisplayP3ToLinRec709,	FMaterialXMaterialFunction{TInPlaceType<FString>{}, TEXT("/InterchangeAssets/Functions/MX_LinearDisplayP3ToRec709.MX_LinearDisplayP3ToRec709")}},
-		{mx::Category::SrgbDisplayP3ToLinRec709,FMaterialXMaterialFunction{TInPlaceType<FString>{}, TEXT("/InterchangeAssets/Functions/MX_SrgbDisplayP3ToRec709.MX_SrgbDisplayP3ToRec709")}},
-		{mx::Category::SrgbTextureToLinRec709,	FMaterialXMaterialFunction{TInPlaceType<FString>{}, TEXT("/InterchangeAssets/Functions/MX_SrgbToRec709.MX_SrgbToRec709")}},
+		{mx::Category::ACEScgToLinRec709,		FMaterialXMaterialFunction{TInPlaceType<FString>{}, UE::MaterialFunctions::Path::MxACEScgToRec709}},
+		{mx::Category::AdobeRgbToLinRec709,		FMaterialXMaterialFunction{TInPlaceType<FString>{}, UE::MaterialFunctions::Path::MxAdobeRGBToRec709}},
+		{mx::Category::LinAdobeRgbToLinRec709,	FMaterialXMaterialFunction{TInPlaceType<FString>{}, UE::MaterialFunctions::Path::MxLinearAdobeRGBToRec709}},
+		{mx::Category::LinDisplayP3ToLinRec709,	FMaterialXMaterialFunction{TInPlaceType<FString>{}, UE::MaterialFunctions::Path::MxLinearDisplayP3ToRec709}},
+		{mx::Category::SrgbDisplayP3ToLinRec709,FMaterialXMaterialFunction{TInPlaceType<FString>{}, UE::MaterialFunctions::Path::MxSrgbDisplayP3ToRec709}},
+		{mx::Category::SrgbTextureToLinRec709,	FMaterialXMaterialFunction{TInPlaceType<FString>{}, UE::MaterialFunctions::Path::MxSrgbToRec709}},
 		// Procedural
-		{mx::Category::RandomFloat,	            FMaterialXMaterialFunction{TInPlaceType<FString>{}, TEXT("/InterchangeAssets/Functions/MX_RandomFloat.MX_RandomFloat")}},
-		{mx::Category::RandomColor,	            FMaterialXMaterialFunction{TInPlaceType<FString>{}, TEXT("/InterchangeAssets/Functions/MX_RandomColor.MX_RandomColor")}},
+		{mx::Category::RandomFloat,	            FMaterialXMaterialFunction{TInPlaceType<FString>{}, UE::MaterialFunctions::Path::MxRandomFloat}},
+		{mx::Category::RandomColor,	            FMaterialXMaterialFunction{TInPlaceType<FString>{}, UE::MaterialFunctions::Path::MxRandomColor}},
 		// Adjustment
-		{mx::Category::ColorCorrect,	        FMaterialXMaterialFunction{TInPlaceType<FString>{}, TEXT("/InterchangeAssets/Functions/MX_ColorCorrect.MX_ColorCorrect")}},
-		{mx::Category::HsvAdjust,	            FMaterialXMaterialFunction{TInPlaceType<FString>{}, TEXT("/InterchangeAssets/Functions/MX_HSVAdjust.MX_HSVAdjust")}},
+		{mx::Category::ColorCorrect,	        FMaterialXMaterialFunction{TInPlaceType<FString>{}, UE::MaterialFunctions::Path::MxColorCorrect}},
+		{mx::Category::HsvAdjust,	            FMaterialXMaterialFunction{TInPlaceType<FString>{}, UE::MaterialFunctions::Path::MxHsvAdjust}},
+	}
+	, CategoriesToSkip {
+		mx::Category::ACEScgToLinRec709,
+		mx::Category::AdobeRgbToLinRec709,
+		mx::Category::ColorCorrect,
+		mx::Category::Contrast,
+		mx::Category::Distance,
+		mx::Category::Extract,
+		mx::Category::Fractal3D,
+		mx::Category::HsvAdjust,
+		mx::Category::LinAdobeRgbToLinRec709,
+		mx::Category::LinDisplayP3ToLinRec709,
+		mx::Category::Noise2D,
+		mx::Category::Noise3D,
+		mx::Category::OpenPBRSurface,
+		mx::Category::Overlay,
+		mx::Category::Place2D,
+		mx::Category::Ramp4,
+		mx::Category::RandomFloat,
+		mx::Category::RandomColor,
+		mx::Category::Range,
+		mx::Category::Refract,
+		mx::Category::Saturate,
+		mx::Category::Separate2,
+		mx::Category::Separate3,
+		mx::Category::Separate4,
+		mx::Category::SrgbDisplayP3ToLinRec709,
+		mx::Category::SrgbTextureToLinRec709,
+		mx::Category::StandardSurface,
+		mx::Category::UsdPreviewSurface
+	}
+	, NodeDefsCategories{
+		mx::Category::IfEqual,
+		mx::Category::IfGreater,
+		mx::Category::IfGreaterEq,
+		mx::Category::Max,
+		mx::Category::Min,
+		mx::Category::Modulo,
+	}
+	, NodeInputsToRemove{
+		{mx::Category::GeomColor, {"index"}},                          // There's only one set of VertexColor
+		{mx::Category::NormalMap, {"normal", "tangent", "bitangent"}}, // FlattenNormalMap Material Function doesn't have normal/tangent inputs, we just remove them to avoid unnecessary connections
 	}
 	, bIsSubstrateEnabled{ IInterchangeImportModule::IsAvailable() ? IInterchangeImportModule::Get().IsSubstrateEnabled() : false }
 {
@@ -345,6 +385,31 @@ bool FMaterialXManager::IsSubstrateEnabled() const
 	return bIsSubstrateEnabled;
 }
 
+bool FMaterialXManager::FilterNodeGraph(MaterialX::NodePtr Node) const
+{
+	// the test seems counterintuitive, but MaterialX check is "!filter" in the flattenSubgraphs functions
+	return CategoriesToSkip.find(Node->getCategory()) == CategoriesToSkip.cend();
+}
+
+void FMaterialXManager::RemoveInputs(MaterialX::NodePtr Node) const
+{
+	if (auto It = NodeInputsToRemove.find(Node->getCategory()); It != NodeInputsToRemove.cend())
+	{
+		for(const std::string& InputName: It->second)
+		{
+			Node->removeInput(InputName);
+		}
+	}
+}
+
+void FMaterialXManager::AddInputsFromNodeDef(MaterialX::NodePtr Node) const
+{
+	if (NodeDefsCategories.find(Node->getCategory()) != NodeDefsCategories.cend())
+	{
+		Node->addInputsFromNodeDef();
+	}
+}
+
 const TCHAR* FMaterialXManager::ExpressionInput(const TCHAR* Input)
 {
 	MaterialExpressionInputs.FindOrAdd(Input);
@@ -382,23 +447,23 @@ namespace UE::Interchange::MaterialX
 			return bAllLoaded;
 		};
 
-		static const bool bPackagesLoaded =	ArePackagesLoaded({ TEXT("MaterialFunction'/Engine/Functions/Engine_MaterialFunctions03/Procedurals/NormalFromHeightmap.NormalFromHeightmap'"),
-															    TEXT("MaterialFunction'/InterchangeAssets/Functions/MX_Artistic_IOR.MX_Artistic_IOR'"),
-															    TEXT("MaterialFunction'/InterchangeAssets/Functions/MX_Roughness_Anisotropy.MX_Roughness_Anisotropy'"), 
-															    TEXT("MaterialFunction'/InterchangeAssets/Functions/MX_Roughness_Dual.MX_Roughness_Dual'"),
-																TEXT("MaterialFunction'/InterchangeAssets/Functions/MX_Place2D.MX_Place2D'"),
-																TEXT("MaterialFunction'/Engine/Functions/Engine_MaterialFunctions01/Vectors/Refract.Refract'"),
-																TEXT("MaterialFunction'/InterchangeAssets/Functions/MX_ACEScgToRec709.MX_ACEScgToRec709'"),
-																TEXT("MaterialFunction'/InterchangeAssets/Functions/MX_AdobeRGBToRec709.MX_AdobeRGBToRec709'"),
-																TEXT("MaterialFunction'/InterchangeAssets/Functions/MX_LinearDisplayP3ToRec709.MX_LinearDisplayP3ToRec709'"),
-																TEXT("MaterialFunction'/InterchangeAssets/Functions/MX_LinearAdobeRGBToRec709.MX_LinearAdobeRGBToRec709'"),
-																TEXT("MaterialFunction'/InterchangeAssets/Functions/MX_SrgbToRec709.MX_SrgbToRec709'"),
-																TEXT("MaterialFunction'/InterchangeAssets/Functions/MX_SrgbDisplayP3ToRec709.MX_SrgbDisplayP3ToRec709'"),
-																TEXT("MaterialFunction'/InterchangeAssets/Functions/MX_RandomFloat.MX_RandomFloat'"),
-																TEXT("MaterialFunction'/InterchangeAssets/Functions/MX_RandomColor.MX_RandomColor'"),
-																TEXT("MaterialFunction'/Engine/Functions/Engine_MaterialFunctions03/Procedurals/HeightToNormalSmooth.HeightToNormalSmooth'"),
-																TEXT("MaterialFunction'/InterchangeAssets/Functions/MX_ColorCorrect.MX_ColorCorrect'"),
-																TEXT("MaterialFunction'/InterchangeAssets/Functions/MX_HSVAdjust.MX_HSVAdjust'"),
+		static const bool bPackagesLoaded =	ArePackagesLoaded({ UE::MaterialFunctions::Path::NormalFromHeightMap,
+															    UE::MaterialFunctions::Path::MxArtisticIOR,
+															    UE::MaterialFunctions::Path::MxRoughnesAnisotropy,
+															    UE::MaterialFunctions::Path::MxRoughnessDual,
+																UE::MaterialFunctions::Path::MxPlace2D,
+																UE::MaterialFunctions::Path::Refract,
+																UE::MaterialFunctions::Path::MxACEScgToRec709,
+																UE::MaterialFunctions::Path::MxAdobeRGBToRec709,
+																UE::MaterialFunctions::Path::MxLinearDisplayP3ToRec709,
+																UE::MaterialFunctions::Path::MxLinearAdobeRGBToRec709,
+																UE::MaterialFunctions::Path::MxSrgbToRec709,
+																UE::MaterialFunctions::Path::MxSrgbDisplayP3ToRec709,
+																UE::MaterialFunctions::Path::MxRandomFloat,
+																UE::MaterialFunctions::Path::MxRandomColor,
+																UE::MaterialFunctions::Path::HeightToNormalSmooth,
+																UE::MaterialFunctions::Path::MxColorCorrect,
+																UE::MaterialFunctions::Path::MxHsvAdjust,
 															  });
 
 		return bPackagesLoaded;
@@ -407,5 +472,3 @@ namespace UE::Interchange::MaterialX
 #endif
 	}
 }
-
-#undef MX_MATERIALFUNCTION
\ No newline at end of file
diff --git a/Engine/Plugins/Interchange/Runtime/Source/Import/Private/MaterialX/MaterialXUtils/MaterialXSurfaceShaderAbstract.cpp b/Engine/Plugins/Interchange/Runtime/Source/Import/Private/MaterialX/MaterialXUtils/MaterialXSurfaceShaderAbstract.cpp
--- a/Engine/Plugins/Interchange/Runtime/Source/Import/Private/MaterialX/MaterialXUtils/MaterialXSurfaceShaderAbstract.cpp
+++ b/Engine/Plugins/Interchange/Runtime/Source/Import/Private/MaterialX/MaterialXUtils/MaterialXSurfaceShaderAbstract.cpp
@@ -259,17 +259,17 @@ bool FMaterialXSurfaceShaderAbstract::ConnectMatchingNodeOutputToInput(const FCo
 
 	bool bIsConnected = false;
 
-	auto ConnectOutputToInputInternal = [&](UInterchangeShaderNode* OperatorNode, bool bFindMatchingInput = true)
+	auto ConnectOutputToInputInternal = [&](UInterchangeShaderNode* OperatorNode)
 	{
 		for(mx::InputPtr Input : Connect.UpstreamNode->getInputs())
 		{
-			if(!bFindMatchingInput)
+			if (const FString* InputNameFound = Manager.FindMaterialExpressionInput(GetInputName(Input)))
 			{
-				AddAttributeFromValueOrInterface(Input, Input->getName().c_str(), OperatorNode);
+				AddAttributeFromValueOrInterface(Input, *InputNameFound, OperatorNode);
 			}
-			else if(const FString* InputNameFound = Manager.FindMaterialExpressionInput(GetInputName(Input)))
+			else 
 			{
-				AddAttributeFromValueOrInterface(Input, *InputNameFound, OperatorNode);
+				AddAttributeFromValueOrInterface(Input, Input->getName().c_str(), OperatorNode);
 			}
 		}
 
@@ -280,24 +280,24 @@ bool FMaterialXSurfaceShaderAbstract::ConnectMatchingNodeOutputToInput(const FCo
 			UInterchangeShaderPortsAPI::ConnectOuputToInputByIndex(Connect.ParentShaderNode, Connect.InputChannelName, OperatorNode->GetUniqueID(), IndexOutput);
 	};
 
-	auto ConnectOutputToInput = [&](const FString* ShaderType, auto* (FMaterialXSurfaceShaderAbstract::* CreateFunctionCallOrShaderNode)(mx::ElementPtr, const FString&, const FString&, const FString&), bool bFindMatchingInput = true)
+	auto ConnectOutputToInput = [&](const FString* ShaderType, auto* (FMaterialXSurfaceShaderAbstract::* CreateFunctionCallOrShaderNode)(mx::ElementPtr, const FString&, const FString&, const FString&))
 	{
 		UInterchangeShaderNode* OperatorNode = nullptr;
 
 		//We don't take the node output here because it would cause the creation of a new node (output is meaningful with ComponentMaskNode/separate where we have to create a new expression
 		OperatorNode = (this->*CreateFunctionCallOrShaderNode)(Connect.UpstreamNode, Connect.UpstreamNode->getName().c_str(), *ShaderType, DefaultOutput);
 
-		ConnectOutputToInputInternal(OperatorNode, bFindMatchingInput);
+		ConnectOutputToInputInternal(OperatorNode);
 	};
 
-	auto ConnectFunctionShaderNodeOutputToInput = [&](uint8 EnumType, uint8 EnumValue, auto* (FMaterialXSurfaceShaderAbstract::* CreateFunctionCallOrShaderNode)(const FString&, uint8, uint8, const FString&), bool bFindMatchingInput = true)
+	auto ConnectFunctionShaderNodeOutputToInput = [&](uint8 EnumType, uint8 EnumValue, auto* (FMaterialXSurfaceShaderAbstract::* CreateFunctionCallOrShaderNode)(mx::ElementPtr, const FString&, uint8, uint8, const FString&))
 	{
 		UInterchangeShaderNode* OperatorNode = nullptr;
 
 		//We don't take the node output here because it would cause the creation of a new node (output is meaningful with ComponentMaskNode/separate where we have to create a new expression
-		OperatorNode = (this->*CreateFunctionCallOrShaderNode)(Connect.UpstreamNode->getName().c_str(), EnumType, EnumValue, DefaultOutput);
+		OperatorNode = (this->*CreateFunctionCallOrShaderNode)(Connect.UpstreamNode, Connect.UpstreamNode->getName().c_str(), EnumType, EnumValue, DefaultOutput);
 
-		ConnectOutputToInputInternal(OperatorNode, bFindMatchingInput);
+		ConnectOutputToInputInternal(OperatorNode);
 	};
 
 	const FString* MaterialFunctionPath = nullptr;
@@ -319,8 +319,8 @@ bool FMaterialXSurfaceShaderAbstract::ConnectMatchingNodeOutputToInput(const FCo
 	else if(uint8 EnumType, EnumValue; Manager.FindMatchingMaterialFunction(Connect.UpstreamNode->getCategory().c_str(), MaterialFunctionPath, EnumType, EnumValue))
 	{
 		MaterialFunctionPath ?
-			ConnectOutputToInput(MaterialFunctionPath, &FMaterialXSurfaceShaderAbstract::CreateFunctionCallShaderNode, false) :
-			ConnectFunctionShaderNodeOutputToInput(EnumType, EnumValue, &FMaterialXSurfaceShaderAbstract::CreateFunctionCallShaderNode, false);
+			ConnectOutputToInput(MaterialFunctionPath, &FMaterialXSurfaceShaderAbstract::CreateFunctionCallShaderNode) :
+			ConnectFunctionShaderNodeOutputToInput(EnumType, EnumValue, &FMaterialXSurfaceShaderAbstract::CreateFunctionCallShaderNode);
 	}
 
 	return bIsConnected;
@@ -328,21 +328,16 @@ bool FMaterialXSurfaceShaderAbstract::ConnectMatchingNodeOutputToInput(const FCo
 
 void FMaterialXSurfaceShaderAbstract::ConnectNodeCategoryOutputToInput(const MaterialX::Edge& Edge, UInterchangeShaderNode* ShaderNode, const FString& ParentInputName, const FString& OutputName)
 {
+	FMaterialXManager& Manager = FMaterialXManager::GetInstance();
 	if(mx::NodePtr UpstreamNode = Edge.getUpstreamElement()->asA<mx::Node>())
 	{
 		// We need to connect the different descending nodes to all the outputs of a node
 		// At least one output connected to the root shader node
 		TArray<UInterchangeShaderNode*> ParentShaderNodeOutputs{ShaderNode};
 		FString InputChannelName = ParentInputName;
-
-		// Add all the inputs from the nodedef in case the node doesn't have them all (we don't have a 1:1 match on the default values otherwise)
-		// For exemple min defaults to (0,0) whereas UE to (0, 1)
-		// We don't want to add every node, since we don't a have perfect match for each input
-		static const std::set<std::string> NodeDefsCategories{"ifgreater", "ifgreatereq", "ifequal", "min", "max", "modulo"};
-		if(NodeDefsCategories.find(UpstreamNode->getCategory()) != NodeDefsCategories.cend())
-		{
-			UpstreamNode->addInputsFromNodeDef();
-		}
+				
+		Manager.AddInputsFromNodeDef(UpstreamNode);
+		Manager.RemoveInputs(UpstreamNode);
 
 		// Replace the input's name by the one used in UE
 		SetMatchingInputsNames(UpstreamNode);
@@ -1426,10 +1421,6 @@ void FMaterialXSurfaceShaderAbstract::ConnectNormalMapInputToOutput(const FConne
 {
 	using namespace UE::Interchange::Materials::Standard::Nodes;
 
-	// FlattenNormalMap Material Function doesn't have normal/tangent inputs, we just remove them to avoid unnecessary connections
-	Connect.UpstreamNode->removeInput("normal");
-	Connect.UpstreamNode->removeInput("tangent");
-
 	// Only create a FunctionCall if there's a scale and the value is not 1, otherwise just like dot
 	mx::InputPtr Input = Connect.UpstreamNode->getInput("scale");
 	bool bIsNotEqualOne = true;
@@ -1526,7 +1517,7 @@ UInterchangeFunctionCallShaderNode* FMaterialXSurfaceShaderAbstract::CreateFunct
 {
 	UInterchangeFunctionCallShaderNode* Node;
 
-	const FString NodeUID = UInterchangeShaderNode::MakeNodeUid(NodeName  + TEXT('_') + OutputName, FStringView{});
+	const FString NodeUID = UInterchangeShaderNode::MakeNodeUid(GetUniqueName(Element) + TEXT("_") + NodeName  + TEXT('_') + OutputName, FStringView{});
 
 	if(Node = const_cast<UInterchangeFunctionCallShaderNode*>(Cast<UInterchangeFunctionCallShaderNode>(NodeContainer.GetNode(NodeUID))); !Node)
 	{
@@ -1541,11 +1532,11 @@ UInterchangeFunctionCallShaderNode* FMaterialXSurfaceShaderAbstract::CreateFunct
 	return Node;
 }
 
-UInterchangeFunctionCallShaderNode* FMaterialXSurfaceShaderAbstract::CreateFunctionCallShaderNode(const FString& NodeName, uint8 EnumType, uint8 EnumValue, const FString& OutputName)
+UInterchangeFunctionCallShaderNode* FMaterialXSurfaceShaderAbstract::CreateFunctionCallShaderNode(MaterialX::ElementPtr Element, const FString& NodeName, uint8 EnumType, uint8 EnumValue, const FString& OutputName)
 {
 	UInterchangeFunctionCallShaderNode* Node;
 
-	const FString NodeUID = UInterchangeShaderNode::MakeNodeUid(NodeName + TEXT('_') + OutputName, FStringView{});
+	const FString NodeUID = UInterchangeShaderNode::MakeNodeUid(GetUniqueName(Element) + TEXT("_") + NodeName + TEXT('_') + OutputName, FStringView{});
 
 	if(Node = const_cast<UInterchangeFunctionCallShaderNode*>(Cast<UInterchangeFunctionCallShaderNode>(NodeContainer.GetNode(NodeUID))); !Node)
 	{
@@ -1659,14 +1650,18 @@ FString FMaterialXSurfaceShaderAbstract::GetAttributeParentName(MaterialX::NodeP
 
 namespace
 {
-	void GetUniqueName(mx::ElementPtr Element, TStringBuilder<256>& Buffer)
+	void GetUniqueName(mx::ElementPtr Element, TStringBuilder<256>& Buffer, const FString& SurfaceMaterialName)
 	{
 		//Write the hierarchy name as A_B_C
 		if (Element)
 		{
 			if (mx::ElementPtr Parent = Element->getParent())
 			{
-				GetUniqueName(Parent, Buffer);
+				GetUniqueName(Parent, Buffer, SurfaceMaterialName);
+			}
+			else
+			{
+				Buffer += SurfaceMaterialName + TEXT("_");
 			}
 
 			Buffer += Element->getName().c_str();
@@ -1685,7 +1680,7 @@ FString FMaterialXSurfaceShaderAbstract::GetUniqueName(MaterialX::ElementPtr Ele
 	else
 	{
 		TStringBuilder<256> Buffer;
-		::GetUniqueName(Element, Buffer);
+		::GetUniqueName(Element, Buffer, SurfaceMaterialName);
 		return FString{ Buffer.ToString() }.TrimChar(TEXT('_')); // remove any trailing '_'
 	}
 }
@@ -1792,7 +1787,7 @@ UInterchangeShaderNode* FMaterialXSurfaceShaderAbstract::Translate(EInterchangeM
 {
 	constexpr bool bInputInTangentSpace = true;
 
-	UInterchangeFunctionCallShaderNode* FunctionCallShaderNode = CreateFunctionCallShaderNode(SurfaceMaterialName + TEXT("_") + SurfaceShaderNode->getName().c_str(), UE::Interchange::MaterialX::IndexSurfaceShaders, uint8(ShaderType));
+	UInterchangeFunctionCallShaderNode* FunctionCallShaderNode = CreateFunctionCallShaderNode(SurfaceShaderNode, SurfaceMaterialName + TEXT("_") + SurfaceShaderNode->getName().c_str(), UE::Interchange::MaterialX::IndexSurfaceShaders, uint8(ShaderType));
 
 	for(mx::InputPtr Input : SurfaceShaderNode->getInputs())
 	{
diff --git a/Engine/Plugins/Interchange/Runtime/Source/Import/Private/Usd/InterchangeUsdTranslator.cpp b/Engine/Plugins/Interchange/Runtime/Source/Import/Private/Usd/InterchangeUsdTranslator.cpp
--- a/Engine/Plugins/Interchange/Runtime/Source/Import/Private/Usd/InterchangeUsdTranslator.cpp
+++ b/Engine/Plugins/Interchange/Runtime/Source/Import/Private/Usd/InterchangeUsdTranslator.cpp
@@ -850,6 +850,30 @@ namespace UE::InterchangeUsdTranslator::Private
 			);
 		}
 
+		// Let's iterate over the Shader Nodes, the TextureSample ones more specifically,
+		// to see if have to retrieve any attributes related to the conversion of geompropvalues
+		// we'll store the UID of the shader node in order to retrieve it during the baking phase (in the post factory import)
+		TArray<FString> ShaderNodesTextureSampleUIDs;
+		NodeContainer.IterateNodesOfType<UInterchangeShaderNode>(
+			[&ShaderNodesTextureSampleUIDs](const FString&, UInterchangeShaderNode* ShaderNode)
+			{
+				using namespace UE::Interchange::Materials::Standard::Nodes;
+				
+				constexpr const TCHAR* GeomPropHandledAttribute = TEXT("GeomPropHandled");
+				if (FString ShaderType; ShaderNode->GetCustomShaderType(ShaderType) && ShaderType == TextureSample::Name.ToString())
+				{
+					if (bool bIsGeomProp; ShaderNode->GetBooleanAttribute(::MaterialX::Attributes::GeomPropImage, bIsGeomProp) && bIsGeomProp)
+					{
+						// We only want the geomprops that were not already assigned to the mesh node
+						if(bool bGeomPropHandled; !ShaderNode->GetBooleanAttribute(GeomPropHandledAttribute, bGeomPropHandled) || !bGeomPropHandled)
+						{
+							ShaderNode->AddBooleanAttribute(GeomPropHandledAttribute, true);
+							ShaderNodesTextureSampleUIDs.Emplace(ShaderNode->GetUniqueID());
+						}
+					}
+				}
+			});
+
 		// Now we need to check if we have to set the slot of the mesh nodes here
 		if (TArray<FMaterialSlotMesh>* SlotMeshes = PrimPathToSlotMeshNodes.Find(MaterialUid))
 		{
@@ -864,6 +888,8 @@ namespace UE::InterchangeUsdTranslator::Private
 						for (int32 Index = 0; Index < GeomPropValues.Num(); ++Index)
 						{
 							MaterialSlotMesh.MeshNode->AddStringAttribute(USD::Primvar::Name + FString::FromInt(Index), GeomPropValues[Index].Name);
+							MaterialSlotMesh.MeshNode->AddBooleanAttribute(USD::Primvar::TangentSpace + FString::FromInt(Index), GeomPropValues[Index].bTangentSpace);
+							MaterialSlotMesh.MeshNode->AddStringAttribute(USD::Primvar::ShaderNodeTextureSample + FString::FromInt(Index), ShaderNodesTextureSampleUIDs[Index]);
 						}
 					}
 
diff --git a/Engine/Plugins/Interchange/Runtime/Source/Import/Public/MaterialX/MaterialXUtils/MaterialXManager.h b/Engine/Plugins/Interchange/Runtime/Source/Import/Public/MaterialX/MaterialXUtils/MaterialXManager.h
--- a/Engine/Plugins/Interchange/Runtime/Source/Import/Public/MaterialX/MaterialXUtils/MaterialXManager.h
+++ b/Engine/Plugins/Interchange/Runtime/Source/Import/Public/MaterialX/MaterialXUtils/MaterialXManager.h
@@ -9,6 +9,7 @@
 #include "MaterialXCore/Document.h"
 #include "Misc/TVariant.h"
 #include "MaterialX/InterchangeMaterialXDefinitions.h"
+#include <unordered_set>
 
 class FMaterialXBase;
 class UInterchangeBaseNodeContainer;
@@ -41,6 +42,18 @@ public:
 	
 	bool IsSubstrateEnabled() const;
 
+	/** Return true if the node category should be filtered out during the flattening of the subgraphs*/
+	bool FilterNodeGraph(MaterialX::NodePtr Node) const;
+
+	/** Remove the inputs of a node with no match in UE Material Expressions*/
+	void RemoveInputs(MaterialX::NodePtr Node) const;
+
+	/**
+	 * Add all the inputs from the nodedef in case the node doesn't have them all (we don't have a 1:1 match on the default values otherwise)
+	 * For exemple 'min' defaults to (0, 0) whereas UE to (0, 1)
+	 * We don't want to add every node, since we don't a have perfect match for each input
+	 */
+	void AddInputsFromNodeDef(MaterialX::NodePtr Node) const;
 
 	FMaterialXManager(const FMaterialXManager&) = delete;
 	FMaterialXManager& operator=(const FMaterialXManager&) = delete;
@@ -102,6 +115,15 @@ private:
 	/** Given a MaterialX node category, return the UE material function, used for BSDF nodes*/
 	TMap<FString, FMaterialXMaterialFunction> MatchingMaterialFunctions;
 
+	/** Categories of nodes to skip during the phase of flattening the subgraphs. Basically if a node has a nodegraph it will be processed as is instead of its nodegraph*/
+	std::unordered_set<std::string> CategoriesToSkip;
+
+	/** Categories of nodes to add the default inputs from the nodedefs if they are not present on the node */
+	std::unordered_set<std::string> NodeDefsCategories;
+
+	/** Inputs of the nodes category to remove because there's no equivalent in UE Material Expression, to avoid creating a floating Scalar/Vector parameter*/
+	std::unordered_map<std::string, std::vector<std::string>> NodeInputsToRemove;
+
 	bool bIsSubstrateEnabled;
 };
 #endif
diff --git a/Engine/Plugins/Interchange/Runtime/Source/Import/Public/MaterialX/MaterialXUtils/MaterialXSurfaceShaderAbstract.h b/Engine/Plugins/Interchange/Runtime/Source/Import/Public/MaterialX/MaterialXUtils/MaterialXSurfaceShaderAbstract.h
--- a/Engine/Plugins/Interchange/Runtime/Source/Import/Public/MaterialX/MaterialXUtils/MaterialXSurfaceShaderAbstract.h
+++ b/Engine/Plugins/Interchange/Runtime/Source/Import/Public/MaterialX/MaterialXUtils/MaterialXSurfaceShaderAbstract.h
@@ -362,7 +362,7 @@ protected:
 	 * @return The shader node that was created.
 	 */
 	UInterchangeFunctionCallShaderNode* CreateFunctionCallShaderNode(MaterialX::ElementPtr Element, const FString& NodeName, const FString& FunctionPath, const FString& OutputName = TEXT("out"));
-	UInterchangeFunctionCallShaderNode* CreateFunctionCallShaderNode(const FString& NodeName, uint8 EnumType, uint8 EnumValue, const FString& OutputName = TEXT("out"));
+	UInterchangeFunctionCallShaderNode* CreateFunctionCallShaderNode(MaterialX::ElementPtr Element, const FString& NodeName, uint8 EnumType, uint8 EnumValue, const FString& OutputName = TEXT("out"));
 
 	/**
 	 * Helper function to create an InterchangeTextureNode.
@@ -382,11 +382,29 @@ protected:
 		{
 			if(MaterialX::InputPtr InputFile = Node->getInput("file"); InputFile && (InputFile->hasValue() || InputFile->hasInterfaceName()))
 			{
-				FString Filepath{ InputFile->hasValue() ? InputFile->getValueString().c_str() : InputFile->getInterfaceInput()->getValueString().c_str()};
-				const FString FilePrefix = GetFilePrefix(InputFile);
-				Filepath = FPaths::Combine(FilePrefix, Filepath);
+				FString Filepath;
+				// a <geompropvalue> converted as an <image> as no filepath by default, since the texture will be created later, we just put a placeholder, in order for the factory to not complain
+				// the texture will be overridden later in the ExecutePostFactoryImport
+				const bool bGeomPropImage = Node->getTypedAttribute<bool>(MaterialX::Attributes::GeomPropImage);
+				if (bGeomPropImage)
+				{
+					Filepath = FPaths::Combine(FPaths::EnginePluginsDir(), TEXT("Interchange"), TEXT("Editor"), TEXT("Content"), TEXT("Resources"), TEXT("Interchange_PixelPowerOfTwo.png"));
+					Filepath = FPaths::ConvertRelativePathToFull(Filepath);
+				}
+				else
+				{
+					Filepath = InputFile->hasValue() ? InputFile->getValueString().c_str() : InputFile->getInterfaceInput()->getValueString().c_str();
+					const FString FilePrefix = GetFilePrefix(InputFile);
+					Filepath = FPaths::Combine(FilePrefix, Filepath);
+				}
+
 				const FString Filename = FPaths::GetCleanFilename(Filepath);
-				const FString TextureNodeUID = TEXT("\\Texture\\") + Filename;
+				FString TextureNodeUID = TEXT("\\Texture\\") + Filename;
+				// We need to duplicate the texture for a geomprop, at the end the texture will be the baking texture
+				if (bGeomPropImage)
+				{
+					TextureNodeUID += TEXT("_") + SurfaceMaterialName + Node->getName().c_str();
+				}
 
 				//Only add the TextureNode once
 				TextureNode = const_cast<UInterchangeTexture2DNode*>(Cast<UInterchangeTexture2DNode>(NodeContainer.GetNode(TextureNodeUID)));
@@ -394,6 +412,10 @@ protected:
 				{
 					TextureNode = NewObject<TextureTypeNode>(&NodeContainer);
 					TextureNode->InitializeNode(TextureNodeUID, Filename, EInterchangeNodeContainerType::TranslatedAsset);
+					if (bGeomPropImage)
+					{
+						TextureNode->SetDisplayLabel(Node->getName().c_str());
+					}
 					NodeContainer.AddNode(TextureNode);
 
 					if(FPaths::IsRelative(Filepath))
diff --git a/Engine/Plugins/Interchange/Runtime/Source/Nodes/Public/InterchangeMaterialDefinitions.h b/Engine/Plugins/Interchange/Runtime/Source/Nodes/Public/InterchangeMaterialDefinitions.h
--- a/Engine/Plugins/Interchange/Runtime/Source/Nodes/Public/InterchangeMaterialDefinitions.h
+++ b/Engine/Plugins/Interchange/Runtime/Source/Nodes/Public/InterchangeMaterialDefinitions.h
@@ -757,6 +757,14 @@ namespace Interchange
 				const FName Occlusion = TEXT("Occlusion");
 			}			
 		}
+
+		namespace Factory
+		{
+			namespace Expression
+			{
+				const FName Path = TEXT("MaterialExpressionPath"); // Attribute used by the factory to store the path of a UMaterialExpression
+			}
+		}
 	}
 }
 }
diff --git a/Engine/Plugins/Interchange/Runtime/Source/Pipelines/InterchangePipelines.Build.cs b/Engine/Plugins/Interchange/Runtime/Source/Pipelines/InterchangePipelines.Build.cs
--- a/Engine/Plugins/Interchange/Runtime/Source/Pipelines/InterchangePipelines.Build.cs
+++ b/Engine/Plugins/Interchange/Runtime/Source/Pipelines/InterchangePipelines.Build.cs
@@ -26,13 +26,20 @@ namespace UnrealBuildTool.Rules
 					"ApplicationCore",
 					"CinematicCamera",
 					"DeveloperSettings",
+					"DynamicMesh",
 					"GeometryCache",
+					"GeometryCore",
 					"GLTFCore",
 					"InputCore",
 					"InterchangeImport",
 					"LevelSequence",
+					"MeshConversion",
+					"MeshDescription",
+					"MeshUtilitiesCommon",
+					"ModelingComponents",
 					"Slate",
 					"SlateCore",
+					"StaticMeshDescription",
 					"UnrealUSDWrapper",
 					"USDClasses",
 				}
diff --git a/Engine/Plugins/Interchange/Runtime/Source/Pipelines/Private/InterchangeUSDPipeline.cpp b/Engine/Plugins/Interchange/Runtime/Source/Pipelines/Private/InterchangeUSDPipeline.cpp
--- a/Engine/Plugins/Interchange/Runtime/Source/Pipelines/Private/InterchangeUSDPipeline.cpp
+++ b/Engine/Plugins/Interchange/Runtime/Source/Pipelines/Private/InterchangeUSDPipeline.cpp
@@ -3,6 +3,7 @@
 #include "InterchangeUSDPipeline.h"
 
 #include "InterchangeManager.h"
+#include "InterchangeMaterialDefinitions.h"
 #include "InterchangeMaterialFactoryNode.h"
 #include "InterchangeMeshNode.h"
 #include "InterchangeMeshFactoryNode.h"
@@ -10,14 +11,29 @@
 #include "InterchangeShaderGraphNode.h"
 #include "InterchangeSkeletalMeshLodDataNode.h"
 #include "InterchangeStaticMeshLodDataNode.h"
+#include "InterchangePipelineLog.h"
+#include "Nodes/InterchangeBaseNodeContainer.h"
+
+#include "AssetUtils/Texture2DBuilder.h"
+#include "DynamicMesh/DynamicMesh3.h"
+#include "Engine/StaticMesh.h"
+#include "Materials/Material.h"
+#include "Materials/MaterialExpressionTextureBase.h"
+#include "Materials/MaterialExpressionTextureSample.h"
+#include "Materials/MaterialInstanceConstant.h"
+#include "Materials/MaterialInstanceDynamic.h"
+#include "MeshDescriptionToDynamicMesh.h"
+#include "Sampling/MeshMapBaker.h"
+#include "Sampling/MeshPropertyMapEvaluator.h"
+#include "StaticMeshAttributes.h"
+#include "StaticMeshOperations.h"
 
 #include "UnrealUSDWrapper.h"
-#include "Usd/InterchangeUsdDefinitions.h"
 #include "USDMaterialUtils.h"
+#include "Usd/InterchangeUsdDefinitions.h"
 #include "UsdWrappers/SdfPath.h"
 
-#include "Materials/MaterialInstanceConstant.h"
-#include "Materials/MaterialInstanceDynamic.h"
+#include UE_INLINE_GENERATED_CPP_BY_NAME(InterchangeUSDPipeline)
 
 namespace UE::InterchangeUsdPipeline::Private
 {
@@ -429,6 +445,7 @@ namespace UE::InterchangeUsdPipeline::Private
 		}
 	}
 
+#if WITH_EDITORONLY_DATA
 	/**
 	 * Add the primvar attributes from the mesh nodes to factory nodes
 	 */
@@ -457,11 +474,259 @@ namespace UE::InterchangeUsdPipeline::Private
 							{
 								MeshFactoryNode->AddBooleanAttribute(PrimvarAttribute, bTangentSpace);
 							}
+
+							FString ShaderNodeTextureSampleAttribute = USD::Primvar::ShaderNodeTextureSample + FString::FromInt(Index);
+							if (FString ShaderNodeTextureSampleUID;  MeshNode->GetStringAttribute(ShaderNodeTextureSampleAttribute, ShaderNodeTextureSampleUID))
+							{
+								// At this stage the USD pipeline should be the latest one in the stack, we should retrieve the MaterialExpression Factory Node associated to this ShaderNode
+								FString TextureSampleFactoryNodeUID = UInterchangeFactoryBaseNode::BuildFactoryNodeUid(ShaderNodeTextureSampleUID);
+								
+								if (Cast<UInterchangeMaterialExpressionFactoryNode>(NodeContainer->GetFactoryNode(TextureSampleFactoryNodeUID)))
+								{
+									MeshFactoryNode->AddStringAttribute(ShaderNodeTextureSampleAttribute, TextureSampleFactoryNodeUID);
+								}
+							}
 						}
 					}
 				}
 			});
 	}
+
+	class FUsdPrimvarBaker
+	{
+	public:
+
+		FUsdPrimvarBaker(UStaticMesh* StaticMeshAsset, const UInterchangeBaseNodeContainer* NodeContainer, const UInterchangeFactoryBaseNode* FactoryNode)
+			: NodeContainer{ NodeContainer }
+			, FactoryNode{ FactoryNode }
+			, StaticMeshAsset{ StaticMeshAsset }
+		{}
+
+		void Bake()
+		{
+			using namespace UE::Interchange;
+			using namespace UE::Interchange::Materials;
+			using namespace UE::Geometry;
+
+			if (int32 NumberOfPrimvar; FactoryNode->GetInt32Attribute(USD::Primvar::Number, NumberOfPrimvar) && NumberOfPrimvar)
+			{
+				FMeshDescription MeshDescriptionPrimvar = CreateMeshDescription();
+
+				FStaticMeshAttributes MeshAttributes(MeshDescriptionPrimvar);
+				TVertexInstanceAttributesRef<FVector4f> VertexColors = MeshAttributes.GetVertexInstanceColors();
+				TVertexInstanceAttributesRef<FVector3f> Normals = MeshAttributes.GetVertexInstanceNormals();
+
+				for (int32 Index = 0; Index < NumberOfPrimvar; ++Index)
+				{
+					UMaterialExpressionTextureSample* TextureSampleExpression = GetMaterialExpressionTextureSample(Index);
+
+					if (!TextureSampleExpression)
+					{
+						continue;
+					}
+
+					FString PrimvarName;
+					bool bTangentSpace = false;
+					FactoryNode->GetStringAttribute(USD::Primvar::Name + FString::FromInt(Index), PrimvarName);
+					FactoryNode->GetBooleanAttribute(USD::Primvar::TangentSpace + FString::FromInt(Index), bTangentSpace);
+
+					TVertexInstanceAttributesRef<FVector4f> Primvar = MeshDescriptionPrimvar.VertexInstanceAttributes().GetAttributesRef<FVector4f>(*PrimvarName);
+					
+					if (!(Primvar.IsValid()
+						&& Primvar.GetNumElements() == VertexColors.GetNumElements()
+						&& Primvar.GetNumElements() == Normals.GetNumElements()))
+					{
+						continue;
+					}
+
+					// We put the primvar in the Color/Normal channel, that way during the conversion to a DynamicMesh it will also handle the non-manifold case
+					for (const int32 VertexInstanceID : MeshDescriptionPrimvar.VertexInstances().GetElementIDs())
+					{
+						bTangentSpace
+							? Normals.Set(VertexInstanceID, Primvar.Get(VertexInstanceID))
+							: VertexColors.Set(VertexInstanceID, Primvar.Get(VertexInstanceID));
+					}
+
+					FDynamicMesh3 DynamicMesh;
+					FMeshDescriptionToDynamicMesh Convert;
+					Convert.Convert(&MeshDescriptionPrimvar, DynamicMesh);
+
+					const FDynamicMeshAABBTree3 DetailSpatial(&DynamicMesh);
+					FMeshBakerDynamicMeshSampler DetailSampler(&DynamicMesh, &DetailSpatial);
+					TSharedPtr<FMeshPropertyMapEvaluator> PropertyEval = MakeShared<FMeshPropertyMapEvaluator>();
+					PropertyEval->Property =
+						bTangentSpace
+						? EMeshPropertyMapType::Normal
+						: EMeshPropertyMapType::VertexColor;
+
+					constexpr int32 TextureSize = 1024;
+					const FImageDimensions ImageDimensions(TextureSize, TextureSize);
+					FMeshMapBaker Baker;
+					Baker.SetCorrespondenceStrategy(FMeshBaseBaker::ECorrespondenceStrategy::Identity);
+					Baker.SetTargetMesh(&DynamicMesh);
+					Baker.SetDetailSampler(&DetailSampler);
+					Baker.AddEvaluator(PropertyEval);
+					Baker.SetTargetMeshUVLayer(UVChannel);
+					Baker.SetDimensions(ImageDimensions);
+					Baker.SetProjectionDistance(3.0f);
+					Baker.SetSamplesPerPixel(1);
+					Baker.SetFilter(FMeshMapBaker::EBakeFilterType::BSpline);
+					Baker.SetGutterEnabled(true);
+					Baker.SetGutterSize(4);
+					Baker.SetTileSize(TextureSize);
+					Baker.Bake();
+										
+					FTexture2DBuilder TextureBuilder;
+					const FTexture2DBuilder::ETextureType TextureType = GetTextureType(TextureSampleExpression);
+
+					TextureBuilder.InitializeAndReplaceExistingTexture(Cast<UTexture2D>(TextureSampleExpression->Texture.Get()), TextureType, ImageDimensions);
+					TextureBuilder.Copy(*Baker.GetBakeResults(0)[0]);
+					TextureBuilder.Commit();
+
+					TextureSampleExpression->ConstCoordinate = UVChannel;
+
+					MaterialExpressionTextureSamples.Add(TextureSampleExpression);
+				}
+
+				UpdateMaterials();
+			}
+		}
+
+	private:
+
+		UMaterialExpressionTextureSample* GetMaterialExpressionTextureSample(int32 IndexPrimvar) const
+		{
+			using namespace UE::Interchange;
+			using namespace UE::Interchange::Materials;
+
+			UMaterialExpressionTextureSample* MaterialExpressionTextureSample = nullptr;
+
+			FString ShaderNodeTextureSampleUID;
+			FactoryNode->GetStringAttribute(USD::Primvar::ShaderNodeTextureSample + FString::FromInt(IndexPrimvar), ShaderNodeTextureSampleUID);
+			if (UInterchangeMaterialExpressionFactoryNode* TextureSampleExpressionFactoryNode = Cast<UInterchangeMaterialExpressionFactoryNode>(NodeContainer->GetFactoryNode(ShaderNodeTextureSampleUID)))
+			{
+				if (FString ExpressionPath; TextureSampleExpressionFactoryNode->GetStringAttribute(Factory::Expression::Path.ToString(), ExpressionPath))
+				{
+					MaterialExpressionTextureSample = Cast<UMaterialExpressionTextureSample>(FSoftObjectPath{ ExpressionPath }.TryLoad());
+				}
+			}
+
+			return MaterialExpressionTextureSample;
+		}
+
+		/** Create a mesh description based on the one from the StaticMesh, the StaticMesh will now hold a new UV set */
+		FMeshDescription CreateMeshDescription()
+		{
+			FMeshDescription* MeshDescription = StaticMeshAsset->GetMeshDescription(0);
+			FStaticMeshAttributes StaticMeshAttributes(*MeshDescription);
+			TVertexInstanceAttributesRef<FVector2f> VertexInstanceUVs = StaticMeshAttributes.GetVertexInstanceUVs();
+
+			if (AreUVsOverlapping(VertexInstanceUVs))
+			{
+				FStaticMeshOperations::FGenerateUVOptions Options;
+				Options.UVMethod = FStaticMeshOperations::EGenerateUVMethod::Default;
+				TArray<FVector2D> UVs;
+				FStaticMeshOperations::GenerateUV(*MeshDescription, Options, UVs);
+				FStaticMeshOperations::AddUVChannel(*MeshDescription);
+				UVChannel = VertexInstanceUVs.GetNumChannels() - 1;
+
+				for (const int32 VertexInstanceID : MeshDescription->VertexInstances().GetElementIDs())
+				{
+					VertexInstanceUVs.Set(VertexInstanceID, UVChannel, FVector2f{ UVs[VertexInstanceID] });
+				}
+			}
+
+			FMeshDescription MeshDescriptionPrimvar(*MeshDescription);
+			return MeshDescriptionPrimvar;
+		}
+
+		void UpdateMaterials() const
+		{
+			TSet<UMaterial*> Materials;
+			for (UMaterialExpressionTextureSample* TextureSampleExpression : MaterialExpressionTextureSamples)
+			{
+				Materials.FindOrAdd(TextureSampleExpression->Material);
+				TextureSampleExpression->PostEditChange();
+				TextureSampleExpression->Modify();
+			}
+
+			for (UMaterial* Material : Materials)
+			{
+				Material->PostEditChange();
+				Material->ForceRecompileForRendering();
+				Material->MarkPackageDirty();
+			}
+		}
+
+		UE::Geometry::FTexture2DBuilder::ETextureType GetTextureType(UMaterialExpressionTextureSample* TextureSampleExpression)
+		{
+			using namespace UE::Geometry;			
+			switch (TextureSampleExpression->SamplerType)
+			{
+			case EMaterialSamplerType::SAMPLERTYPE_Color:
+				return FTexture2DBuilder::ETextureType::Color;
+
+			case EMaterialSamplerType::SAMPLERTYPE_Normal:
+				return FTexture2DBuilder::ETextureType::NormalMap;
+
+			default:
+				return FTexture2DBuilder::ETextureType::ColorLinear;
+			}
+		}
+
+		struct FUVKeyFunc : public BaseKeyFuncs<FVector2f, FVector2f>
+		{
+			static constexpr float Threshold = 0.00001;
+
+			static FORCEINLINE KeyInitType GetSetKey(ElementInitType Element)
+			{
+				return Element;
+			}
+
+			static FORCEINLINE uint32 GetKeyHash(ElementInitType& Key)
+			{
+				int32 RoundedU = FMath::RoundToInt(Key.X / Threshold);
+				int32 RoundedV = FMath::RoundToInt(Key.Y / Threshold);
+				return HashCombine(GetTypeHash(RoundedU), GetTypeHash(RoundedV));
+			}
+
+			static FORCEINLINE bool Matches(ElementInitType A, ElementInitType B)
+			{
+				return FMath::IsNearlyEqual(A.X, B.X, Threshold) &&
+					   FMath::IsNearlyEqual(A.Y, B.Y, Threshold);
+			}
+		};
+
+		bool AreUVsOverlapping(const TVertexInstanceAttributesRef<FVector2f>& UVs)
+		{
+			FMeshDescription* MeshDecription = StaticMeshAsset->GetMeshDescription(0);
+			TSet<FVector2f, FUVKeyFunc> UVSet;
+
+			for (const FVertexInstanceID VertexInstanceID : MeshDecription->VertexInstances().GetElementIDs())
+			{
+				FVector2f UV = UVs.Get(VertexInstanceID);
+
+				if (UVSet.Contains(UV))
+				{
+					UE_LOG(LogInterchangePipeline, Warning, TEXT("UVs are Overlapping"));
+					return true;
+				}
+
+				UVSet.Add(UV);
+			}
+
+			return false;
+		}
+
+	private:
+
+		const UInterchangeBaseNodeContainer* NodeContainer;
+		const UInterchangeFactoryBaseNode* FactoryNode;
+		UStaticMesh* StaticMeshAsset;
+		TArray<UMaterialExpressionTextureSample*> MaterialExpressionTextureSamples;
+		int32 UVChannel = 0;
+	};
+#endif // WITH_EDITORONLY_DATA
 }	 // namespace UE::InterchangeUsdPipeline::Private
 
 UInterchangeUsdPipeline::UInterchangeUsdPipeline()
@@ -470,6 +735,29 @@ UInterchangeUsdPipeline::UInterchangeUsdPipeline()
 {
 }
 
+void UInterchangeUsdPipeline::ExecutePostFactoryPipeline(const UInterchangeBaseNodeContainer* BaseNodeContainer, const FString& NodeKey, UObject* CreatedAsset, bool bIsAReimport)
+{
+	if (!BaseNodeContainer || !CreatedAsset)
+	{
+		return;
+	}
+
+	const UInterchangeFactoryBaseNode* FactoryNode = BaseNodeContainer->GetFactoryNode(NodeKey);
+	if (!FactoryNode)
+	{
+		return;
+	}
+
+#if WITH_EDITORONLY_DATA
+	if (UStaticMesh* StaticMeshAsset = Cast<UStaticMesh>(CreatedAsset))
+	{
+		using namespace UE::InterchangeUsdPipeline::Private;		
+		FUsdPrimvarBaker Baker(StaticMeshAsset, BaseNodeContainer, FactoryNode);
+		Baker.Bake();
+	}
+#endif // WITH_EDITORONLY_DATA
+}
+
 void UInterchangeUsdPipeline::ExecutePipeline(
 	UInterchangeBaseNodeContainer* NodeContainer,
 	const TArray<UInterchangeSourceData*>& InSourceDatas,
@@ -494,7 +782,9 @@ void UInterchangeUsdPipeline::ExecutePipeline(
 		RemovePseudoRootFactoryNode(*NodeContainer);
 	}
 
+#if WITH_EDITORONLY_DATA
 	AddPrimvarAttributesToFactoryNodes(NodeContainer);
+#endif // WITH_EDITORONLY_DATA
 }
 
 bool UInterchangeUsdPipeline::CanExecuteOnAnyThread(EInterchangePipelineTask PipelineTask)
diff --git a/Engine/Plugins/Interchange/Runtime/Source/Pipelines/Public/InterchangeUSDPipeline.h b/Engine/Plugins/Interchange/Runtime/Source/Pipelines/Public/InterchangeUSDPipeline.h
--- a/Engine/Plugins/Interchange/Runtime/Source/Pipelines/Public/InterchangeUSDPipeline.h
+++ b/Engine/Plugins/Interchange/Runtime/Source/Pipelines/Public/InterchangeUSDPipeline.h
@@ -2,11 +2,8 @@
 
 #pragma once
 
-#include "CoreMinimal.h"
-
 #include "InterchangePipelineBase.h"
 #include "InterchangeSourceData.h"
-#include "Nodes/InterchangeBaseNodeContainer.h"
 
 #include "InterchangeUSDPipeline.generated.h"
 
@@ -39,6 +36,10 @@ public:
 	UPROPERTY(EditAnywhere, BlueprintReadWrite, Category = "USD")
 	bool bImportPseudoRoot;
 
+public:
+
+	virtual void ExecutePostFactoryPipeline(const UInterchangeBaseNodeContainer* BaseNodeContainer, const FString& NodeKey, UObject* CreatedAsset, bool bIsAReimport) override;
+
 protected:
 	virtual void ExecutePipeline(
 		UInterchangeBaseNodeContainer* BaseNodeContainer,
diff --git a/Engine/Plugins/Runtime/USDCore/Source/USDUtilities/Private/USDGeomMeshConversion.cpp b/Engine/Plugins/Runtime/USDCore/Source/USDUtilities/Private/USDGeomMeshConversion.cpp
--- a/Engine/Plugins/Runtime/USDCore/Source/USDUtilities/Private/USDGeomMeshConversion.cpp
+++ b/Engine/Plugins/Runtime/USDCore/Source/USDUtilities/Private/USDGeomMeshConversion.cpp
@@ -1151,6 +1151,16 @@ namespace UE::UsdGeomMeshConversion::Private
 		pxr::VtArray<float> DisplayOpacities;
 		TArray<pxr::VtArray<pxr::GfVec2f>> UVSets;
 
+		// Custom attributes, that may come from primvars (used essentially for baking textures in Interchange)
+		struct FCustomPrimvar
+		{
+			pxr::VtArray<pxr::GfVec4f> Data; // Let's put any primvar type in a vec4, we'll convert later in the MeshDescription
+			pxr::VtArray<int> Indices;
+			pxr::TfToken Interpolation;
+		};
+
+		TMap<FString, FCustomPrimvar> CustomPrimvars;
+
 		pxr::TfToken PointInterpolation;
 		pxr::TfToken NormalInterpolation;
 		pxr::TfToken VelocityInterpolation;
@@ -1269,6 +1279,86 @@ namespace UE::UsdGeomMeshConversion::Private
 		}
 	}
 
+	void CollectPrimvarData(
+		const pxr::UsdPrim& Prim,
+		const UsdToUnreal::FUsdMeshConversionOptions& Options,
+		FUsdMeshData& InOutMeshData)
+	{
+		const pxr::TfToken DefaultInterpolation = pxr::UsdGeomTokens->varying;
+
+		pxr::UsdGeomPrimvarsAPI PrimvarAPI{ Prim };
+		std::vector<pxr::UsdGeomPrimvar> Primvars = PrimvarAPI.GetPrimvarsWithAuthoredValues();
+		for (const pxr::UsdGeomPrimvar& Primvar : Primvars)
+		{
+			std::string PrimvarName = Primvar.GetPrimvarName().GetString();
+			if (PrimvarName == "displayColor" || PrimvarName == "displayOpacity" || PrimvarName == "st")
+			{
+				continue;
+			}
+
+			std::string PrimvarType = Primvar.GetTypeName().GetAsToken().GetString();
+
+			FUsdMeshData::FCustomPrimvar CustomPrimvar;
+			auto CollectPrimvarData = [&PrimvarType, &Options, &CustomPrimvar]<class T>(const pxr::VtArray<T> &Primvar)
+			{
+				size_t SizePrimvar = Primvar.size();
+				CustomPrimvar.Data.reserve(Primvar.size());
+				for (const T& Value : Primvar)
+				{
+					if constexpr (std::is_arithmetic_v<std::remove_cvref_t<T>>)
+					{
+						CustomPrimvar.Data.emplace_back(Value);
+					}
+					else
+					{
+						pxr::GfVec4f Vec(1.f);
+						for (size_t IndexDimension = 0; IndexDimension < T::dimension; ++IndexDimension)
+						{
+							Vec.data()[IndexDimension] = Value[IndexDimension];
+						}
+						CustomPrimvar.Data.emplace_back(Vec);
+					}
+				}
+				return true;
+			};
+
+			bool bDataCollected = false;
+			if (pxr::VtArray<float> PrimvarFloat; PrimvarType == pxr::SdfValueTypeNames->FloatArray && Primvar.Get(&(PrimvarFloat), Options.TimeCode))
+			{
+				bDataCollected = CollectPrimvarData(PrimvarFloat);
+			}
+			else if (pxr::VtArray<int32> PrimvarInt; PrimvarType == pxr::SdfValueTypeNames->IntArray && Primvar.Get(&(PrimvarInt), Options.TimeCode))
+			{
+				bDataCollected = CollectPrimvarData(PrimvarInt);
+			}
+			else if (pxr::VtArray<bool> PrimvarBool; PrimvarType == pxr::SdfValueTypeNames->BoolArray && Primvar.Get(&(PrimvarBool), Options.TimeCode))
+			{
+				bDataCollected = CollectPrimvarData(PrimvarBool);
+			}
+			else if (pxr::VtArray<pxr::GfVec2f> PrimvarFloat2; PrimvarType == pxr::SdfValueTypeNames->Float2Array && Primvar.Get(&(PrimvarFloat2), Options.TimeCode))
+			{
+				bDataCollected = CollectPrimvarData(PrimvarFloat2);
+			}
+			else if (pxr::VtArray<pxr::GfVec3f> PrimvarFloat3; PrimvarType == pxr::SdfValueTypeNames->Float3Array && Primvar.Get(&(PrimvarFloat3), Options.TimeCode))
+			{
+				bDataCollected = CollectPrimvarData(PrimvarFloat3);
+			}
+			else if (pxr::VtArray<pxr::GfVec4f> PrimvarFloat4; PrimvarType == pxr::SdfValueTypeNames->Float4Array && Primvar.Get(&(PrimvarFloat4), Options.TimeCode))
+			{
+				CustomPrimvar.Data = MoveTemp(PrimvarFloat4);
+				bDataCollected = true;
+			}
+
+			// only add if we collected the data
+			if(bDataCollected)
+			{
+				CustomPrimvar.Interpolation = GetAttrInterpolation(Primvar, DefaultInterpolation);
+				Primvar.GetIndices(&CustomPrimvar.Indices, Options.TimeCode);
+				InOutMeshData.CustomPrimvars.FindOrAdd(PrimvarName.c_str(), MoveTemp(CustomPrimvar));
+			}
+		}
+	}
+
 	bool CollectMeshData(
 		const pxr::UsdPrim& Prim,
 		const UsdToUnreal::FUsdMeshConversionOptions& Options,
@@ -1394,6 +1484,11 @@ namespace UE::UsdGeomMeshConversion::Private
 				InOutMeshData.VelocityInterpolation = GetAttrInterpolation(VelocitiesAttr, DefaultInterpolation);
 			}
 
+			// Custom primvars
+			{
+				CollectPrimvarData(Prim, Options, InOutMeshData);
+			}
+
 			// Collect the subdivision attributes only if we plan on subdividing
 			if (Options.SubdivisionLevel > 0)
 			{
@@ -2383,6 +2478,29 @@ namespace UE::UsdGeomMeshConversion::Private
 			UVs = ComputeFlattened(UVs, UVIndices);
 			UVIndices = {};
 		}
+
+		for (TPair<FString, FUsdMeshData::FCustomPrimvar>& Pair : InOutMeshData.CustomPrimvars)
+		{
+			FUsdMeshData::FCustomPrimvar& Primvar = Pair.Value;
+			Primvar.Data = ComputeFlattened(Primvar.Data, Primvar.Indices);
+			Primvar.Indices = {};
+		}
+	}
+
+	void RegisterPrimvarAttributes(const TMap<FString, FUsdMeshData::FCustomPrimvar>& InCustomPrimvars, FMeshDescription& OutMeshDescription)
+	{
+		// Register Primvars attribute
+		for (const TPair<FString, FUsdMeshData::FCustomPrimvar>& Pair : InCustomPrimvars)
+		{
+			const FString& PrimvarName = Pair.Key;
+			const FUsdMeshData::FCustomPrimvar& CustomPrimvar = Pair.Value;
+
+			// Register the primvar if it doesn't exist
+			if (!OutMeshDescription.VertexInstanceAttributes().HasAttribute(*PrimvarName))
+			{
+				OutMeshDescription.VertexInstanceAttributes().RegisterAttribute<FVector4f>(*PrimvarName);
+			}
+		}
 	}
 
 	bool ConvertMeshData(
@@ -2406,6 +2524,10 @@ namespace UE::UsdGeomMeshConversion::Private
 		{
 			ensure(InMeshData.UVSetIndices[UVSetIndex].empty());
 		}
+		for (const TPair<FString, FUsdMeshData::FCustomPrimvar>& Pair : InMeshData.CustomPrimvars)
+		{
+			ensure(Pair.Value.Indices.empty());
+		}
 
 		// Material assignments
 		const TArray<UsdUtils::FUsdPrimMaterialSlot>& LocalMaterialSlots = InMeshData.LocalMaterialInfo.Slots;
@@ -2572,6 +2694,9 @@ namespace UE::UsdGeomMeshConversion::Private
 				}
 			}
 
+			// Register Primvars attribute
+			RegisterPrimvarAttributes(InMeshData.CustomPrimvars, OutMeshDescription);
+
 			// UVs
 			TVertexInstanceAttributesRef<FVector2f> MeshDescriptionUVs = StaticMeshAttributes.GetVertexInstanceUVs();
 
@@ -2643,6 +2768,13 @@ namespace UE::UsdGeomMeshConversion::Private
 																						MeshAttribute::VertexInstance::Velocity
 																					);
 
+			// We handle all the primvars as FVector4f, because during the baking phase we'll put them in the Color channel (see InterchangeGenericTexturePipeline::ExecutePostFactoryPipeline)
+			TMap<FString, TVertexInstanceAttributesRef<FVector4f>> MeshDescriptionPrimvar;
+			for (const TPair<FString, FUsdMeshData::FCustomPrimvar>& Pair : InMeshData.CustomPrimvars)
+			{
+				MeshDescriptionPrimvar.Add(Pair.Key, OutMeshDescription.VertexInstanceAttributes().GetAttributesRef<FVector4f>(*(Pair.Key)));
+			}
+
 			OutMeshDescription.ReserveNewVertexInstances(InMeshData.FaceVertexCounts.size() * 3);
 			OutMeshDescription.ReserveNewPolygons(InMeshData.FaceVertexCounts.size());
 			OutMeshDescription.ReserveNewEdges(InMeshData.FaceVertexCounts.size() * 2);
@@ -2813,6 +2945,40 @@ namespace UE::UsdGeomMeshConversion::Private
 							}
 						}
 					}
+
+					//Primvars
+					for (const TPair<FString, FUsdMeshData::FCustomPrimvar>& Pair : InMeshData.CustomPrimvars)
+					{
+						const FString& PrimvarName = Pair.Key;
+						const FUsdMeshData::FCustomPrimvar& CustomPrimvar = Pair.Value;
+						const size_t ValueIndex = GetPrimValueIndex(CustomPrimvar.Interpolation,
+																	ControlPointIndex,
+																	CurrentVertexInstanceIndex,
+																	PolygonIndex);
+
+						if (!CustomPrimvar.Data.empty())
+						{
+							if (CustomPrimvar.Data.size() > ValueIndex)
+							{
+								const pxr::GfVec4f& VertexPrimvar = CustomPrimvar.Data[ValueIndex];
+								TVertexInstanceAttributesRef<FVector4f>& VertexInstanceAttributesPrimvar = *MeshDescriptionPrimvar.Find(PrimvarName);
+								VertexInstanceAttributesPrimvar[AddedVertexInstanceId] = FVector4f{ VertexPrimvar[0], VertexPrimvar[1], VertexPrimvar[2], VertexPrimvar[3] };
+							}
+							else
+							{
+								UE_LOG(
+									LogUsd,
+									Warning,
+									TEXT("Trying to read Primvar '%s' at index %u from prim '%s' but the prim only has %u values! Using "
+									"zeros instead."),
+									*PrimvarName,
+									ValueIndex,
+									*InMeshData.SourcePrimPath,
+									InMeshData.DisplayColors.size()
+								);
+							}
+						}
+					}
 				}
 
 				// This polygon was using the same vertex instance more than once and we removed too many
