Fix warning about MeshModelingToolset dependancy

diff --git a/Engine/Plugins/Interchange/Runtime/Source/Pipelines/Private/InterchangeUSDPipeline.cpp b/Engine/Plugins/Interchange/Runtime/Source/Pipelines/Private/InterchangeUSDPipeline.cpp
--- a/Engine/Plugins/Interchange/Runtime/Source/Pipelines/Private/InterchangeUSDPipeline.cpp
+++ b/Engine/Plugins/Interchange/Runtime/Source/Pipelines/Private/InterchangeUSDPipeline.cpp
@@ -14,7 +14,6 @@
 #include "InterchangePipelineLog.h"
 #include "Nodes/InterchangeBaseNodeContainer.h"
 
-#include "AssetUtils/Texture2DBuilder.h"
 #include "DynamicMesh/DynamicMesh3.h"
 #include "Engine/StaticMesh.h"
 #include "Materials/Material.h"
@@ -33,6 +32,10 @@
 #include "Usd/InterchangeUsdDefinitions.h"
 #include "UsdWrappers/SdfPath.h"
 
+#if WITH_EDITORONLY_DATA
+#include "AssetUtils/Texture2DBuilder.h"
+#endif // WITH_EDITORONLY_DATA
+
 #include UE_INLINE_GENERATED_CPP_BY_NAME(InterchangeUSDPipeline)
 
 namespace UE::InterchangeUsdPipeline::Private
