USD: Fix InterchangeUSD retaining every .vdb file reference on asset import data, leading to hundreds of entries for animated SVTs.

diff --git a/Engine/Plugins/Interchange/Runtime/Source/Import/Private/Volume/InterchangeSparseVolumeTextureFactory.cpp b/Engine/Plugins/Interchange/Runtime/Source/Import/Private/Volume/InterchangeSparseVolumeTextureFactory.cpp
--- a/Engine/Plugins/Interchange/Runtime/Source/Import/Private/Volume/InterchangeSparseVolumeTextureFactory.cpp
+++ b/Engine/Plugins/Interchange/Runtime/Source/Import/Private/Volume/InterchangeSparseVolumeTextureFactory.cpp
@@ -741,17 +741,10 @@ UInterchangeFactoryBase::FImportAssetResult UInterchangeSparseVolumeTextureFacto
 #if WITH_EDITOR
 	// Create SourceFile hashes while we're still in an async thread (we'll move this into AssetImportData later)
 	{
-		SourceFiles.Empty(VolumeNodes.Num());
-
-		for (const UInterchangeVolumeNode* VolumeNode : VolumeNodes)
-		{
-			FString FileName;
-			VolumeNode->GetCustomFileName(FileName);
-
-			FAssetImportInfo::FSourceFile& SourceFile = SourceFiles.Emplace_GetRef();
-			SourceFile.RelativeFilename = FileName;
-		}
-
+		// Just hash one file though, as that's what the standard SVT importer seems to also do
+		SourceFiles.Empty(1);
+		FAssetImportInfo::FSourceFile& SourceFile = SourceFiles.Emplace_GetRef();
+		SourceFile.RelativeFilename = Arguments.SourceData->GetFilename();
 		HashSourceFiles(SourceFiles);
 	}
 #endif	  // WITH_EDITOR
