Interchange: Fix incorrect FText::Format placeholders

diff --git a/Engine/Plugins/Interchange/Runtime/Source/Import/Private/Material/InterchangeMaterialFactory.cpp b/Engine/Plugins/Interchange/Runtime/Source/Import/Private/Material/InterchangeMaterialFactory.cpp
--- a/Engine/Plugins/Interchange/Runtime/Source/Import/Private/Material/InterchangeMaterialFactory.cpp
+++ b/Engine/Plugins/Interchange/Runtime/Source/Import/Private/Material/InterchangeMaterialFactory.cpp
@@ -848,7 +848,7 @@ UInterchangeFactoryBase::FImportAssetResult UInterchangeMaterialFactory::BeginIm
 		Message->SourceAssetName = Arguments.SourceData->GetFilename();
 		Message->DestinationAssetName = Arguments.AssetName;
 		Message->AssetType = GetFactoryClass();
-		Message->Text = FText::Format(LOCTEXT("MatFactory_CouldNotCreateMat", "Could not create Material asset %s. Reason: %s"), FText::FromString(Arguments.AssetName), Info);
+		Message->Text = FText::Format(LOCTEXT("MatFactory_CouldNotCreateMat", "Could not create Material asset {0}. Reason: {1}"), FText::FromString(Arguments.AssetName), Info);
 		bSkipImport = true;
 		ImportAssetResult.bIsFactorySkipAsset = true;
 	};
@@ -2063,7 +2063,7 @@ UInterchangeFactoryBase::FImportAssetResult UInterchangeMaterialFunctionFactory:
 		Message->SourceAssetName = Arguments.SourceData->GetFilename();
 		Message->DestinationAssetName = Arguments.AssetName;
 		Message->AssetType = GetFactoryClass();
-		Message->Text = FText::Format(LOCTEXT("MatFunc_CouldNotCreateMat", "Could not create Material asset %s. Reason: %s"), FText::FromString(Arguments.AssetName), Info);
+		Message->Text = FText::Format(LOCTEXT("MatFunc_CouldNotCreateMat", "Could not create Material asset {0}. Reason: {1}"), FText::FromString(Arguments.AssetName), Info);
 		bSkipImport = true;
 		ImportAssetResult.bIsFactorySkipAsset = true;
 	};
diff --git a/Engine/Plugins/Interchange/Runtime/Source/Import/Private/Scene/InterchangeDecalActorFactory.cpp b/Engine/Plugins/Interchange/Runtime/Source/Import/Private/Scene/InterchangeDecalActorFactory.cpp
--- a/Engine/Plugins/Interchange/Runtime/Source/Import/Private/Scene/InterchangeDecalActorFactory.cpp
+++ b/Engine/Plugins/Interchange/Runtime/Source/Import/Private/Scene/InterchangeDecalActorFactory.cpp
@@ -51,7 +51,7 @@ UObject* UInterchangeDecalActorFactory::ProcessActor(AActor& SpawnedActor, const
 				}
 				else
 				{
-					const FText Message = FText::Format(NSLOCTEXT("DecalActorImport", "NoMaterial", "No valid decal material found. Make sure that the DecalMaterialPath is valid: %s"), FText::FromString(DecalMaterialPath));
+					const FText Message = FText::Format(NSLOCTEXT("DecalActorImport", "NoMaterial", "No valid decal material found. Make sure that the DecalMaterialPath is valid: {0}"), FText::FromString(DecalMaterialPath));
 #if WITH_EDITOR
 					LogMessage<UInterchangeResultError_Generic>(Params, Message, SpawnedActor.GetActorLabel());
 #else
diff --git a/Engine/Plugins/Interchange/Runtime/Source/Import/Private/Texture/InterchangeTextureFactory.cpp b/Engine/Plugins/Interchange/Runtime/Source/Import/Private/Texture/InterchangeTextureFactory.cpp
--- a/Engine/Plugins/Interchange/Runtime/Source/Import/Private/Texture/InterchangeTextureFactory.cpp
+++ b/Engine/Plugins/Interchange/Runtime/Source/Import/Private/Texture/InterchangeTextureFactory.cpp
@@ -1932,7 +1932,7 @@ UInterchangeFactoryBase::FImportAssetResult UInterchangeTextureFactory::BeginImp
 		Message->SourceAssetName = Arguments.SourceData->GetFilename();
 		Message->DestinationAssetName = Arguments.AssetName;
 		Message->AssetType = GetFactoryClass();
-		Message->Text = FText::Format(LOCTEXT("TexFact_CouldNotCreateMat", "UInterchangeTextureFactory: Could not create texture asset %s. Reason: %s"), FText::FromString(Arguments.AssetName), Info);
+		Message->Text = FText::Format(LOCTEXT("TexFact_CouldNotCreateMat", "UInterchangeTextureFactory: Could not create texture asset {0}. Reason: {1}"), FText::FromString(Arguments.AssetName), Info);
 		bSkipImport = true;
 		ImportAssetResult.bIsFactorySkipAsset = true;
 	};
@@ -1961,7 +1961,7 @@ UInterchangeFactoryBase::FImportAssetResult UInterchangeTextureFactory::BeginImp
 	FTextureNodeVariant TextureNodeVariant = GetTextureNodeVariantFromFactoryVariant(GetAsTextureFactoryNodeVariant(Arguments.AssetNode, SupportedFactoryNodeClass), Arguments.NodeContainer);
 	if (TextureNodeVariant.IsType<FEmptyVariantState>())
 	{
-		FText Info = FText::Format(LOCTEXT("TextureFactory_InvalidTextureTranslated", "Asset factory node (%s) doesn't reference a valid texture translated node.")
+		FText Info = FText::Format(LOCTEXT("TextureFactory_InvalidTextureTranslated", "Asset factory node ({0}) doesn't reference a valid texture translated node.")
 			, FText::FromString(SupportedFactoryNodeClass->GetAuthoredName()));
 		CouldNotCreateTextureLog(Info);
 		return ImportAssetResult;
@@ -2078,7 +2078,7 @@ UInterchangeFactoryBase::FImportAssetResult UInterchangeTextureFactory::ImportAs
 	FTextureNodeVariant TextureNodeVariant = GetTextureNodeVariantFromFactoryVariant(TextureFactoryNodeVariant, Arguments.NodeContainer);
 	if (TextureNodeVariant.IsType<FEmptyVariantState>())
 	{
-		FText Info = FText::Format(LOCTEXT("TextureFactory_Async_InvalidTextureTranslated", "UInterchangeTextureFactory: Asset factory node (%s) doesn't reference a valid texture translated node.")
+		FText Info = FText::Format(LOCTEXT("TextureFactory_Async_InvalidTextureTranslated", "UInterchangeTextureFactory: Asset factory node ({0}) doesn't reference a valid texture translated node.")
 			, FText::FromString(SupportedFactoryNodeClass->GetAuthoredName()));
 		ImportTextureErrorLog(Info);
 		return ImportAssetResult;
diff --git a/Engine/Plugins/Interchange/Runtime/Source/Parsers/GLTFCore/Private/GLTF/ExtensionsHandler.cpp b/Engine/Plugins/Interchange/Runtime/Source/Parsers/GLTFCore/Private/GLTF/ExtensionsHandler.cpp
--- a/Engine/Plugins/Interchange/Runtime/Source/Parsers/GLTFCore/Private/GLTF/ExtensionsHandler.cpp
+++ b/Engine/Plugins/Interchange/Runtime/Source/Parsers/GLTFCore/Private/GLTF/ExtensionsHandler.cpp
@@ -797,7 +797,7 @@ namespace GLTF
 		else if (Type == TEXT("directional"))
 			Light.Type = FLight::EType::Directional;
 		else
-			Messages.Emplace(RuntimeWarningSeverity(), FText::Format(LOCTEXT("UnspecifiedLightType", "Light has no type specified: %s"), FText::FromString(Light.Name)));
+			Messages.Emplace(RuntimeWarningSeverity(), FText::Format(LOCTEXT("UnspecifiedLightType", "Light has no type specified: {0}"), FText::FromString(Light.Name)));
 	}
 
 }  // namespace GLTF
