Interchange: Add new MaterialReferenceNode to allow describing material assignments of existing material assets;
USD: Add support for "Unreal" render context materials via USD Interchange.

Small cleanup of UInterchangeGenericMaterialPipeline;
Iterate over node container once instead of 4 times during UInterchangeGenericMaterialPipeline::ExecutePipeline.

diff --git a/Engine/Plugins/Interchange/Runtime/Source/FactoryNodes/Private/InterchangeMaterialFactoryNode.cpp b/Engine/Plugins/Interchange/Runtime/Source/FactoryNodes/Private/InterchangeMaterialFactoryNode.cpp
--- a/Engine/Plugins/Interchange/Runtime/Source/FactoryNodes/Private/InterchangeMaterialFactoryNode.cpp
+++ b/Engine/Plugins/Interchange/Runtime/Source/FactoryNodes/Private/InterchangeMaterialFactoryNode.cpp
@@ -54,7 +54,7 @@ bool UInterchangeMaterialFactoryNode::ConnectToBaseColor(const FString& Attribut
 {
 	return UInterchangeShaderPortsAPI::ConnectDefaultOuputToInput(this, UE::Interchange::Materials::PBRMR::Parameters::BaseColor.ToString(), AttributeValue);
 }
-	
+
 bool UInterchangeMaterialFactoryNode::ConnectOutputToBaseColor(const FString& ExpressionNodeUid, const FString& OutputName)
 {
 	return UInterchangeShaderPortsAPI::ConnectOuputToInputByName(this, UE::Interchange::Materials::PBRMR::Parameters::BaseColor.ToString(), ExpressionNodeUid, OutputName);
@@ -486,6 +486,21 @@ bool UInterchangeMaterialInstanceFactoryNode::SetCustomParent(const FString& Att
 	IMPLEMENT_NODE_ATTRIBUTE_SETTER_NODELEGATE(Parent, FString);
 }
 
+FString UInterchangeMaterialReferenceFactoryNode::GetTypeName() const
+{
+	const FString TypeName = TEXT("MaterialReferenceFactoryNode");
+	return TypeName;
+}
+
+UClass* UInterchangeMaterialReferenceFactoryNode::GetObjectClass() const
+{
+#if WITH_ENGINE
+	return UMaterialInterface::StaticClass();
+#else
+	return nullptr;
+#endif
+}
+
 FString UInterchangeMaterialFunctionCallExpressionFactoryNode::GetTypeName() const
 {
 	const FString TypeName = TEXT("MaterialFunctionCallExpressionFactoryNode");
@@ -499,7 +514,7 @@ bool UInterchangeMaterialFunctionCallExpressionFactoryNode::GetCustomMaterialFun
 
 bool UInterchangeMaterialFunctionCallExpressionFactoryNode::SetCustomMaterialFunctionDependency(const FString& AttributeValue)
 {
-	
+
 	auto ImplementationFunction = [this, &AttributeValue]()
 		{
 			IMPLEMENT_NODE_ATTRIBUTE_SETTER_NODELEGATE(MaterialFunctionDependency, FString);
diff --git a/Engine/Plugins/Interchange/Runtime/Source/FactoryNodes/Public/InterchangeMaterialFactoryNode.h b/Engine/Plugins/Interchange/Runtime/Source/FactoryNodes/Public/InterchangeMaterialFactoryNode.h
--- a/Engine/Plugins/Interchange/Runtime/Source/FactoryNodes/Public/InterchangeMaterialFactoryNode.h
+++ b/Engine/Plugins/Interchange/Runtime/Source/FactoryNodes/Public/InterchangeMaterialFactoryNode.h
@@ -24,7 +24,6 @@ public:
 	UFUNCTION(BlueprintCallable, Category = "Interchange | Node | Material")
 	bool GetCustomIsMaterialImportEnabled(bool& AttributeValue) const;
 
-	/** Sets if this shader graph should be rendered two sided or not. Defaults to off. */
 	UFUNCTION(BlueprintCallable, Category = "Interchange | Node | Material")
 	bool SetCustomIsMaterialImportEnabled(const bool& AttributeValue);
 
@@ -332,7 +331,24 @@ public:
 private:
 	const UE::Interchange::FAttributeKey Macro_CustomInstanceClassNameKey = UE::Interchange::FAttributeKey(TEXT("InstanceClassName"));
 	const UE::Interchange::FAttributeKey Macro_CustomParentKey = UE::Interchange::FAttributeKey(TEXT("Parent"));
-	
+
+};
+
+/**
+ * Describes a reference to an existing (as in, not imported) material. Note that the material is referenced
+ * via the UInterchangeFactoryBaseNode::CustomReferenceObject member.
+ *
+ * The idea is that mesh / actor factory nodes can reference one of these nodes as a slot dependency, and
+ * Interchange will assign that existing material to the corresponding slot during import
+ */
+UCLASS(BlueprintType)
+class INTERCHANGEFACTORYNODES_API UInterchangeMaterialReferenceFactoryNode : public UInterchangeBaseMaterialFactoryNode
+{
+	GENERATED_BODY()
+
+public:
+	virtual FString GetTypeName() const override;
+	virtual UClass* GetObjectClass() const override;
 };
 
 UCLASS(BlueprintType)
diff --git a/Engine/Plugins/Interchange/Runtime/Source/Import/Private/Usd/InterchangeUsdTranslator.cpp b/Engine/Plugins/Interchange/Runtime/Source/Import/Private/Usd/InterchangeUsdTranslator.cpp
--- a/Engine/Plugins/Interchange/Runtime/Source/Import/Private/Usd/InterchangeUsdTranslator.cpp
+++ b/Engine/Plugins/Interchange/Runtime/Source/Import/Private/Usd/InterchangeUsdTranslator.cpp
@@ -41,6 +41,7 @@
 #include "InterchangeCameraNode.h"
 #include "InterchangeLightNode.h"
 #include "InterchangeManager.h"
+#include "InterchangeMaterialReferenceNode.h"
 #include "InterchangeMeshNode.h"
 #include "InterchangeSceneNode.h"
 #include "InterchangeShaderGraphNode.h"
@@ -95,6 +96,7 @@ namespace UE::InterchangeUsdTranslator::Private
 	const static FString LightPrefix = TEXT("\\Light\\");
 	const static FString LODPrefix = TEXT("\\LOD\\");
 	const static FString MaterialPrefix = TEXT("\\Material\\");
+	const static FString MaterialReferencePrefix = TEXT("\\MaterialReference\\");
 	const static FString MeshPrefix = TEXT("\\Mesh\\");
 	const static FString MorphTargetPrefix = TEXT("\\MorphTarget\\");
 	const static FString PrimitiveShapePrefix = TEXT("\\PrimitiveShape\\");
@@ -820,6 +822,61 @@ namespace UE::InterchangeUsdTranslator::Private
 		NodeContainer.AddNode(Node);
 	}
 
+	FString AddDisplayColorMaterialInstanceNodeIfNeeded(UInterchangeBaseNodeContainer& NodeContainer, const FString& DisplayColorDesc)
+	{
+		using namespace UsdUnreal::MaterialUtils;
+
+		FString NodeUid = MaterialPrefix + DisplayColorDesc;
+
+		// We'll treat the DisplayColorDesc (something like "!DisplayColor_1_0") as the material instance UID here
+		const UInterchangeShaderGraphNode* Node = Cast<UInterchangeShaderGraphNode>(NodeContainer.GetNode(NodeUid));
+		if (Node)
+		{
+			return NodeUid;
+		}
+
+		// Need to create a new instance
+		TOptional<FDisplayColorMaterial> ParsedMat = FDisplayColorMaterial::FromString(DisplayColorDesc);
+		if (!ParsedMat)
+		{
+			return {};
+		}
+		FString NodeName = ParsedMat->ToPrettyString();
+
+		UInterchangeShaderGraphNode* NewNode = NewObject<UInterchangeShaderGraphNode>(&NodeContainer);
+		NewNode->InitializeNode(NodeUid, NodeName, EInterchangeNodeContainerType::TranslatedAsset);
+		NewNode->AddBooleanAttribute(UE::Interchange::USD::ParseMaterialIdentifier, true);
+		NewNode->AddStringAttribute(*UE::Interchange::USD::DisplayColorDescIdentifier, DisplayColorDesc);
+
+		NodeContainer.AddNode(NewNode);
+
+		return NodeUid;
+	}
+
+	FString AddUnrealMaterialReferenceNodeIfNeeded(UInterchangeBaseNodeContainer& NodeContainer, const FString& ContentPath)
+	{
+		using namespace UsdUnreal::MaterialUtils;
+
+		// e.g. "\\MaterialReference\\/Game/MyFolder/Red.Red"
+		const FString NodeUid = MaterialReferencePrefix + ContentPath;
+
+		const UInterchangeMaterialReferenceNode* Node = Cast<UInterchangeMaterialReferenceNode>(NodeContainer.GetNode(NodeUid));
+		if (Node)
+		{
+			return NodeUid;
+		}
+
+		const FString DisplayName = FPaths::GetBaseFilename(ContentPath);
+
+		UInterchangeMaterialReferenceNode* NewNode = NewObject<UInterchangeMaterialReferenceNode>(&NodeContainer);
+		NewNode->InitializeNode(NodeUid, DisplayName, EInterchangeNodeContainerType::TranslatedAsset);
+		NewNode->SetCustomContentPath(ContentPath);
+
+		NodeContainer.AddNode(NewNode);
+
+		return NodeUid;
+	}
+
 	void UInterchangeUSDTranslatorImpl::SetMaterialSlotDependencies(
 		const FString& MaterialPrimName,
 		const FString& MaterialUid,
@@ -909,7 +966,7 @@ namespace UE::InterchangeUsdTranslator::Private
 		TArray<FUsdMaterialXShaderGraph::FGeomProp>& GeomPropValues
 	)
 	{
-		FName RenderContext = TranslatorSettings ? TranslatorSettings->RenderContext : UnrealIdentifiers::UniversalRenderContext;
+		FName RenderContext = TranslatorSettings ? TranslatorSettings->RenderContext : UnrealIdentifiers::UnrealRenderContext;
 
 		// Check for any references of MaterialX
 #if WITH_EDITOR
@@ -1132,6 +1189,24 @@ namespace UE::InterchangeUsdTranslator::Private
 		UInterchangeBaseNodeContainer& NodeContainer
 	)
 	{
+		FName RenderContext = TranslatorSettings ? TranslatorSettings->RenderContext : UnrealIdentifiers::UniversalRenderContext;
+
+		// If this material has an unreal surface output and we're in the unreal render context, just emit a material reference,
+		// as we never want this to become a UMaterial / UMaterialInstance anyway.
+		//
+		// We could just early out here completely and not emit anything, as we also emit the material reference node on-demand, whenever
+		// we parse an actual material assignment from a Mesh. The user may have custom pipelines that expect to find these though,
+		// even if no mesh is actually using the materials
+		if (RenderContext == UnrealIdentifiers::UnrealRenderContext)
+		{
+			TOptional<FString> UnrealContentPath = UsdUtils::GetUnrealSurfaceOutput(Prim);
+			if (UnrealContentPath.IsSet())
+			{
+				UE::InterchangeUsdTranslator::Private::AddUnrealMaterialReferenceNodeIfNeeded(NodeContainer, UnrealContentPath.GetValue());
+				return;
+			}
+		}
+
 		FString PrimPath = Prim.GetPrimPath().GetString();
 		FString MaterialUid = MaterialPrefix + PrimPath;
 		FString MaterialPrimName(Prim.GetName().ToString());
@@ -1149,21 +1224,6 @@ namespace UE::InterchangeUsdTranslator::Private
 			return;
 		}
 
-		FName RenderContext = TranslatorSettings ? TranslatorSettings->RenderContext : UnrealIdentifiers::UniversalRenderContext;
-
-		if (RenderContext == UnrealIdentifiers::UnrealRenderContext)
-		{
-			UE_LOG(
-				LogUsd,
-				Warning,
-				TEXT(
-					"The 'unreal' render context is not yet supported via USD Interchange: The material '%s' will use the universal render context instead"
-				),
-				*PrimPath
-			);
-			RenderContext = UnrealIdentifiers::UniversalRenderContext;
-		}
-
 		UInterchangeShaderGraphNode* MaterialNode = NewObject<UInterchangeShaderGraphNode>(&NodeContainer);
 		MaterialNode->InitializeNode(MaterialUid, MaterialPrimName, EInterchangeNodeContainerType::TranslatedAsset);
 		MaterialNode->SetAssetName(MaterialPrimName);
@@ -1476,35 +1536,6 @@ namespace UE::InterchangeUsdTranslator::Private
 		CurrentLODSceneNodes.Reset();
 	}
 
-	void AddDisplayColorMaterialInstanceNodeIfNeeded(UInterchangeBaseNodeContainer& NodeContainer, const FString& DisplayColorDesc)
-	{
-		using namespace UsdUnreal::MaterialUtils;
-
-		FString NodeUid = MaterialPrefix + DisplayColorDesc;
-
-		// We'll treat the DisplayColorDesc (something like "!DisplayColor_1_0") as the material instance UID here
-		const UInterchangeShaderGraphNode* Node = Cast<UInterchangeShaderGraphNode>(NodeContainer.GetNode(DisplayColorDesc));
-		if (Node)
-		{
-			return;
-		}
-
-		// Need to create a new instance
-		TOptional<FDisplayColorMaterial> ParsedMat = FDisplayColorMaterial::FromString(DisplayColorDesc);
-		if (!ParsedMat)
-		{
-			return;
-		}
-		FString NodeName = ParsedMat->ToPrettyString();
-
-		UInterchangeShaderGraphNode* NewNode = NewObject<UInterchangeShaderGraphNode>(&NodeContainer);
-		NewNode->InitializeNode(NodeUid, NodeName, EInterchangeNodeContainerType::TranslatedAsset);
-		NewNode->AddBooleanAttribute(UE::Interchange::USD::ParseMaterialIdentifier, true);
-		NewNode->AddStringAttribute(*UE::Interchange::USD::DisplayColorDescIdentifier, DisplayColorDesc);
-
-		NodeContainer.AddNode(NewNode);
-	}
-
 	void AddLightNode(const UE::FUsdPrim& Prim, UInterchangeBaseNodeContainer& NodeContainer)
 	{
 #if USE_USD_SDK
@@ -1850,24 +1881,23 @@ namespace UE::InterchangeUsdTranslator::Private
 
 			for (const UsdUtils::FUsdPrimMaterialSlot& Slot : Assignments.Slots)
 			{
-				// Use the material prim path/display color desc as the material slot name, because Interchange
-				// already has a mechanism to merge material slots with the same name. Using the material name itself
-				// as the slot name has Interchange combine slots with identical materials, which works fine. If we
-				// were to use GeomSubset names or prim names in here though, it's possible that two similarly named
-				// slots in different skeletal mesh chunks (but with different materials!) could get merged together,
-				// which is not what we want
+				// We do this because Interchange will, in some scenarios, merge material slots with identical slot names.
+				// By using the source (which is the displaycolor desc / material prim path / unreal material content path)
+				// we do end up with a goofy looking super long material slot names, but it will have Interchange only combine
+				// slots if they really are pointing at the exact same thing
 				const FString& SlotName = Slot.MaterialSource;
 
 				// Get the Uid of the material instance that we'll end up assigning to this slot
 				FString MaterialInstanceUid;
-				bool bIsDisplayColor = false;
+				bool bAddDirectlyAsSlotDependency = false;
 				switch (Slot.AssignmentType)
 				{
 					case UsdUtils::EPrimAssignmentType::DisplayColor:
 					{
-						bIsDisplayColor = true;
-						AddDisplayColorMaterialInstanceNodeIfNeeded(NodeContainer, Slot.MaterialSource);
-						MaterialInstanceUid = MaterialPrefix + Slot.MaterialSource;	   // This is e.g. "!DisplayColor_0_1"
+						bAddDirectlyAsSlotDependency = true;
+
+						// MaterialSource here is e.g. "!DisplayColor_0_1"
+						MaterialInstanceUid = AddDisplayColorMaterialInstanceNodeIfNeeded(NodeContainer, Slot.MaterialSource);
 						break;
 					}
 					case UsdUtils::EPrimAssignmentType::MaterialPrim:
@@ -1877,8 +1907,10 @@ namespace UE::InterchangeUsdTranslator::Private
 					}
 					case UsdUtils::EPrimAssignmentType::UnrealMaterial:
 					{
-						// TODO: We can't support these yet without a custom pipeline unfortunately
-						// We could spawn a material instance of the referenced material... That's probably not what you'd expect though
+						bAddDirectlyAsSlotDependency = true;
+
+						// MaterialSource here is the content path, e.g. "/Game/MyFolder/Red.Red"
+						MaterialInstanceUid = AddUnrealMaterialReferenceNodeIfNeeded(NodeContainer, Slot.MaterialSource);
 						break;
 					}
 					default:
@@ -1888,8 +1920,8 @@ namespace UE::InterchangeUsdTranslator::Private
 					}
 				}
 
-				// If we're a DisplayColor material, we know everything we need right now
-				if (bIsDisplayColor)
+				// If we're a DisplayColor/Unreal material, we know everything we need right now
+				if (bAddDirectlyAsSlotDependency)
 				{
 					MeshNode->SetSlotMaterialDependencyUid(SlotName, *MaterialInstanceUid);
 				}
@@ -3322,7 +3354,7 @@ namespace UE::InterchangeUsdTranslator::Private
 
 UInterchangeUsdTranslatorSettings::UInterchangeUsdTranslatorSettings()
 	: GeometryPurpose((int32)(EUsdPurpose::Default | EUsdPurpose::Proxy | EUsdPurpose::Render | EUsdPurpose::Guide))
-	, RenderContext(UnrealIdentifiers::UniversalRenderContext)	  // Default to the universal render context for now as we don't support 'unreal' yet
+	, RenderContext(UnrealIdentifiers::UnrealRenderContext)
 	, MaterialPurpose(*UnrealIdentifiers::MaterialPreviewPurpose)
 	, InterpolationType(EUsdInterpolationType::Linear)
 	, bOverrideStageOptions(false)
diff --git a/Engine/Plugins/Interchange/Runtime/Source/Nodes/Private/InterchangeMaterialReferenceNode.cpp b/Engine/Plugins/Interchange/Runtime/Source/Nodes/Private/InterchangeMaterialReferenceNode.cpp
new file mode 100644
--- /dev/null
+++ b/Engine/Plugins/Interchange/Runtime/Source/Nodes/Private/InterchangeMaterialReferenceNode.cpp
@@ -0,0 +1,24 @@
+// Copyright Epic Games, Inc. All Rights Reserved.
+
+#include "InterchangeMaterialReferenceNode.h"
+
+#include "InterchangeShaderGraphNode.h"
+#include "Nodes/InterchangeBaseNodeContainer.h"
+
+#include UE_INLINE_GENERATED_CPP_BY_NAME(InterchangeMaterialReferenceNode)
+
+FString UInterchangeMaterialReferenceNode::GetTypeName() const
+{
+	return TEXT("MaterialReferenceNode");
+}
+
+bool UInterchangeMaterialReferenceNode::GetCustomContentPath(FString& AttributeValue) const
+{
+	IMPLEMENT_NODE_ATTRIBUTE_GETTER(ContentPath, FString);
+}
+
+bool UInterchangeMaterialReferenceNode::SetCustomContentPath(const FString& AttributeValue) const
+{
+	IMPLEMENT_NODE_ATTRIBUTE_SETTER_NODELEGATE(ContentPath, FString);
+}
+
diff --git a/Engine/Plugins/Interchange/Runtime/Source/Nodes/Public/InterchangeMaterialReferenceNode.h b/Engine/Plugins/Interchange/Runtime/Source/Nodes/Public/InterchangeMaterialReferenceNode.h
new file mode 100644
--- /dev/null
+++ b/Engine/Plugins/Interchange/Runtime/Source/Nodes/Public/InterchangeMaterialReferenceNode.h
@@ -0,0 +1,35 @@
+// Copyright Epic Games, Inc. All Rights Reserved.
+
+#pragma once
+
+#include "CoreMinimal.h"
+#include "Nodes/InterchangeBaseNode.h"
+
+#include "InterchangeMaterialReferenceNode.generated.h"
+
+/**
+ * Describes a reference to an existing (as in, not imported) material.
+ *
+ * The idea is that mesh / actor nodes can reference one of these nodes as a slot dependency, and
+ * Interchange will assign that existing material to the corresponding slot during import
+ */
+UCLASS(BlueprintType)
+class INTERCHANGENODES_API UInterchangeMaterialReferenceNode : public UInterchangeBaseNode
+{
+	GENERATED_BODY()
+
+public:
+	virtual FString GetTypeName() const override;
+
+public:
+	/** Gets the content path of the target material (e.g. "/Game/MyFolder/Red.Red") */
+	UFUNCTION(BlueprintCallable, Category = "Interchange | Node | Material Reference")
+	bool GetCustomContentPath(FString& AttributeValue) const;
+
+	/** Sets the content path of the target material (e.g. "/Game/MyFolder/Red.Red") */
+	UFUNCTION(BlueprintCallable, Category = "Interchange | Node | Material Reference")
+	bool SetCustomContentPath(const FString& AttributeValue) const;
+
+private:
+	IMPLEMENT_NODE_ATTRIBUTE_KEY(ContentPath);
+};
diff --git a/Engine/Plugins/Interchange/Runtime/Source/Pipelines/Private/InterchangeGenericMaterialPipeline.cpp b/Engine/Plugins/Interchange/Runtime/Source/Pipelines/Private/InterchangeGenericMaterialPipeline.cpp
--- a/Engine/Plugins/Interchange/Runtime/Source/Pipelines/Private/InterchangeGenericMaterialPipeline.cpp
+++ b/Engine/Plugins/Interchange/Runtime/Source/Pipelines/Private/InterchangeGenericMaterialPipeline.cpp
@@ -9,6 +9,7 @@
 #include "InterchangeMaterialDefinitions.h"
 #include "InterchangeMaterialFactoryNode.h"
 #include "InterchangeMaterialInstanceNode.h"
+#include "InterchangeMaterialReferenceNode.h"
 #include "InterchangePipelineLog.h"
 #include "InterchangeShaderGraphNode.h"
 #include "InterchangeSourceData.h"
@@ -183,7 +184,7 @@ namespace UE::Interchange::InterchangeGenericMaterialPipeline::Private
 		UClass* CustomExpressionClass = UMaterialExpressionMaterialFunctionCall::StaticClass();
 
 		FunctionCallExpression.SetCustomExpressionClassName(CustomExpressionClass->GetName());
-		
+
 		FunctionCallExpression.AddStringAttribute(MaterialFunctionMemberName.ToString(), MaterialFunctionPath);
 		FunctionCallExpression.AddApplyAndFillDelegates<FString>(MaterialFunctionMemberName.ToString(), CustomExpressionClass, MaterialFunctionMemberName);
 	}
@@ -195,7 +196,7 @@ namespace UE::Interchange::InterchangeGenericMaterialPipeline::Private
 		const FString& MaterialFunctionPath)
 	{
 		const FString UniqueId = MaterialFactoryNode->GetUniqueID() + Label;
-		
+
 		UInterchangeMaterialExpressionFactoryNode* Expression = NewObject<UInterchangeMaterialExpressionFactoryNode>(BaseNodeContainer, NAME_None);
 		if (!Expression)
 		{
@@ -208,7 +209,7 @@ namespace UE::Interchange::InterchangeGenericMaterialPipeline::Private
 		BaseNodeContainer->SetNodeParentUid(UniqueId, MaterialFactoryNode->GetUniqueID());
 
 		UpdateFunctionCallExpression(*Expression, MaterialFunctionPath);
-		
+
 		return Expression;
 	}
 
@@ -357,7 +358,7 @@ namespace UE::Interchange::Materials::HashUtils
 	public:
 		FMaterialHashDebugData(const FString& InLogDirectoryPath)
 			:LogDirectoryPath(InLogDirectoryPath) {}
-		
+
 		void Reset();
 
 		void SaveLogsToFile(const FString& FileName);
@@ -494,12 +495,12 @@ namespace UE::Interchange::Materials::HashUtils
 		void CopyLeafInputsToFactoryNode(UInterchangeBaseMaterialFactoryNode* FactoryNode);
 
 		void ComputMaterialHash(const UInterchangeShaderGraphNode* ShaderGraphNode);
-		
+
 		void SetupOverridableScalarParameter(const UInterchangeShaderNode* ShaderNode, const FString& ParameterKey, const FString& OverridableParameterNameKey);
 		void SetupOverridableVectorParameter(const UInterchangeShaderNode* ShaderNode, const FString& ParameterKey, const FString& OverridableParameterNameKey);
 		void SetupOverridableStaticBoolParameter(const UInterchangeShaderNode* ShaderNode, const FString& ParameterKey, const FString& OverridableParameterNameKey);
 		void SetupOverridableTextureParameter(const UInterchangeShaderNode* ShaderNode, const FString& InputKey, const FString& OverridableParameterNameKey);
-		
+
 		/**
 		 * Creates a Base Material Factory Node based on if the material is a duplicate material or if it is found for the first time.
 		 * If the option to create a material instance for the parent is enabled, then additional material instance factory for parent would also be created.
@@ -596,8 +597,8 @@ HashDebugData->PopNodeAddressesToLastPopIndex();\
 #else
 
 #define ADD_LOG_MESSAGE(...)
-#define ADD_NODE_ADDRESS_MESSAGE() 
-#define PUSH_NODE_ADDRESS(Node) 
+#define ADD_NODE_ADDRESS_MESSAGE()
+#define PUSH_NODE_ADDRESS(Node)
 #define PUSH_NODE_ADDRESS_WITHOUT_CHECKPOINT(Node)
 #define POP_NODE_ADDRESSES()
 
@@ -692,7 +693,7 @@ void UInterchangeGenericMaterialPipeline::AdjustSettingsForContext(const FInterc
 			HideProperty(OuterMostPipeline, this, GET_MEMBER_NAME_CHECKED(UInterchangeGenericMaterialPipeline, AssetName));
 		}
 	}
-	
+
 #endif //WITH_EDITOR
 	using namespace UE::Interchange;
 
@@ -807,7 +808,7 @@ void UInterchangeGenericMaterialPipeline::ExecutePipeline(UInterchangeBaseNodeCo
 	{
 		SourceDatas.Add(SourceData);
 	}
-	
+
 	if (TexturePipeline)
 	{
 		TexturePipeline->ScriptedExecutePipeline(InBaseNodeContainer, InSourceDatas, ContentBasePath);
@@ -817,24 +818,46 @@ void UInterchangeGenericMaterialPipeline::ExecutePipeline(UInterchangeBaseNodeCo
 		SparseVolumeTexturePipeline->ScriptedExecutePipeline(InBaseNodeContainer, InSourceDatas, ContentBasePath);
 	}
 
-	
-	MaterialNodes.Empty();
-	MaterialFactoryNodes.Empty();
-	//Find all translated node we need for this pipeline
-	BaseNodeContainer->IterateNodes([this](const FString& NodeUid, UInterchangeBaseNode* Node)
-	{
-		switch(Node->GetNodeContainerType())
+	TSet<UInterchangeShaderGraphNode*> ShaderGraphNodes;
+	TSet<UInterchangeMaterialInstanceNode*> MaterialInstanceNodes;
+	TSet<UInterchangeMaterialReferenceNode*> MaterialReferenceNodes;
+	TSet<UInterchangeBaseMaterialFactoryNode*> BaseMaterialFactoryNodes;
+
+	// Find all nodes we need for this pipeline
+	BaseNodeContainer->IterateNodes(
+		[&ShaderGraphNodes,
+		 &MaterialInstanceNodes,
+		 &MaterialReferenceNodes,
+		 &BaseMaterialFactoryNodes](const FString& NodeUid, UInterchangeBaseNode* Node)
 		{
-			case EInterchangeNodeContainerType::TranslatedAsset:
+			switch (Node->GetNodeContainerType())
 			{
-				if (UInterchangeShaderGraphNode* MaterialNode = Cast<UInterchangeShaderGraphNode>(Node))
+				case EInterchangeNodeContainerType::TranslatedAsset:
+				{
+					if (UInterchangeShaderGraphNode* ShaderGraphNode = Cast<UInterchangeShaderGraphNode>(Node))
+					{
+						ShaderGraphNodes.Add(ShaderGraphNode);
+					}
+					else if (UInterchangeMaterialInstanceNode* InstanceNode = Cast<UInterchangeMaterialInstanceNode>(Node))
+					{
+						MaterialInstanceNodes.Add(InstanceNode);
+					}
+					else if (UInterchangeMaterialReferenceNode* ReferenceNode = Cast<UInterchangeMaterialReferenceNode>(Node))
+					{
+						MaterialReferenceNodes.Add(ReferenceNode);
+					}
+				}
+				case EInterchangeNodeContainerType::FactoryData:
 				{
-					MaterialNodes.Add(MaterialNode);
+					if (UInterchangeBaseMaterialFactoryNode* BaseMaterialFactoryNode = Cast<UInterchangeBaseMaterialFactoryNode>(Node))
+					{
+						BaseMaterialFactoryNodes.Add(BaseMaterialFactoryNode);
+					}
 				}
+				break;
 			}
-			break;
 		}
-	});
+	);
 
 	// Check to see whether materials should be created even if unused
 	// By default we let the setting of the pipeline to decide if we create the materials, every node with mesh attribute can enable/disable them, depending on the pipeline stack chosen.
@@ -854,10 +877,10 @@ void UInterchangeGenericMaterialPipeline::ExecutePipeline(UInterchangeBaseNodeCo
 
 	if (MaterialImport == EInterchangeMaterialImportOption::ImportAsMaterials)
 	{
-		for (const UInterchangeShaderGraphNode* ShaderGraphNode : MaterialNodes)
+		for (const UInterchangeShaderGraphNode* ShaderGraphNode : ShaderGraphNodes)
 		{
 			UInterchangeBaseMaterialFactoryNode* MaterialBaseFactoryNode = nullptr;
-			
+
 			bool bIsAShaderFunction;
 			if (ShaderGraphNode->GetCustomIsAShaderFunction(bIsAShaderFunction) && bIsAShaderFunction)
 			{
@@ -882,11 +905,16 @@ void UInterchangeGenericMaterialPipeline::ExecutePipeline(UInterchangeBaseNodeCo
 				/* Clearing the AttributeStorageNode as it might affect how the MaterialFunctionsFactories are created. */
 				AttributeStorageNode = nullptr;
 			}
+
+			if (MaterialBaseFactoryNode)
+			{
+				BaseMaterialFactoryNodes.Add(MaterialBaseFactoryNode);
+			}
 		}
 	}
 	else if (MaterialImport == EInterchangeMaterialImportOption::ImportAsMaterialInstances)
 	{
-		for (const UInterchangeShaderGraphNode* ShaderGraphNode : MaterialNodes)
+		for (const UInterchangeShaderGraphNode* ShaderGraphNode : ShaderGraphNodes)
 		{
 			if (UInterchangeMaterialInstanceFactoryNode* MaterialInstanceFactoryNode = CreateMaterialInstanceFactoryNode(ShaderGraphNode))
 			{
@@ -895,40 +923,36 @@ void UInterchangeGenericMaterialPipeline::ExecutePipeline(UInterchangeBaseNodeCo
 		}
 	}
 
-	
-	BaseNodeContainer->IterateNodesOfType<UInterchangeBaseMaterialFactoryNode>([&ContentBasePath, bClosureImportMaterials = bImportMaterials || bImportUnusedMaterial, ClosureSearchLocation = SearchLocation](const FString& NodeUid, UInterchangeBaseMaterialFactoryNode* MaterialBaseFactoryNode)
+	const bool bCombinedImportMaterials = bImportMaterials || bImportUnusedMaterial;
+	for (UInterchangeBaseMaterialFactoryNode* BaseMaterialFactoryNode : BaseMaterialFactoryNodes)
+	{
+		using namespace UE::Interchange::InterchangeGenericMaterialPipeline::Private;
+
+		BaseMaterialFactoryNode->SetCustomIsMaterialImportEnabled(bCombinedImportMaterials);
+
+		// Disable all materials if the toggle is off
+		if (!bCombinedImportMaterials)
 		{
-			if (MaterialBaseFactoryNode)
-			{
-				MaterialBaseFactoryNode->SetCustomIsMaterialImportEnabled(bClosureImportMaterials);
-				//Disable all materials if the toggle is off
-				if (!bClosureImportMaterials)
-				{
-					MaterialBaseFactoryNode->SetEnabled(false);
-				}
-				
-				//See if we can assign an existing material from the search location
-				FString MaterialName = MaterialBaseFactoryNode->GetDisplayLabel();
-				if (UMaterialInterface* ExistingMaterial = UE::Interchange::InterchangeGenericMaterialPipeline::Private::FindExistingMaterialFromSearchLocation(MaterialName, ContentBasePath, ClosureSearchLocation))
-				{
-					//Make sure we have the correct type of material (can be material instance) before setting the custom object reference.
-					const bool bIsMaterial = MaterialBaseFactoryNode->IsA<UInterchangeMaterialFactoryNode>() && ExistingMaterial->IsA<UMaterial>();
-					const bool bIsMaterialinstance = MaterialBaseFactoryNode->IsA<UInterchangeMaterialInstanceFactoryNode>() && ExistingMaterial->IsA<UMaterialInstance>();
-
-					MaterialBaseFactoryNode->SetCustomReferenceObject(ExistingMaterial);
-					//Reimport can only be done on material instances
-					const bool bEnableReimport = !bIsMaterial && bIsMaterialinstance;
-					MaterialBaseFactoryNode->SetCustomIsMaterialImportEnabled(bEnableReimport);
-					MaterialBaseFactoryNode->SetEnabled(bEnableReimport);
-				}
-			}
-		});
+			BaseMaterialFactoryNode->SetEnabled(false);
+		}
 
-	TArray<UInterchangeMaterialInstanceNode*> MaterialInstanceNodes;
-	BaseNodeContainer->IterateNodesOfType<UInterchangeMaterialInstanceNode>([&MaterialInstanceNodes](const FString& NodeUid, UInterchangeMaterialInstanceNode* MaterialNode)
+		// See if we can assign an existing material from the search location
+		FString MaterialName = BaseMaterialFactoryNode->GetDisplayLabel();
+		if (UMaterialInterface* ExistingMaterial = FindExistingMaterialFromSearchLocation(MaterialName, ContentBasePath, SearchLocation))
 		{
-			MaterialInstanceNodes.Add(MaterialNode);
-		});
+			// Make sure we have the correct type of material (can be material instance) before setting the custom object reference.
+			const bool bIsMaterial = BaseMaterialFactoryNode->IsA<UInterchangeMaterialFactoryNode>() && ExistingMaterial->IsA<UMaterial>();
+			const bool bIsMaterialinstance = BaseMaterialFactoryNode->IsA<UInterchangeMaterialInstanceFactoryNode>()
+											 && ExistingMaterial->IsA<UMaterialInstance>();
+
+			BaseMaterialFactoryNode->SetCustomReferenceObject(ExistingMaterial);
+
+			// Reimport can only be done on material instances
+			const bool bEnableReimport = !bIsMaterial && bIsMaterialinstance;
+			BaseMaterialFactoryNode->SetCustomIsMaterialImportEnabled(bEnableReimport);
+			BaseMaterialFactoryNode->SetEnabled(bEnableReimport);
+		}
+	}
 
 	for (UInterchangeMaterialInstanceNode* MaterialNode : MaterialInstanceNodes)
 	{
@@ -955,17 +979,17 @@ void UInterchangeGenericMaterialPipeline::ExecutePipeline(UInterchangeBaseNodeCo
 			MaterialFactoryNode = NewObject<UInterchangeMaterialInstanceFactoryNode>(BaseNodeContainer);
 			if (!ensure(MaterialFactoryNode))
 			{
-				MaterialFactoryNode = nullptr;
+				continue;
 			}
 			//Creating a Material
 			MaterialFactoryNode->InitializeNode(NodeUid, DisplayLabel, EInterchangeNodeContainerType::FactoryData);
 
 			BaseNodeContainer->AddNode(MaterialFactoryNode);
-			MaterialFactoryNodes.Add(MaterialFactoryNode);
 			MaterialFactoryNode->AddTargetNodeUid(MaterialNode->GetUniqueID());
 			MaterialNode->AddTargetNodeUid(MaterialFactoryNode->GetUniqueID());
 		}
 
+		BaseMaterialFactoryNodes.Add(MaterialFactoryNode);
 
 		// Set MaterialFactoryNode's display label to MaterialNode's uniqueID
 		// to reconcile mesh's slot names and material assets
@@ -976,7 +1000,7 @@ void UInterchangeGenericMaterialPipeline::ExecutePipeline(UInterchangeBaseNodeCo
 		MaterialFactoryNode->SetCustomInstanceClassName(MaterialClass->GetPathName());
 
 		UE::Interchange::InterchangeGenericMaterialPipeline::Private::AddSpecularProfileToFactoryNode(MaterialNode, MaterialFactoryNode, BaseNodeContainer);
-		
+
 		TArray<FString> Inputs;
 		UInterchangeShaderPortsAPI::GatherInputs(MaterialNode, Inputs);
 
@@ -1030,6 +1054,28 @@ void UInterchangeGenericMaterialPipeline::ExecutePipeline(UInterchangeBaseNodeCo
 		}
 	}
 
+	for (UInterchangeMaterialReferenceNode* MaterialReferenceNode : MaterialReferenceNodes)
+	{
+		FString MaterialContentPath;
+		if (MaterialReferenceNode->GetCustomContentPath(MaterialContentPath))
+		{
+			UInterchangeMaterialReferenceFactoryNode* FactoryNode = NewObject<UInterchangeMaterialReferenceFactoryNode>(BaseNodeContainer);
+
+			const FString DisplayLabel = MaterialReferenceNode->GetDisplayLabel();
+			const FString NodeUid = UInterchangeFactoryBaseNode::BuildFactoryNodeUid(MaterialReferenceNode->GetUniqueID());
+			FactoryNode->InitializeNode(NodeUid, DisplayLabel, EInterchangeNodeContainerType::FactoryData);
+
+			FactoryNode->SetCustomReferenceObject(FSoftObjectPath{MaterialContentPath});
+
+			BaseNodeContainer->AddNode(FactoryNode);
+
+			BaseMaterialFactoryNodes.Add(FactoryNode);
+
+			FactoryNode->AddTargetNodeUid(MaterialReferenceNode->GetUniqueID());
+			MaterialReferenceNode->AddTargetNodeUid(FactoryNode->GetUniqueID());
+		}
+	}
+
 	//If we have a valid override name
 	FString OverrideAssetName = IsStandAlonePipeline() ? DestinationName : FString();
 	if (OverrideAssetName.IsEmpty() && IsStandAlonePipeline())
@@ -1037,18 +1083,12 @@ void UInterchangeGenericMaterialPipeline::ExecutePipeline(UInterchangeBaseNodeCo
 		OverrideAssetName = AssetName;
 	}
 
-	if (IsStandAlonePipeline() && !OverrideAssetName.IsEmpty())
+	if (IsStandAlonePipeline() && !OverrideAssetName.IsEmpty() && BaseMaterialFactoryNodes.Num() == 1)
 	{
-		TArray<UInterchangeBaseMaterialFactoryNode*> BaseMaterialNodes;
-		BaseNodeContainer->IterateNodesOfType<UInterchangeBaseMaterialFactoryNode>([&BaseMaterialNodes](const FString& NodeUid, UInterchangeBaseMaterialFactoryNode* MaterialNode)
-			{
-				BaseMaterialNodes.Add(MaterialNode);
-			});
-
-		if(BaseMaterialNodes.Num() == 1)
+		if (UInterchangeBaseMaterialFactoryNode* Node = *BaseMaterialFactoryNodes.CreateIterator())
 		{
-			BaseMaterialNodes[0]->SetAssetName(OverrideAssetName);
-			BaseMaterialNodes[0]->SetDisplayLabel(OverrideAssetName);
+			Node->SetAssetName(OverrideAssetName);
+			Node->SetDisplayLabel(OverrideAssetName);
 		}
 	}
 
@@ -1129,11 +1169,10 @@ UInterchangeBaseMaterialFactoryNode* UInterchangeGenericMaterialPipeline::Create
 		{
 			return nullptr;
 		}
-		//Creating a Material
+
 		MaterialFactoryNode->InitializeNode(NodeUid, DisplayLabel, EInterchangeNodeContainerType::FactoryData);
-		
+
 		BaseNodeContainer->AddNode(MaterialFactoryNode);
-		MaterialFactoryNodes.Add(MaterialFactoryNode);
 		MaterialFactoryNode->AddTargetNodeUid(MaterialNode->GetUniqueID());
 		MaterialNode->AddTargetNodeUid(MaterialFactoryNode->GetUniqueID());
 	}
@@ -2284,7 +2323,7 @@ void UInterchangeGenericMaterialPipeline::HandleTextureCoordinateNode(const UInt
 			TexCoordFactoryNode = MultiplyExpression;
 		}
 	}
-	
+
 	// Rotate
 	{
 		TTuple<UInterchangeMaterialExpressionFactoryNode*, FString> RotateExpression =
@@ -2350,7 +2389,7 @@ void UInterchangeGenericMaterialPipeline::HandleLerpNode(const UInterchangeShade
 				ColorAExpression.Get<0>()->GetUniqueID(), ColorAExpression.Get<1>());
 		}
 	}
-	
+
 	// B
 	{
 		TTuple<UInterchangeMaterialExpressionFactoryNode*, FString> ColorBExpression =
@@ -2362,7 +2401,7 @@ void UInterchangeGenericMaterialPipeline::HandleLerpNode(const UInterchangeShade
 				ColorBExpression.Get<0>()->GetUniqueID(), ColorBExpression.Get<1>());
 		}
 	}
-	
+
 	// Factor
 	{
 		TTuple<UInterchangeMaterialExpressionFactoryNode*, FString> FactorExpression =
@@ -3186,7 +3225,7 @@ UInterchangeMaterialExpressionFactoryNode* UInterchangeGenericMaterialPipeline::
 		|| *ShaderType == Nodes::TextureObject::Name)
 	{
 		FString TextureUid;
-		
+
 		if (*ShaderType == Nodes::TextureSample::Name)
 		{
 			const bool bIsAParameter = UInterchangeShaderPortsAPI::HasParameter(ShaderNode, Nodes::TextureSample::Inputs::Texture);
@@ -3572,7 +3611,7 @@ UInterchangeMaterialInstanceFactoryNode* UInterchangeGenericMaterialPipeline::Cr
 #endif
 
 	VisitShaderGraphNode(ShaderGraphNode, MaterialInstanceFactoryNode);
-	
+
 	AddSpecularProfileToFactoryNode(ShaderGraphNode, MaterialInstanceFactoryNode, BaseNodeContainer);
 
 	return MaterialInstanceFactoryNode;
@@ -3586,7 +3625,7 @@ void UInterchangeGenericMaterialPipeline::CreateSpecularProfileFactoryNode(const
 
 	FactoryNode->InitializeNode(FactoryNodeUid, SpecularProfileNode->GetDisplayLabel(), EInterchangeNodeContainerType::FactoryData);
 	FactoryNode->SetEnabled(true);
-	
+
 	if (uint8 Format; SpecularProfileNode->GetCustomFormat(Format))
 	{
 		FactoryNode->SetCustomFormat(ESpecularProfileFormat{ Format });
@@ -3795,7 +3834,7 @@ bool UInterchangeGenericMaterialPipeline::HandleBxDFInput(const UInterchangeShad
 UInterchangeMaterialFunctionFactoryNode* UInterchangeGenericMaterialPipeline::CreateMaterialFunctionFactoryNode(const UInterchangeShaderGraphNode* ShaderGraphNode)
 {
 	UInterchangeMaterialFunctionFactoryNode* FactoryNode = Cast<UInterchangeMaterialFunctionFactoryNode>(CreateBaseMaterialFactoryNode(ShaderGraphNode, UInterchangeMaterialFunctionFactoryNode::StaticClass()));
-	
+
 	TArray<FString> InputNames;
 	UInterchangeShaderPortsAPI::GatherInputs(ShaderGraphNode, InputNames);
 
@@ -4246,7 +4285,7 @@ namespace UE::Interchange::Materials::HashUtils
 
 			// Just setup all the Parameters as overridable parameters. Do not include the values in the Hash
 			if (bIsAParameter)
-			{				
+			{
 				using namespace UE::Interchange::Materials::HashUtils;
 				const FString ParameterKey = UInterchangeShaderPortsAPI::MakeInputParameterKey(InputName);
 				const FString OverridableParameterNameKey = FInterchangeMaterialInstanceOverridesAPI::MakeOverrideParameterName(ShaderNode->GetDisplayLabel());
@@ -4333,7 +4372,7 @@ namespace UE::Interchange::Materials::HashUtils
 	}
 
 	void FDuplicateMaterialHelper::SetupOverridableVectorParameter(const UInterchangeShaderNode* ShaderNode, const FString& ParameterKey, const FString& OverridableParameterNameKey)
-	{		
+	{
 		FLinearColor InputValue;
 		if (ShaderNode->GetLinearColorAttribute(ParameterKey, InputValue))
 		{
@@ -4402,7 +4441,7 @@ namespace UE::Interchange::Materials::HashUtils
 	{
 		UInterchangeBaseNode::CopyStorageAttributes(AttributeStorageNode, FactoryNode, LeafInputAttributeKeys);
 	}
-	
+
 	UInterchangeBaseMaterialFactoryNode* FDuplicateMaterialHelper::CreateFactoryForDuplicateMaterials(const UInterchangeShaderGraphNode* ShaderGraphNode, bool bImportUnusedMaterial, bool bCreateMaterialInstanceForParent)
 	{
 		UInterchangeBaseMaterialFactoryNode* MaterialFactoryNode = nullptr;
@@ -4439,7 +4478,7 @@ namespace UE::Interchange::Materials::HashUtils
 		{
 			ParentMaterialFactory = *ParentMaterialFactoryEntry;
 		}
-		
+
 		ensure(ParentMaterialFactory);
 
 		if (!ParentMaterialFactory)
diff --git a/Engine/Plugins/Interchange/Runtime/Source/Pipelines/Public/InterchangeGenericMaterialPipeline.h b/Engine/Plugins/Interchange/Runtime/Source/Pipelines/Public/InterchangeGenericMaterialPipeline.h
--- a/Engine/Plugins/Interchange/Runtime/Source/Pipelines/Public/InterchangeGenericMaterialPipeline.h
+++ b/Engine/Plugins/Interchange/Runtime/Source/Pipelines/Public/InterchangeGenericMaterialPipeline.h
@@ -127,12 +127,6 @@ protected:
 	TArray<const UInterchangeSourceData*> SourceDatas;
 	
 private:
-	/** Material translated assets nodes */
-	TArray<UInterchangeShaderGraphNode*> MaterialNodes;
-	
-	/** Material factory assets nodes */
-	TArray<UInterchangeBaseMaterialFactoryNode*> MaterialFactoryNodes;
-
 	UInterchangeBaseMaterialFactoryNode* CreateBaseMaterialFactoryNode(const UInterchangeBaseNode* MaterialNode, TSubclassOf<UInterchangeBaseMaterialFactoryNode> NodeType, bool bAddMaterialInstanceSuffix = false);
 	UInterchangeMaterialFactoryNode* CreateMaterialFactoryNode(const UInterchangeShaderGraphNode* ShaderGraphNode);
 	UInterchangeMaterialFunctionFactoryNode* CreateMaterialFunctionFactoryNode(const UInterchangeShaderGraphNode* FunctionCallShaderNode);
