USD: Support stages referencing OpenVDB textures via Interchange.

Fix handling of reimport for SparseVolumeTexture assets, retaining AssetUserData/ImportData;
General refactor/cleanup of the SVT factory;
Fix user attributes not being moved over to the SVT factory nodes and some USD factory nodes;
Reset USDTranslator impl on every translation to make sure we don't retain state from a previous one;
Added some mutex / locking similar to the info cache within the VDB translator in case it has to open files anew when retrieving payloads;
Some renaming of previous attributes for consistency, having "Custom" in the name, and saying "Format" instead of "DataType";

diff --git a/Engine/Plugins/Importers/USDImporter/Source/USDSchemas/Private/USDVolVolumeTranslator.cpp b/Engine/Plugins/Importers/USDImporter/Source/USDSchemas/Private/USDVolVolumeTranslator.cpp
--- a/Engine/Plugins/Importers/USDImporter/Source/USDSchemas/Private/USDVolVolumeTranslator.cpp
+++ b/Engine/Plugins/Importers/USDImporter/Source/USDSchemas/Private/USDVolVolumeTranslator.cpp
@@ -46,347 +46,25 @@
 #include <pxr/usd/usdVol/volume.h>
 #include "USDIncludesEnd.h"
 
-static bool bRemoveDuplicates = true;
-static FAutoConsoleVariableRef CVarRemoveDuplicates(
-	TEXT("USD.Volume.RemoveDuplicateAnimatedFrames"),
-	bRemoveDuplicates,
-	TEXT(
-		"If this is true (default), the contents of a .VDB file are added only once to animated Sparse Volume Textures (SVT), even if the same file shows in multiple different time samples. If this is false, every OpenVDBAsset prim filePath time sample is parse as a new frame on the animated SVT."
-	)
-);
-
 namespace UE::UsdVolVolumeTranslator::Private
 {
 #if WITH_EDITOR
-	struct FSparseVolumeTextureInfo
-	{
-		FString SourceVDBFilePath;
-
-		TArray<FString> SourceOpenVDBAssetPrimPaths;
-
-		TArray<double> TimeSamplePathTimeCodes;
-		TArray<int32> TimeSamplePathIndices;
-		TArray<FString> TimeSamplePaths;	// First item may equal SourceVDBFilePath. May have duplicate entries
-
-		// Fields on the Volume prim that were connected to this .vdb file
-		TArray<FString> VolumeFieldNames;
+	static const int32 SparseVolumeTextureChannelCount = 8;
 
-		// e.g. {'velocity': {'AttributesA.R': 'X', 'AttributesA.G': 'Y', 'AttributesA.B': 'Z'}}
-		TMap<FString, TMap<FString, FString>> GridNameToChannelComponentMapping;
-		TOptional<ESparseVolumeAttributesFormat> AttributesAFormat;
-		TOptional<ESparseVolumeAttributesFormat> AttributesBFormat;
+	static_assert((int)UsdUtils::ESparseVolumeAttributesFormat::Unorm8 == (int)::ESparseVolumeAttributesFormat::Unorm8);
+	static_assert((int)UsdUtils::ESparseVolumeAttributesFormat::Float16 == (int)::ESparseVolumeAttributesFormat::Float16);
+	static_assert((int)UsdUtils::ESparseVolumeAttributesFormat::Float32 == (int)::ESparseVolumeAttributesFormat::Float32);
 
+	struct FSparseVolumeTextureInfo
+	{
+		const UsdUtils::FVolumePrimInfo* InnerInfo = nullptr;
 		USparseVolumeTexture* SparseVolumeTexture = nullptr;
 		FString PrefixedAssetHash;
 	};
 
-	static const int32 SparseVolumeTextureChannelCount = 8;
-
-	// Collect all the .vdb files that this prim wants to parse, and the desired fields/grids from them.
-	//
-	// In VDB terminology a "grid" is essentially a 3D texture, and can have formats like float, double3, half, etc.
-	// In USD the analogous term is "field", but essentially means the same thing. Possibly the terminology is abstracted
-	// to also fit the Field3D library, which we don't support it. Field/grid will be used interchangeably here.
-	//
-	// USD is very flexible and allows the user to reference specific grids from of each .vdb file. The syntax makes it
-	// difficult to find out at once all the grids we'll need to parse from each the .vdb files, so here we need to group them
-	// up first before deferring to the SparseVolumeTextureFactory.
-	//
-	// Note that USD allows a single Volume prim to reference grids from multiple .vdb files, and to also "timeSample" the
-	// file reference to allow for volume animations. This means that in UE a "Volume" prim corresponds to a single
-	// HeterogeneousVolumeActor, but which in turn can have any number of Sparse Volume Textures (one for each .vdb file referenced).
-	TMap<FString, FSparseVolumeTextureInfo> CollectSparseVolumeTextureInfoByFilePath(const pxr::UsdVolVolume& Volume)
-	{
-		TMap<FString, FSparseVolumeTextureInfo> FilePathHashToInfo;
-
-		FScopedUsdAllocs UsdAllocs;
-
-		pxr::UsdStageRefPtr Stage = Volume.GetPrim().GetStage();
-
-		const std::map<pxr::TfToken, pxr::SdfPath>& FieldMap = Volume.GetFieldPaths();
-		for (std::map<pxr::TfToken, pxr::SdfPath>::const_iterator Iter = FieldMap.cbegin(); Iter != FieldMap.end(); ++Iter)
-		{
-			// This field name is the name of the field for the Volume prim, which can be anything and differ from the
-			// grid name within the .vdb files
-			const pxr::TfToken& FieldName = Iter->first;
-			const pxr::SdfPath& AssetPrimPath = Iter->second;
-
-			pxr::UsdPrim OpenVDBPrim = Stage->GetPrimAtPath(AssetPrimPath);
-			pxr::UsdVolOpenVDBAsset OpenVDBPrimSchema = pxr::UsdVolOpenVDBAsset{OpenVDBPrim};
-			if (OpenVDBPrimSchema)
-			{
-				pxr::UsdAttribute FilePathAttr = OpenVDBPrimSchema.GetFilePathAttr();
-
-				FString ResolvedVDBPath = UsdUtils::GetResolvedAssetPath(FilePathAttr, pxr::UsdTimeCode::Default());
-
-				// Find timesampled paths, if any
-				TArray<double> TimeSamplePathTimeCodes;
-				TArray<int32> TimeSamplePathIndices;
-				TArray<FString> TimeSamplePaths;
-
-				TMap<FString, int32> PathToIndex;
-
-				std::vector<double> TimeSamples;
-				if (FilePathAttr.GetTimeSamples(&TimeSamples) && TimeSamples.size() > 0)
-				{
-					UE::FSdfLayerOffset CombinedOffset = UsdUtils::GetPrimToStageOffset(UE::FUsdPrim{OpenVDBPrim});
-
-					TimeSamplePathTimeCodes.Reserve(TimeSamples.size());
-					TimeSamplePaths.Reserve(TimeSamples.size());
-					TimeSamplePathIndices.Reserve(TimeSamples.size());
-					for (double TimeSample : TimeSamples)
-					{
-						// We always want to store on the AssetUserData (which is where this stuff will end up in)
-						// the time codes in the layer where the actual OpenVDBPrim is authored. If that layer is referenced
-						// by a parent layer through an offset and scale, TimeSample will contain that offset and scale here,
-						// which we need to undo
-						double LayerLocalTimeCode = (TimeSample - CombinedOffset.Offset) / CombinedOffset.Scale;
-						TimeSamplePathTimeCodes.Add(LayerLocalTimeCode);
-
-						FString ResolvedFramePath = UsdUtils::GetResolvedAssetPath(FilePathAttr, TimeSample);
-
-						// If we had no default time sample to act as the "main file", take the first frame
-						if (ResolvedVDBPath.IsEmpty())
-						{
-							ResolvedVDBPath = ResolvedFramePath;
-						}
-
-						if (bRemoveDuplicates)
-						{
-							if (int32* FoundIndex = PathToIndex.Find(ResolvedFramePath))
-							{
-								TimeSamplePathIndices.Add(*FoundIndex);
-							}
-							else
-							{
-								TimeSamplePaths.Add(ResolvedFramePath);
-
-								const int32 NewIndex = TimeSamplePaths.Num() - 1;
-								PathToIndex.Add(ResolvedFramePath, NewIndex);
-								TimeSamplePathIndices.Add(NewIndex);
-							}
-						}
-						else
-						{
-							TimeSamplePaths.Add(ResolvedFramePath);
-
-							const int32 NewIndex = TimeSamplePaths.Num() - 1;
-							TimeSamplePathIndices.Add(NewIndex);
-						}
-					}
-				}
-
-				// Hash all the relevant file paths here: The collection of file paths to parse determines the SVT, and
-				// we want one FSparseVolumeTextureInfo per SVT
-				FString FilePathHashString;
-				{
-					FSHA1 SHA1;
-					SHA1.UpdateWithString(*ResolvedVDBPath, ResolvedVDBPath.Len());
-					for (const FString& TimeSamplePath : TimeSamplePaths)
-					{
-						SHA1.UpdateWithString(*TimeSamplePath, TimeSamplePath.Len());
-					}
-					SHA1.Final();
-
-					FSHAHash FilePathHash;
-					SHA1.GetHash(&FilePathHash.Hash[0]);
-
-					FilePathHashString = FilePathHash.ToString();
-				}
-
-				if (!ResolvedVDBPath.IsEmpty())
-				{
-					FSparseVolumeTextureInfo& SparseVolumeTextureInfo = FilePathHashToInfo.FindOrAdd(FilePathHashString);
-					SparseVolumeTextureInfo.SourceOpenVDBAssetPrimPaths.AddUnique(UsdToUnreal::ConvertPath(AssetPrimPath));
-					SparseVolumeTextureInfo.SourceVDBFilePath = ResolvedVDBPath;
-					SparseVolumeTextureInfo.TimeSamplePathTimeCodes = MoveTemp(TimeSamplePathTimeCodes);
-					SparseVolumeTextureInfo.TimeSamplePathIndices = MoveTemp(TimeSamplePathIndices);
-					SparseVolumeTextureInfo.TimeSamplePaths = MoveTemp(TimeSamplePaths);
-
-					FString FieldNameStr = UsdToUnreal::ConvertToken(FieldName);
-					SparseVolumeTextureInfo.VolumeFieldNames.AddUnique(FieldNameStr);
-
-					pxr::TfToken GridName;
-					pxr::UsdAttribute Attr = OpenVDBPrimSchema.GetFieldNameAttr();
-					if (Attr && Attr.Get<pxr::TfToken>(&GridName))
-					{
-						FString GridNameStr = UsdToUnreal::ConvertToken(GridName);
-
-						// Note we want this to add an entry to SparseVolumeTexture.GridNameToChannelNames even if we won't
-						// find the schema on the prim, as we'll use these entries to make sure the generated Sparse Volume
-						// Texture contains theses desired fields
-						TMap<FString, FString>& ChannelToComponent = SparseVolumeTextureInfo.GridNameToChannelComponentMapping.FindOrAdd(GridNameStr);
-
-						if (UsdUtils::PrimHasSchema(OpenVDBPrim.GetPrim(), UnrealIdentifiers::SparseVolumeTextureAPI))
-						{
-							// Parse desired data types for AttributesA and AttributesB channels
-							TFunction<void(pxr::TfToken, TOptional<ESparseVolumeAttributesFormat>&)> HandleAttribute =
-								[&OpenVDBPrim,
-								 &AssetPrimPath,
-								 &ResolvedVDBPath](pxr::TfToken AttrName, TOptional<ESparseVolumeAttributesFormat>& AttributeFormat)
-							{
-								using FormatMapType = std::unordered_map<pxr::TfToken, ESparseVolumeAttributesFormat, pxr::TfHash>;
-								using InverseFormatMapType = std::unordered_map<ESparseVolumeAttributesFormat, FString>;
-
-								// clang-format off
-								const static FormatMapType FormatMap = {
-									{pxr::TfToken("unorm8"),  ESparseVolumeAttributesFormat::Unorm8},
-									{pxr::TfToken("float16"), ESparseVolumeAttributesFormat::Float16},
-									{pxr::TfToken("float32"), ESparseVolumeAttributesFormat::Float32},
-								};
-								const static InverseFormatMapType InverseFormatMap = {
-									{ESparseVolumeAttributesFormat::Unorm8,  TEXT("unorm8")},
-									{ESparseVolumeAttributesFormat::Float16, TEXT("float16")},
-									{ESparseVolumeAttributesFormat::Float32, TEXT("float32")},
-								};
-								// clang-format on
-
-								pxr::TfToken DataType;
-								pxr::UsdAttribute AttrA = OpenVDBPrim.GetAttribute(AttrName);
-								if (AttrA && AttrA.Get<pxr::TfToken>(&DataType))
-								{
-									FormatMapType::const_iterator Iter = FormatMap.find(DataType);
-									if (Iter != FormatMap.end())
-									{
-										ESparseVolumeAttributesFormat TargetFormat = Iter->second;
-
-										// Check in case multiple OpenVDBAsset prims want different values for the data type
-										const bool bIsSet = AttributeFormat.IsSet();
-										if (bIsSet && AttributeFormat.GetValue() != TargetFormat)
-										{
-											const FString* ExistingFormat = nullptr;
-											InverseFormatMapType::const_iterator ExistingIter = InverseFormatMap.find(AttributeFormat.GetValue());
-											if (ExistingIter != InverseFormatMap.end())
-											{
-												ExistingFormat = &ExistingIter->second;
-											}
-
-											UE_LOG(
-												LogUsd,
-												Warning,
-												TEXT(
-													"OpenVDBAsset prims disagree on the attribute channel format for the Sparse Volume Texture generated for VDB file '%s' (encountered '%s' and '%s'). If there are multiple opinions for the attribute channel formats from different OpenVDBAsset prims, they must all agree!"
-												),
-												*ResolvedVDBPath,
-												*UsdToUnreal::ConvertToken(DataType),
-												ExistingFormat ? **ExistingFormat : TEXT("unknown")
-											);
-										}
-										else if (!bIsSet)
-										{
-											AttributeFormat = TargetFormat;
-										}
-									}
-									else
-									{
-										UE_LOG(
-											LogUsd,
-											Warning,
-											TEXT(
-												"Invalid Sparse Volume Texture attribute channel format '%s'. Available formats: 'unorm8', 'float16' and 'float32'."
-											),
-											*UsdToUnreal::ConvertToken(DataType),
-											*UsdToUnreal::ConvertPath(AssetPrimPath)
-										);
-									}
-								}
-							};
-							HandleAttribute(UnrealIdentifiers::UnrealSVTAttributesADataType, SparseVolumeTextureInfo.AttributesAFormat);
-							HandleAttribute(UnrealIdentifiers::UnrealSVTAttributesBDataType, SparseVolumeTextureInfo.AttributesBFormat);
-
-							// Parse desired channel assignment
-							pxr::VtArray<pxr::TfToken> DesiredChannels;
-							pxr::VtArray<pxr::TfToken> DesiredComponents;
-							pxr::UsdAttribute ComponentsAttr = OpenVDBPrim.GetAttribute(UnrealIdentifiers::UnrealSVTMappedGridComponents);
-							pxr::UsdAttribute ChannelsAttr = OpenVDBPrim.GetAttribute(UnrealIdentifiers::UnrealSVTMappedAttributeChannels);
-							if (ChannelsAttr &&														 //
-								ComponentsAttr &&													 //
-								ChannelsAttr.Get<pxr::VtArray<pxr::TfToken>>(&DesiredChannels) &&	 //
-								ComponentsAttr.Get<pxr::VtArray<pxr::TfToken>>(&DesiredComponents))
-							{
-								// These must always match of course
-								if (DesiredChannels.size() == DesiredComponents.size())
-								{
-									// If we have more than one OpenVDBAsset prim reading from the same VDB file, the declared component to
-									// channel mappings must be compatible
-									for (uint32 Index = 0; Index < DesiredChannels.size(); ++Index)
-									{
-										FString Channel = UsdToUnreal::ConvertString(DesiredChannels[Index]);
-										FString Component = UsdToUnreal::ConvertString(DesiredComponents[Index]);
-
-										if (const FString* ExistingComponentMapping = ChannelToComponent.Find(Channel))
-										{
-											if (Component != *ExistingComponentMapping)
-											{
-												UE_LOG(
-													LogUsd,
-													Warning,
-													TEXT(
-														"Found multiple OpenVDBAsset prims (including '%s') targetting the same grid '%s', but with with conflicting grid component to Sparse Volume Texture attribute channel mapping (for example, both components '%s' and '%s' are mapped to the same channel '%s', which is not allowed)"
-													),
-													*UsdToUnreal::ConvertPath(AssetPrimPath),
-													*GridNameStr,
-													**ExistingComponentMapping,
-													*Component,
-													*Channel
-												);
-											}
-										}
-										else
-										{
-											ChannelToComponent.Add(Channel, Component);
-										}
-									}
-								}
-								else
-								{
-									UE_LOG(
-										LogUsd,
-										Warning,
-										TEXT(
-											"Failed to parse custom component to attribute mapping from OpenVDBAsset prim '%s': The '%s' and '%s' attributes should have the same number of entries, but the former has %u entries while the latter has %u"
-										),
-										*UsdToUnreal::ConvertPath(AssetPrimPath),
-										*UsdToUnreal::ConvertToken(UnrealIdentifiers::UnrealSVTMappedGridComponents),
-										*UsdToUnreal::ConvertToken(UnrealIdentifiers::UnrealSVTMappedAttributeChannels),
-										DesiredComponents.size(),
-										DesiredChannels.size()
-									);
-								}
-							}
-						}
-					}
-				}
-				else
-				{
-					UE_LOG(
-						LogUsd,
-						Warning,
-						TEXT("Failed to find the VDB file '%s' referenced by OpenVDBAsset prim at path '%s'"),
-						*ResolvedVDBPath,
-						*UsdToUnreal::ConvertPath(AssetPrimPath)
-					);
-				}
-			}
-			else
-			{
-				UE_LOG(
-					LogUsd,
-					Warning,
-					TEXT("Failed to find an OpenVDBAsset prim at path '%s' for field '%s' of prim '%s'"),
-					*UsdToUnreal::ConvertPath(AssetPrimPath),
-					*UsdToUnreal::ConvertToken(FieldName),
-					*UsdToUnreal::ConvertPath(Volume.GetPrim().GetPath())
-				);
-			}
-		}
-
-		return FilePathHashToInfo;
-	}
-
 	// Here we must stash into InOutPreviewData.ImportOptions the desired channel mapping for this SVT given all the GridNameToChannelNames
 	// mappings we pulled out of the prims if they had any instances of our SparseVolumeTextureAPI schema
-	void SetVDBImportOptions(const FSparseVolumeTextureInfo& ParsedTexture, FOpenVDBPreviewData& InOutPreviewData)
+	void SetVDBImportOptions(const UsdUtils::FVolumePrimInfo& ParsedTexture, FOpenVDBPreviewData& InOutPreviewData)
 	{
 		// Tweak the collected filenames for other frames: The OpenVDB importer will scan for similar filenames
 		// in the same folder as the main file, but through USD we expect the user to manually pick file paths
@@ -398,11 +76,15 @@ namespace UE::UsdVolVolumeTranslator::Private
 		check(InOutPreviewData.DefaultImportOptions.Attributes.Num() == 2);
 		if (ParsedTexture.AttributesAFormat.IsSet())
 		{
-			InOutPreviewData.DefaultImportOptions.Attributes[0].Format = ParsedTexture.AttributesAFormat.GetValue();
+			InOutPreviewData.DefaultImportOptions.Attributes[0].Format = static_cast<ESparseVolumeAttributesFormat>(
+				ParsedTexture.AttributesAFormat.GetValue()
+			);
 		}
 		if (ParsedTexture.AttributesBFormat.IsSet())
 		{
-			InOutPreviewData.DefaultImportOptions.Attributes[1].Format = ParsedTexture.AttributesBFormat.GetValue();
+			InOutPreviewData.DefaultImportOptions.Attributes[1].Format = static_cast<ESparseVolumeAttributesFormat>(
+				ParsedTexture.AttributesBFormat.GetValue()
+			);
 		}
 
 		static const TMap<FString, int32> ChannelIndexMap = {
@@ -734,95 +416,49 @@ namespace UE::UsdVolVolumeTranslator::Private
 		FieldNameToInfo.Reserve(FilePathHashToTextureInfo.Num());
 		for (const TPair<FString, FSparseVolumeTextureInfo>& FilePathToInfo : FilePathHashToTextureInfo)
 		{
-			for (const FString& FieldName : FilePathToInfo.Value.VolumeFieldNames)
+			for (const FString& FieldName : FilePathToInfo.Value.InnerInfo->VolumeFieldNames)
 			{
-				// We won't be modifying FSparseVolumeTextureInfo in this function anyway, it's just so that we can
+				// We won't be modifying UsdUtils::FVolumePrimInfo in this function anyway, it's just so that we can
 				// return non-const pointers
 				FieldNameToInfo.Add(FieldName, const_cast<FSparseVolumeTextureInfo*>(&FilePathToInfo.Value));
 			}
 		}
 
 		// Collect material parameter assignments manually specified via the custom schema, if any
-		bool bHadManualAssignment = false;
-		if (UsdUtils::PrimHasSchema(VolumePrim, UnrealIdentifiers::SparseVolumeTextureAPI))
+		TMultiMap<FString, FString> MaterialParameterToFieldName = UsdUtils::GetVolumeMaterialParameterToFieldNameMap(VolumePrim);
+		for (const TPair<FString, FString>& Pair : MaterialParameterToFieldName)
 		{
-			pxr::VtArray<pxr::TfToken> MappedFields;
-			pxr::VtArray<pxr::TfToken> MappedParameters;
-			pxr::UsdAttribute FieldsAttr = VolumePrim.GetAttribute(UnrealIdentifiers::UnrealSVTMappedFields);
-			pxr::UsdAttribute ParametersAttr = VolumePrim.GetAttribute(UnrealIdentifiers::UnrealSVTMappedMaterialParameters);
-			if (FieldsAttr &&													//
-				ParametersAttr &&												//
-				FieldsAttr.Get<pxr::VtArray<pxr::TfToken>>(&MappedFields) &&	//
-				ParametersAttr.Get<pxr::VtArray<pxr::TfToken>>(&MappedParameters))
+			const FString& MaterialParameter = Pair.Key;
+			const FString& FieldName = Pair.Value;
+
+			if (FSparseVolumeTextureInfo* FoundParsedTexture = FieldNameToInfo.FindRef(FieldName))
 			{
-				// These must always match of course
-				if (MappedFields.size() == MappedParameters.size())
+				if (const FSparseVolumeTextureInfo* JobAtParam = ResultParameterToInfo.FindRef(MaterialParameter))
 				{
-					for (uint32 Index = 0; Index < MappedFields.size(); ++Index)
+					if (JobAtParam->SparseVolumeTexture != FoundParsedTexture->SparseVolumeTexture)
 					{
-						FString FieldName = UsdToUnreal::ConvertString(MappedFields[Index]);
-						FString MaterialParameter = UsdToUnreal::ConvertString(MappedParameters[Index]);
-						bHadManualAssignment = true;
-
-						if (FSparseVolumeTextureInfo* FoundParsedTexture = FieldNameToInfo.FindRef(FieldName))
-						{
-							if (const FSparseVolumeTextureInfo* JobAtParam = ResultParameterToInfo.FindRef(MaterialParameter))
-							{
-								if (JobAtParam->SparseVolumeTexture != FoundParsedTexture->SparseVolumeTexture)
-								{
-									UE_LOG(
-										LogUsd,
-										Warning,
-										TEXT(
-											"Trying to assign different Sparse Volume Textures to the same material parameter '%s' on the material instantiated for Volume prim '%s' and field name '%s'! Only a single material can be assigned to a material parameter at a time."
-										),
-										*MaterialParameter,
-										*UsdToUnreal::ConvertPath(VolumePrim.GetPrimPath()),
-										*FieldName
-									);
-								}
-							}
-							else
-							{
-								// We won't be modifying anything in here, it's just so that we can return non-const pointers
-								ResultParameterToInfo.Add(MaterialParameter, FoundParsedTexture);
-							}
-						}
+						UE_LOG(
+							LogUsd,
+							Warning,
+							TEXT(
+								"Trying to assign different Sparse Volume Textures to the same material parameter '%s' on the material instantiated for Volume prim '%s' and field name '%s'! Only a single texture can be assigned to a material parameter at a time."
+							),
+							*MaterialParameter,
+							*UsdToUnreal::ConvertPath(VolumePrim.GetPrimPath()),
+							*FieldName
+						);
 					}
 				}
 				else
 				{
-					UE_LOG(
-						LogUsd,
-						Warning,
-						TEXT(
-							"Failed to parse custom parsed texture to material parameter mapping from volume prim '%s': The '%s' and '%s' attributes should have the same number of entries, but the former has %u entries while the latter has %u"
-						),
-						*UsdToUnreal::ConvertPath(VolumePrim.GetPrimPath()),
-						*UsdToUnreal::ConvertToken(UnrealIdentifiers::UnrealSVTMappedFields),
-						*UsdToUnreal::ConvertToken(UnrealIdentifiers::UnrealSVTMappedMaterialParameters),
-						MappedFields.size(),
-						MappedParameters.size()
-					);
+					// We won't be modifying anything in here, it's just so that we can return non-const pointers
+					ResultParameterToInfo.Add(MaterialParameter, FoundParsedTexture);
 				}
 			}
 		}
 
 		// Collect available parameter names on this material instance
-		TSet<FString> SparseVolumeTextureParameterNames;
-		{
-			TMap<FMaterialParameterInfo, FMaterialParameterMetadata> SparseVolumeTextureParameters;
-			ReferenceMaterial->GetAllParametersOfType(EMaterialParameterType::SparseVolumeTexture, SparseVolumeTextureParameters);
-
-			SparseVolumeTextureParameterNames.Reserve(SparseVolumeTextureParameters.Num());
-			for (const TPair<FMaterialParameterInfo, FMaterialParameterMetadata>& ParameterPair : SparseVolumeTextureParameters)
-			{
-				const FMaterialParameterInfo& ParameterInfo = ParameterPair.Key;
-				FString ParameterName = ParameterInfo.Name.ToString();
-
-				SparseVolumeTextureParameterNames.Add(ParameterName);
-			}
-		}
+		TArray<FString> SparseVolumeTextureParameterNames = UsdUtils::GetSparseVolumeTextureParameterNames(ReferenceMaterial);
 
 		// Validate that all parameters exist on the material, or else emit a warning
 		for (const TPair<FString, FSparseVolumeTextureInfo*>& AssignmentPair : ResultParameterToInfo)
@@ -840,6 +476,8 @@ namespace UE::UsdVolVolumeTranslator::Private
 			}
 		}
 
+		bool bHadManualAssignment = MaterialParameterToFieldName.Num() > 0;
+
 		// No manual material parameter assignment specified via custom schema: First let's assume that the field names match material parameter names
 		bool bHadParameterNameMatch = false;
 		if (!bHadManualAssignment)
@@ -889,13 +527,12 @@ namespace UE::UsdVolVolumeTranslator::Private
 			FieldNameToInfo.GetKeys(SortedFieldNames);
 			SortedFieldNames.Sort();
 
-			TArray<FString> SortedParameterNames = SparseVolumeTextureParameterNames.Array();
-			SortedParameterNames.Sort();
+			SparseVolumeTextureParameterNames.Sort();
 
-			for (int32 Index = 0; Index < SortedFieldNames.Num() && Index < SortedParameterNames.Num(); ++Index)
+			for (int32 Index = 0; Index < SortedFieldNames.Num() && Index < SparseVolumeTextureParameterNames.Num(); ++Index)
 			{
 				const FString& FieldName = SortedFieldNames[Index];
-				const FString& ParameterName = SortedParameterNames[Index];
+				const FString& ParameterName = SparseVolumeTextureParameterNames[Index];
 
 				FSparseVolumeTextureInfo* FoundTextureJob = FieldNameToInfo.FindRef(FieldName);
 				if (ensure(FoundTextureJob))
@@ -964,14 +601,23 @@ void FUsdVolVolumeTranslator::CreateAssets()
 
 	FString VolumePrimHashPrefix = UsdUtils::GetAssetHashPrefix(GetPrim(), Context->bShareAssetsForIdenticalPrims);
 
-	// Create the info structs from the requested files
-	TMap<FString, FSparseVolumeTextureInfo> FilePathHashToTextureInfo = CollectSparseVolumeTextureInfoByFilePath(Volume);
+	// Collect info from requested files
+	const TMap<FString, UsdUtils::FVolumePrimInfo> FilePathHashToVolumeInfo = UsdUtils::GetVolumeInfoByFilePathHash(VolumePrim);
+
+	// Move the info into another struct so that we can tack on the generated SVT and prefixed asset hash
+	TMap<FString, FSparseVolumeTextureInfo> FilePathHashToSparseVolumeInfo;
+	FilePathHashToSparseVolumeInfo.Reserve(FilePathHashToVolumeInfo.Num());
+	for (const TPair<FString, UsdUtils::FVolumePrimInfo>& Pair : FilePathHashToVolumeInfo)
+	{
+		FSparseVolumeTextureInfo& SparseVolumeTextureInfo = FilePathHashToSparseVolumeInfo.FindOrAdd(Pair.Key);
+		SparseVolumeTextureInfo.InnerInfo = &Pair.Value;
+	}
 
 	// Create SVT assets from the info structs
-	for (TPair<FString, FSparseVolumeTextureInfo>& FilePathHashToInfo : FilePathHashToTextureInfo)
+	for (TPair<FString, FSparseVolumeTextureInfo>& FilePathHashToInfo : FilePathHashToSparseVolumeInfo)
 	{
-		const FString& VDBFilePath = FilePathHashToInfo.Value.SourceVDBFilePath;
-		FSparseVolumeTextureInfo& ParsedTexture = FilePathHashToInfo.Value;
+		const FString& VDBFilePath = FilePathHashToInfo.Value.InnerInfo->SourceVDBFilePath;
+		FSparseVolumeTextureInfo& SparseVolumeTextureInfo = FilePathHashToInfo.Value;
 
 		if (!FPaths::FileExists(VDBFilePath))
 		{
@@ -997,7 +643,7 @@ void FUsdVolVolumeTranslator::CreateAssets()
 		// attributes, in a similar way to how blendshapes are mapped. We'll check for those in SetVDBImportOptions
 		TStrongObjectPtr<UOpenVDBImportOptionsObject> ImportOptions{NewObject<UOpenVDBImportOptionsObject>()};
 		LoadOpenVDBPreviewData(VDBFilePath, &ImportOptions->PreviewData);
-		SetVDBImportOptions(ParsedTexture, ImportOptions->PreviewData);
+		SetVDBImportOptions(*SparseVolumeTextureInfo.InnerInfo, ImportOptions->PreviewData);
 
 		// Collect a hash for this VDB asset
 		FSHAHash VDBAndAssignmentHash;
@@ -1008,14 +654,14 @@ void FUsdVolVolumeTranslator::CreateAssets()
 			SHA1.Final();
 			SHA1.GetHash(&VDBAndAssignmentHash.Hash[0]);
 		}
-		ParsedTexture.PrefixedAssetHash = VolumePrimHashPrefix + VDBAndAssignmentHash.ToString();
+		SparseVolumeTextureInfo.PrefixedAssetHash = VolumePrimHashPrefix + VDBAndAssignmentHash.ToString();
 
 		// File path instead of prim path in case we have  multiple .vdb files in the same Volume prim
 		const FString& DesiredName = FPaths::GetBaseFilename(VDBFilePath);
 
 		bool bCreatedNew = false;
 		USparseVolumeTexture* SparseVolumeTexture = Context->UsdAssetCache->GetOrCreateCustomCachedAsset<USparseVolumeTexture>(
-			ParsedTexture.PrefixedAssetHash,
+			SparseVolumeTextureInfo.PrefixedAssetHash,
 			DesiredName,
 			Context->ObjectFlags,
 			[&VDBFilePath, &ImportOptions](UPackage* Outer, FName SanitizedName, EObjectFlags FlagsToUse)
@@ -1078,10 +724,10 @@ void FUsdVolVolumeTranslator::CreateAssets()
 					UUsdSparseVolumeTextureAssetUserData>(SparseVolumeTexture))
 			{
 				UserData->PrimPaths.AddUnique(VolumePrimPathString);
-				UserData->SourceOpenVDBAssetPrimPaths = ParsedTexture.SourceOpenVDBAssetPrimPaths;
-				UserData->TimeSamplePaths = ParsedTexture.TimeSamplePaths;
-				UserData->TimeSamplePathIndices = ParsedTexture.TimeSamplePathIndices;
-				UserData->TimeSamplePathTimeCodes = ParsedTexture.TimeSamplePathTimeCodes;
+				UserData->SourceOpenVDBAssetPrimPaths = SparseVolumeTextureInfo.InnerInfo->SourceOpenVDBAssetPrimPaths;
+				UserData->TimeSamplePaths = SparseVolumeTextureInfo.InnerInfo->TimeSamplePaths;
+				UserData->TimeSamplePathIndices = SparseVolumeTextureInfo.InnerInfo->TimeSamplePathIndices;
+				UserData->TimeSamplePathTimeCodes = SparseVolumeTextureInfo.InnerInfo->TimeSamplePathTimeCodes;
 
 				if (Context->MetadataOptions.bCollectMetadata)
 				{
@@ -1100,7 +746,7 @@ void FUsdVolVolumeTranslator::CreateAssets()
 			}
 		}
 
-		ParsedTexture.SparseVolumeTexture = SparseVolumeTexture;
+		SparseVolumeTextureInfo.SparseVolumeTexture = SparseVolumeTexture;
 	}
 
 	// Create SVT materials
@@ -1181,7 +827,7 @@ void FUsdVolVolumeTranslator::CreateAssets()
 	TMap<FString, FSparseVolumeTextureInfo*> MaterialParameterToTexture = CollectMaterialParameterTextureAssignment(
 		VolumePrim,
 		Material,
-		FilePathHashToTextureInfo
+		FilePathHashToSparseVolumeInfo
 	);
 
 	FSHAHash MaterialHash;
diff --git a/Engine/Plugins/Interchange/Runtime/Source/Common/Public/Usd/InterchangeUsdDefinitions.h b/Engine/Plugins/Interchange/Runtime/Source/Common/Public/Usd/InterchangeUsdDefinitions.h
--- a/Engine/Plugins/Interchange/Runtime/Source/Common/Public/Usd/InterchangeUsdDefinitions.h
+++ b/Engine/Plugins/Interchange/Runtime/Source/Common/Public/Usd/InterchangeUsdDefinitions.h
@@ -32,6 +32,27 @@ namespace UE::Interchange::USD
 	// attributes.
 	const FString NeedsVTParentIdentifier = TEXT("USD_MI_NeedsVTParent");
 
+	// Used for volumetric material parameters, whenever we assign an SVT to a material as a fallback due
+	// to it's field name only
+	const FString VolumeFieldNameMaterialParameterPrefix = TEXT("USD_FieldName_");
+
+	// Custom attribute keys used to describe SVT info extracted from the USD custom schemas.
+	// We add these to the Volume nodes that the USD translator emits.
+	namespace SparseVolumeTexture
+	{
+		const FString AttributesAFormat = TEXT("USD_AttributesA_Format");
+		const FString AttributesBFormat = TEXT("USD_AttributesB_Format");
+
+		const FString AttributesAChannelR = TEXT("USD_AttributesA_R");
+		const FString AttributesAChannelG = TEXT("USD_AttributesA_G");
+		const FString AttributesAChannelB = TEXT("USD_AttributesA_B");
+		const FString AttributesAChannelA = TEXT("USD_AttributesA_A");
+		const FString AttributesBChannelR = TEXT("USD_AttributesB_R");
+		const FString AttributesBChannelG = TEXT("USD_AttributesB_G");
+		const FString AttributesBChannelB = TEXT("USD_AttributesB_B");
+		const FString AttributesBChannelA = TEXT("USD_AttributesB_A");
+	}
+
 	namespace Primvar
 	{
 		const FString Number = TEXT("USD_PrimvarNumber");
@@ -41,4 +62,4 @@ namespace UE::Interchange::USD
 		// Attribute that informs about the UID of a ShaderNode of type TextureSample
 		const FString ShaderNodeTextureSample = TEXT("USD_ShaderNodeTextureSample");
 	}
-}
+}	 // namespace UE::Interchange::USD
diff --git a/Engine/Plugins/Interchange/Runtime/Source/Common/Public/Volume/InterchangeVolumeDefinitions.h b/Engine/Plugins/Interchange/Runtime/Source/Common/Public/Volume/InterchangeVolumeDefinitions.h
--- a/Engine/Plugins/Interchange/Runtime/Source/Common/Public/Volume/InterchangeVolumeDefinitions.h
+++ b/Engine/Plugins/Interchange/Runtime/Source/Common/Public/Volume/InterchangeVolumeDefinitions.h
@@ -20,6 +20,8 @@ namespace UE::Interchange::Volume
 	const FString DensityGridName = TEXT("density");
 	const FString GridNameAndComponentIndexSeparator = TEXT("_");
 
+	const FString VolumetricMaterial = TEXT("Volumetric_Material");
+
 	// These structs are direct copies of the ones from "OpenVDBImportOptions.h".
 	// We do this instead of using them directly as that belongs to an editor-only module, and we can't
 	// use #if WITH_EDITOR here as the FVolumePayloadKey shows up on the payload interface, which would
diff --git a/Engine/Plugins/Interchange/Runtime/Source/FactoryNodes/Private/InterchangeHeterogeneousVolumeActorFactoryNode.cpp b/Engine/Plugins/Interchange/Runtime/Source/FactoryNodes/Private/InterchangeHeterogeneousVolumeActorFactoryNode.cpp
new file mode 100644
--- /dev/null
+++ b/Engine/Plugins/Interchange/Runtime/Source/FactoryNodes/Private/InterchangeHeterogeneousVolumeActorFactoryNode.cpp
@@ -0,0 +1,24 @@
+// Copyright Epic Games, Inc. All Rights Reserved.
+
+#include "InterchangeHeterogeneousVolumeActorFactoryNode.h"
+
+#include "Components/HeterogeneousVolumeComponent.h"
+
+UClass* UInterchangeHeterogeneousVolumeActorFactoryNode::GetObjectClass() const
+{
+#if WITH_ENGINE
+	return AHeterogeneousVolume::StaticClass();
+#else
+	return nullptr;
+#endif
+}
+
+bool UInterchangeHeterogeneousVolumeActorFactoryNode::GetCustomVolumetricMaterialUid(FString& AttributeValue) const
+{
+	IMPLEMENT_NODE_ATTRIBUTE_GETTER(MaterialDependency, FString)
+}
+
+bool UInterchangeHeterogeneousVolumeActorFactoryNode::SetCustomVolumetricMaterialUid(const FString& AttributeValue)
+{
+	IMPLEMENT_NODE_ATTRIBUTE_SETTER_NODELEGATE(MaterialDependency, FString)
+}
diff --git a/Engine/Plugins/Interchange/Runtime/Source/FactoryNodes/Private/InterchangeSparseVolumeTextureFactoryNode.cpp b/Engine/Plugins/Interchange/Runtime/Source/FactoryNodes/Private/InterchangeSparseVolumeTextureFactoryNode.cpp
--- a/Engine/Plugins/Interchange/Runtime/Source/FactoryNodes/Private/InterchangeSparseVolumeTextureFactoryNode.cpp
+++ b/Engine/Plugins/Interchange/Runtime/Source/FactoryNodes/Private/InterchangeSparseVolumeTextureFactoryNode.cpp
@@ -23,22 +23,22 @@ UClass* UInterchangeSparseVolumeTextureFactoryNode::GetObjectClass() const
 #endif
 }
 
-bool UInterchangeSparseVolumeTextureFactoryNode::GetCustomAttributesADataType(EInterchangeSparseVolumeTextureFormat& AttributeValue) const
+bool UInterchangeSparseVolumeTextureFactoryNode::GetCustomAttributesAFormat(EInterchangeSparseVolumeTextureFormat& AttributeValue) const
 {
 	IMPLEMENT_NODE_ATTRIBUTE_GETTER(AttributesADataType, EInterchangeSparseVolumeTextureFormat);
 }
 
-bool UInterchangeSparseVolumeTextureFactoryNode::SetCustomAttributesADataType(EInterchangeSparseVolumeTextureFormat AttributeValue)
+bool UInterchangeSparseVolumeTextureFactoryNode::SetCustomAttributesAFormat(EInterchangeSparseVolumeTextureFormat AttributeValue)
 {
 	IMPLEMENT_NODE_ATTRIBUTE_SETTER_NODELEGATE(AttributesADataType, EInterchangeSparseVolumeTextureFormat);
 }
 
-bool UInterchangeSparseVolumeTextureFactoryNode::GetCustomAttributesBDataType(EInterchangeSparseVolumeTextureFormat& AttributeValue) const
+bool UInterchangeSparseVolumeTextureFactoryNode::GetCustomAttributesBFormat(EInterchangeSparseVolumeTextureFormat& AttributeValue) const
 {
 	IMPLEMENT_NODE_ATTRIBUTE_GETTER(AttributesBDataType, EInterchangeSparseVolumeTextureFormat);
 }
 
-bool UInterchangeSparseVolumeTextureFactoryNode::SetCustomAttributesBDataType(EInterchangeSparseVolumeTextureFormat AttributeValue)
+bool UInterchangeSparseVolumeTextureFactoryNode::SetCustomAttributesBFormat(EInterchangeSparseVolumeTextureFormat AttributeValue)
 {
 	IMPLEMENT_NODE_ATTRIBUTE_SETTER_NODELEGATE(AttributesBDataType, EInterchangeSparseVolumeTextureFormat);
 }
@@ -122,3 +122,13 @@ bool UInterchangeSparseVolumeTextureFactoryNode::SetCustomAttributesBChannelW(FS
 {
 	IMPLEMENT_NODE_ATTRIBUTE_SETTER_NODELEGATE(AttributesBChannelW, FString);
 }
+
+bool UInterchangeSparseVolumeTextureFactoryNode::GetCustomAnimationID(FString& AttributeValue) const
+{
+	IMPLEMENT_NODE_ATTRIBUTE_GETTER(AnimationID, FString);
+}
+
+bool UInterchangeSparseVolumeTextureFactoryNode::SetCustomAnimationID(const FString& AttributeValue)
+{
+	IMPLEMENT_NODE_ATTRIBUTE_SETTER_NODELEGATE(AnimationID, FString);
+}
diff --git a/Engine/Plugins/Interchange/Runtime/Source/FactoryNodes/Public/InterchangeHeterogeneousVolumeActorFactoryNode.h b/Engine/Plugins/Interchange/Runtime/Source/FactoryNodes/Public/InterchangeHeterogeneousVolumeActorFactoryNode.h
new file mode 100644
--- /dev/null
+++ b/Engine/Plugins/Interchange/Runtime/Source/FactoryNodes/Public/InterchangeHeterogeneousVolumeActorFactoryNode.h
@@ -0,0 +1,28 @@
+// Copyright Epic Games, Inc. All Rights Reserved.
+
+#pragma once
+
+#include "CoreMinimal.h"
+#include "InterchangeActorFactoryNode.h"
+
+#include "InterchangeHeterogeneousVolumeActorFactoryNode.generated.h"
+
+UCLASS(BlueprintType)
+class INTERCHANGEFACTORYNODES_API UInterchangeHeterogeneousVolumeActorFactoryNode : public UInterchangeActorFactoryNode
+{
+	GENERATED_BODY()
+
+public:
+	UClass* GetObjectClass() const override;
+
+	/** Gets the Uid of the material that should be assigned to the spawned HeterogeneousVolumeActor at its single material slot */
+	UFUNCTION(BlueprintCallable, Category = "Interchange | Node | HeterogeneousVolumeActor")
+	bool GetCustomVolumetricMaterialUid(FString& MaterialFactoryNodeUid) const;
+
+	/** Sets the Uid of the material that should be assigned to the spawned HeterogeneousVolumeActor at its single material slot */
+	UFUNCTION(BlueprintCallable, Category = "Interchange | Node | HeterogeneousVolumeActor")
+	bool SetCustomVolumetricMaterialUid(const FString& MaterialFactoryNodeUid);
+
+private:
+	IMPLEMENT_NODE_ATTRIBUTE_KEY(MaterialDependency)
+};
diff --git a/Engine/Plugins/Interchange/Runtime/Source/FactoryNodes/Public/InterchangeSparseVolumeTextureFactoryNode.h b/Engine/Plugins/Interchange/Runtime/Source/FactoryNodes/Public/InterchangeSparseVolumeTextureFactoryNode.h
--- a/Engine/Plugins/Interchange/Runtime/Source/FactoryNodes/Public/InterchangeSparseVolumeTextureFactoryNode.h
+++ b/Engine/Plugins/Interchange/Runtime/Source/FactoryNodes/Public/InterchangeSparseVolumeTextureFactoryNode.h
@@ -20,19 +20,19 @@ public:
 public:
 	/** Gets the data type of the AttributesA texture of the SparseVolumeTexture we'll create */
 	UFUNCTION(BlueprintCallable, Category = "Interchange | Node | Material")
-	bool GetCustomAttributesADataType(EInterchangeSparseVolumeTextureFormat& DataType) const;
+	bool GetCustomAttributesAFormat(EInterchangeSparseVolumeTextureFormat& Format) const;
 
 	/** Sets the data type of the AttributesA texture of the SparseVolumeTexture we'll create */
 	UFUNCTION(BlueprintCallable, Category = "Interchange | Node | Material")
-	bool SetCustomAttributesADataType(EInterchangeSparseVolumeTextureFormat DataType);
+	bool SetCustomAttributesAFormat(EInterchangeSparseVolumeTextureFormat Format);
 
 	/** Gets the data type of the AttributesB texture of the SparseVolumeTexture we'll create */
 	UFUNCTION(BlueprintCallable, Category = "Interchange | Node | Material")
-	bool GetCustomAttributesBDataType(EInterchangeSparseVolumeTextureFormat& DataType) const;
+	bool GetCustomAttributesBFormat(EInterchangeSparseVolumeTextureFormat& Format) const;
 
 	/** Sets the data type of the AttributesB texture of the SparseVolumeTexture we'll create */
 	UFUNCTION(BlueprintCallable, Category = "Interchange | Node | Material")
-	bool SetCustomAttributesBDataType(EInterchangeSparseVolumeTextureFormat DataType);
+	bool SetCustomAttributesBFormat(EInterchangeSparseVolumeTextureFormat Format);
 
 	/** Gets the grid name and component index that will be assigned to the AttributesA texture, channel X (e.g. "density_0" or "temperature_2") */
 	UFUNCTION(BlueprintCallable, Category = "Interchange | Node | Material")
@@ -98,6 +98,14 @@ public:
 	UFUNCTION(BlueprintCallable, Category = "Interchange | Node | Material")
 	bool SetCustomAttributesBChannelW(FString GridNameAndComponentIndex);
 
+	/** Gets the animation ID of the volume nodes that were grouped together to create this animated SparseVolumeTextureFactoryNode, if any */
+	UFUNCTION(BlueprintCallable, Category = "Interchange | Node | Material")
+	bool GetCustomAnimationID(FString& OutAnimationID) const;
+
+	/** Sets the animation ID of the volume nodes that were grouped together to create this animated SparseVolumeTextureFactoryNode, if any */
+	UFUNCTION(BlueprintCallable, Category = "Interchange | Node | Material")
+	bool SetCustomAnimationID(const FString& InAnimationID);
+
 private:
 	IMPLEMENT_NODE_ATTRIBUTE_KEY(AttributesADataType)
 	IMPLEMENT_NODE_ATTRIBUTE_KEY(AttributesBDataType)
@@ -111,4 +119,6 @@ private:
 	IMPLEMENT_NODE_ATTRIBUTE_KEY(AttributesBChannelY)
 	IMPLEMENT_NODE_ATTRIBUTE_KEY(AttributesBChannelZ)
 	IMPLEMENT_NODE_ATTRIBUTE_KEY(AttributesBChannelW)
+
+	IMPLEMENT_NODE_ATTRIBUTE_KEY(AnimationID)
 };
diff --git a/Engine/Plugins/Interchange/Runtime/Source/Import/Private/Animation/InterchangeLevelSequenceHelper.cpp b/Engine/Plugins/Interchange/Runtime/Source/Import/Private/Animation/InterchangeLevelSequenceHelper.cpp
--- a/Engine/Plugins/Interchange/Runtime/Source/Import/Private/Animation/InterchangeLevelSequenceHelper.cpp
+++ b/Engine/Plugins/Interchange/Runtime/Source/Import/Private/Animation/InterchangeLevelSequenceHelper.cpp
@@ -6,13 +6,14 @@
 #include "Camera/CameraComponent.h"
 #include "CineCameraComponent.h"
 #include "Components/DirectionalLightComponent.h"
+#include "Components/HeterogeneousVolumeComponent.h"
 #include "Components/LocalLightComponent.h"
 #include "Components/PointLightComponent.h"
 #include "Components/RectLightComponent.h"
 #include "Components/SkeletalMeshComponent.h"
 #include "Components/SkyLightComponent.h"
-#include "Components/SpotLightComponent.h"
 #include "Components/SplineComponent.h"
+#include "Components/SpotLightComponent.h"
 #include "Components/StaticMeshComponent.h"
 
 #include "Engine/SkeletalMesh.h"
@@ -357,8 +358,8 @@ namespace UE::Interchange::Private
 			{EInterchangePropertyTracks::SkeletalMeshAnimationMode, {UMovieSceneEnumTrack::StaticClass()->GetName(),  GET_MEMBER_NAME_STRING_CHECKED(USkeletalMeshComponent, USkeletalMeshComponent::GetAnimationModePropertyNameChecked()) , TEXT("Animation Mode"), StaticEnum<EAnimationMode::Type>()}},
 			{EInterchangePropertyTracks::SkeletalMeshClothBlendWeight, {UMovieSceneFloatTrack::StaticClass()->GetName(),  GET_MEMBER_NAME_STRING_CHECKED(USkeletalMeshComponent, ClothBlendWeight) , TEXT("Cloth Blend Weight")}},
 			{EInterchangePropertyTracks::SkeletalMeshClothMaxDistanceScale, {UMovieSceneFloatTrack::StaticClass()->GetName(),  GET_MEMBER_NAME_STRING_CHECKED(USkeletalMeshComponent, ClothMaxDistanceScale) , TEXT("Cloth Max Distance Scale")}},
-			
-			// Static Mesh				
+
+			// Static Mesh
 			{EInterchangePropertyTracks::StaticMesh, {UMovieSceneObjectPropertyTrack::StaticClass()->GetName(), UStaticMeshComponent::GetMemberNameChecked_StaticMesh().ToString() , TEXT("Static Mesh"), UStaticMesh::StaticClass()}},
 			{EInterchangePropertyTracks::StaticMeshDistanceFieldSelfShadowBias, {UMovieSceneFloatTrack::StaticClass()->GetName(),  GET_MEMBER_NAME_STRING_CHECKED(UStaticMeshComponent, DistanceFieldSelfShadowBias) , TEXT("Distance Field Self Shadow Bias")}},
 			{EInterchangePropertyTracks::StaticMeshEvaluateWorldPositionOffset, {UMovieSceneBoolTrack::StaticClass()->GetName(),  GET_MEMBER_NAME_STRING_CHECKED(UStaticMeshComponent, bEvaluateWorldPositionOffset) , TEXT("Evaluate World Position Offset")}},
@@ -366,6 +367,9 @@ namespace UE::Interchange::Private
 			{EInterchangePropertyTracks::StaticMeshForcedLodModel, {UMovieSceneIntegerTrack::StaticClass()->GetName(),  GET_MEMBER_NAME_STRING_CHECKED(UStaticMeshComponent, ForcedLodModel) , TEXT("Forced Lod Model")}},
 			{EInterchangePropertyTracks::StaticMeshReverseCulling, {UMovieSceneBoolTrack::StaticClass()->GetName(),  GET_MEMBER_NAME_STRING_CHECKED(UStaticMeshComponent, bReverseCulling) , TEXT("Reverse Culling")}},
 			{EInterchangePropertyTracks::StaticMeshWorldPositionOffsetDisableDistance, {UMovieSceneIntegerTrack::StaticClass()->GetName(),  GET_MEMBER_NAME_STRING_CHECKED(UStaticMeshComponent, WorldPositionOffsetDisableDistance) , TEXT("World Position Offset Disable Distance")}},
+
+			// Heterogeneous Volume Component
+			{EInterchangePropertyTracks::HeterogeneousVolumeFrame, {UMovieSceneFloatTrack::StaticClass()->GetName(),  GET_MEMBER_NAME_STRING_CHECKED(UHeterogeneousVolumeComponent, Frame) , TEXT("Frame")}},
 		}
 	{}
 
diff --git a/Engine/Plugins/Interchange/Runtime/Source/Import/Private/InterchangeImportModule.cpp b/Engine/Plugins/Interchange/Runtime/Source/Import/Private/InterchangeImportModule.cpp
--- a/Engine/Plugins/Interchange/Runtime/Source/Import/Private/InterchangeImportModule.cpp
+++ b/Engine/Plugins/Interchange/Runtime/Source/Import/Private/InterchangeImportModule.cpp
@@ -25,6 +25,7 @@
 #include "Scene/InterchangeActorFactory.h"
 #include "Scene/InterchangeCameraActorFactory.h"
 #include "Scene/InterchangeDecalActorFactory.h"
+#include "Scene/InterchangeHeterogeneousVolumeActorFactory.h"
 #include "Scene/InterchangeLevelFactory.h"
 #include "Scene/InterchangeLevelInstanceActorFactory.h"
 #include "Scene/InterchangeLightActorFactory.h"
@@ -132,6 +133,7 @@ void FInterchangeImportModule::StartupModule()
 		InterchangeManager.RegisterFactory(UInterchangeLightActorFactory::StaticClass());
 		InterchangeManager.RegisterFactory(UInterchangeSceneImportAssetFactory::StaticClass());
 		InterchangeManager.RegisterFactory(UInterchangeSpecularProfileFactory::StaticClass());
+		InterchangeManager.RegisterFactory(UInterchangeHeterogeneousVolumeActorFactory::StaticClass());
 #if WITH_EDITOR
 		// Editor-only because SVT internals like UStreamableSparseVolumeTexture::AppendFrame are editor-only
 		InterchangeManager.RegisterFactory(UInterchangeSparseVolumeTextureFactory::StaticClass());
diff --git a/Engine/Plugins/Interchange/Runtime/Source/Import/Private/Material/InterchangeMaterialFactory.cpp b/Engine/Plugins/Interchange/Runtime/Source/Import/Private/Material/InterchangeMaterialFactory.cpp
--- a/Engine/Plugins/Interchange/Runtime/Source/Import/Private/Material/InterchangeMaterialFactory.cpp
+++ b/Engine/Plugins/Interchange/Runtime/Source/Import/Private/Material/InterchangeMaterialFactory.cpp
@@ -1,41 +1,45 @@
-// Copyright Epic Games, Inc. All Rights Reserved. 
+// Copyright Epic Games, Inc. All Rights Reserved.
 #include "Material/InterchangeMaterialFactory.h"
 
+#include "Engine/RendererSettings.h"
+#include "Engine/SpecularProfile.h"
 #include "InterchangeAssetImportData.h"
+#include "InterchangeDecalMaterialFactoryNode.h"
 #include "InterchangeImportCommon.h"
 #include "InterchangeImportLog.h"
 #include "InterchangeMaterialFactoryNode.h"
-#include "InterchangeDecalMaterialFactoryNode.h"
 #include "InterchangeShaderGraphNode.h"
 #include "InterchangeSourceData.h"
+#include "InterchangeSparseVolumeTextureFactoryNode.h"
 #include "InterchangeSpecularProfileFactoryNode.h"
-#include "InterchangeTextureNode.h"
 #include "InterchangeTextureFactoryNode.h"
+#include "InterchangeTextureNode.h"
 #include "InterchangeTranslatorBase.h"
 #include "Materials/Material.h"
 #include "Materials/MaterialExpression.h"
 #include "Materials/MaterialExpressionClearCoatNormalCustomOutput.h"
-#include "Materials/MaterialExpressionMaterialFunctionCall.h"
 #include "Materials/MaterialExpressionFunctionInput.h"
 #include "Materials/MaterialExpressionFunctionOutput.h"
 #include "Materials/MaterialExpressionMakeMaterialAttributes.h"
+#include "Materials/MaterialExpressionMaterialFunctionCall.h"
 #include "Materials/MaterialExpressionScalarParameter.h"
+#include "Materials/MaterialExpressionSubstrate.h"
+#include "Materials/MaterialExpressionSwitch.h"
+#include "Materials/MaterialExpressionSwitch.h"
 #include "Materials/MaterialExpressionTextureCoordinate.h"
 #include "Materials/MaterialExpressionTextureSample.h"
 #include "Materials/MaterialExpressionThinTranslucentMaterialOutput.h"
 #include "Materials/MaterialExpressionVectorParameter.h"
-#include "Materials/MaterialExpressionSubstrate.h"
-#include "Materials/MaterialExpressionSwitch.h"
-#include "Materials/MaterialInterface.h"
 #include "Materials/MaterialInstanceConstant.h"
 #include "Materials/MaterialInstanceDynamic.h"
+#include "Materials/MaterialInterface.h"
+#include "Materials/MaterialInterface.h"
+#include "Misc/PackageName.h"
 #include "Nodes/InterchangeBaseNode.h"
 #include "Nodes/InterchangeBaseNodeContainer.h"
 #include "Nodes/InterchangeUserDefinedAttribute.h"
+#include "SparseVolumeTexture/SparseVolumeTexture.h"
 #include "UObject/ObjectRedirector.h"
-#include "Misc/PackageName.h"
-#include "Engine/RendererSettings.h"
-#include "Engine/SpecularProfile.h"
 
 #include UE_INLINE_GENERATED_CPP_BY_NAME(InterchangeMaterialFactory)
 
@@ -91,8 +95,8 @@ namespace UE::Interchange::MaterialFactory::Internal
 			// MaterialFuncCall appends the type to the input name when calling GetInputName
 			// and the InputName in FExpressionInput is optional so we'll check both here to be safe
 			// Sometimes we have to retrieve the input name as is because some Expressions override it based on different factors and we have no mechanism to keep tracks of it
-			if (MaterialExpression.GetInputName(It.Index) == *InputName || 
-				It->InputName == *InputName || 
+			if (MaterialExpression.GetInputName(It.Index) == *InputName ||
+				It->InputName == *InputName ||
 				MaterialExpression.UMaterialExpression::GetInputName(It.Index) == *InputName)
 			{
 				return It.Index;
@@ -210,7 +214,7 @@ namespace UE::Interchange::MaterialFactory::Internal
 		{
 			check((Material || MaterialFunction) && !(Material && MaterialFunction));
 
-			//We need to put in place a better mechanism for the reimport, for the moment let's delete all material expressions to avoid duplicates	
+			//We need to put in place a better mechanism for the reimport, for the moment let's delete all material expressions to avoid duplicates
 			DeleteAllMaterialExpressions();
 		}
 		UMaterialExpression* CreateExpressionsForNode(const UInterchangeMaterialExpressionFactoryNode& ExpressionNode)
@@ -419,7 +423,7 @@ namespace UE::Interchange::MaterialFactory::Internal
 	 * @type AssetFactoryType asset factory to use to get the referenced object
 	 * @param NodeContainer container that stores all the nodes
 	 * @param PathString path or UID of a node
-	 * 
+	 *
 	 * @returns the loaded asset if there is any
 	 */
 	template<typename AssetType, typename AssetFactoryType>
@@ -525,21 +529,23 @@ namespace UE::Interchange::MaterialFactory::Internal
 	void UpdateParameterTexture(UMaterialInstance& MaterialInstance, const FString& InputName, const UInterchangeMaterialInstanceFactoryNode& FactoryNode, const UInterchangeBaseNodeContainer& NodeContainer, bool bIsAParameter)
 	{
 		const FName ParameterName = *InputName;
-		UTexture* InstanceValue;
+		UTexture * TextureInstanceValue;
+		USparseVolumeTexture* SparseInstanceValue;
 
-		if (MaterialInstance.GetTextureParameterValue(ParameterName, InstanceValue))
+		if (MaterialInstance.GetTextureParameterValue(ParameterName, TextureInstanceValue))
 		{
 			FString InputValue;
-			const FString InputKey = bIsAParameter ? UInterchangeShaderPortsAPI::MakeInputParameterKey(InputName) : UInterchangeShaderPortsAPI::MakeInputValueKey(InputName);
+			const FString InputKey = bIsAParameter ? UInterchangeShaderPortsAPI::MakeInputParameterKey(InputName)
+												   : UInterchangeShaderPortsAPI::MakeInputValueKey(InputName);
 			if (FactoryNode.GetStringAttribute(InputKey, InputValue))
 			{
-				if (const UInterchangeTextureFactoryNode* TextureFactoryNode = Cast<UInterchangeTextureFactoryNode>(NodeContainer.GetNode(InputValue)))
+				if (const UInterchangeTextureFactoryNode* TextureFactoryNode = Cast<UInterchangeTextureFactoryNode>(NodeContainer.GetFactoryNode(InputValue)))
 				{
 					FSoftObjectPath ReferenceObject;
 					TextureFactoryNode->GetCustomReferenceObject(ReferenceObject);
 					if (UTexture* InputTexture = Cast<UTexture>(ReferenceObject.TryLoad()))
 					{
-						if (InputTexture != InstanceValue)
+						if (InputTexture != TextureInstanceValue)
 						{
 #if WITH_EDITOR
 							if (UMaterialInstanceConstant* MaterialInstanceConstant = Cast<UMaterialInstanceConstant>(&MaterialInstance))
@@ -547,7 +553,7 @@ namespace UE::Interchange::MaterialFactory::Internal
 								MaterialInstanceConstant->SetTextureParameterValueEditorOnly(ParameterName, InputTexture);
 							}
 							else
-#endif // #if WITH_EDITOR
+#endif	  // #if WITH_EDITOR
 							if (UMaterialInstanceDynamic* MaterialInstanceDynamic = Cast<UMaterialInstanceDynamic>(&MaterialInstance))
 							{
 								MaterialInstanceDynamic->SetTextureParameterValue(ParameterName, InputTexture);
@@ -557,6 +563,37 @@ namespace UE::Interchange::MaterialFactory::Internal
 				}
 			}
 		}
+		else if (MaterialInstance.GetSparseVolumeTextureParameterValue(ParameterName, SparseInstanceValue))
+		{
+			FString InputValue;
+			const FString InputKey = bIsAParameter ? UInterchangeShaderPortsAPI::MakeInputParameterKey(InputName)
+												   : UInterchangeShaderPortsAPI::MakeInputValueKey(InputName);
+			if (FactoryNode.GetStringAttribute(InputKey, InputValue))
+			{
+				if (const UInterchangeSparseVolumeTextureFactoryNode* SparseTextureNode = Cast<UInterchangeSparseVolumeTextureFactoryNode>(NodeContainer.GetFactoryNode(InputValue)))
+				{
+					FSoftObjectPath ReferenceObject;
+					SparseTextureNode->GetCustomReferenceObject(ReferenceObject);
+					if (USparseVolumeTexture* InputTexture = Cast<USparseVolumeTexture>(ReferenceObject.TryLoad()))
+					{
+						if (InputTexture != SparseInstanceValue)
+						{
+#if WITH_EDITOR
+							if (UMaterialInstanceConstant* MaterialInstanceConstant = Cast<UMaterialInstanceConstant>(&MaterialInstance))
+							{
+								MaterialInstanceConstant->SetSparseVolumeTextureParameterValueEditorOnly(ParameterName, InputTexture);
+							}
+							else
+#endif	  // #if WITH_EDITOR
+							if (UMaterialInstanceDynamic* MaterialInstanceDynamic = Cast<UMaterialInstanceDynamic>(&MaterialInstance))
+							{
+								MaterialInstanceDynamic->SetSparseVolumeTextureParameterValue(ParameterName, InputTexture);
+							}
+						}
+					}
+				}
+			}
+		}
 	}
 
 	void OverrideBoolParameter(const UInterchangeMaterialInstanceFactoryNode& FactoryNode, const FString& AttributeKey, UMaterialInstance& MaterialInstance, const FName& ParameterName)
@@ -762,7 +799,7 @@ namespace UE::Interchange::MaterialFactory::Internal
 		}
 	}
 }
-		
+
 namespace UE::Interchange::Materials::HashUtils
 {
 
@@ -852,7 +889,7 @@ UInterchangeFactoryBase::FImportAssetResult UInterchangeMaterialFactory::BeginIm
 		bSkipImport = true;
 		ImportAssetResult.bIsFactorySkipAsset = true;
 	};
-	
+
 	const FText MissMatchClassText = LOCTEXT("MatFactory_CouldNotCreateMat_MissMatchClass", "Mismatch between Interchange material factory node class and factory class.");
 
 	if (!Arguments.AssetNode || !Arguments.AssetNode->GetObjectClass()->IsChildOf(GetFactoryClass()))
@@ -1038,7 +1075,7 @@ UInterchangeFactoryBase::FImportAssetResult UInterchangeMaterialFactory::BeginIm
 		}
 #endif // #if WITH_EDITOR
 	}
-	
+
 	ImportAssetResult.ImportedObject = Material;
 	return ImportAssetResult;
 }
@@ -1103,7 +1140,7 @@ UInterchangeFactoryBase::FImportAssetResult UInterchangeMaterialFactory::ImportA
 #if WITH_EDITORONLY_DATA
 	else if(UMaterialInstance* MaterialInstance = Cast<UMaterialInstance>(MaterialObject))
 	{
-		//Apply the re import strategy 
+		//Apply the re import strategy
 		UInterchangeAssetImportData* InterchangeAssetImportData = Cast<UInterchangeAssetImportData>(MaterialInstance->AssetImportData);
 		UInterchangeFactoryBaseNode* PreviousNode = nullptr;
 		if(InterchangeAssetImportData)
@@ -1427,7 +1464,7 @@ void UInterchangeMaterialFactory::SetupMaterial(UMaterial* Material, const FImpo
 			}
 		}
 	}
-	
+
 	// Metallic
 	{
 		FString MetallicUid;
@@ -1454,7 +1491,7 @@ void UInterchangeMaterialFactory::SetupMaterial(UMaterial* Material, const FImpo
 	{
 		FString SpecularUid;
 		FString OutputName;
-		
+
 		if (MaterialFactoryNode->GetSpecularConnection(SpecularUid, OutputName))
 		{
 			const UInterchangeMaterialExpressionFactoryNode* SpecularNode = Cast<UInterchangeMaterialExpressionFactoryNode>(Arguments.NodeContainer->GetNode(SpecularUid));
@@ -1520,7 +1557,7 @@ void UInterchangeMaterialFactory::SetupMaterial(UMaterial* Material, const FImpo
 	{
 		FString EmissiveColorUid;
 		FString OutputName;
-		
+
 		if (MaterialFactoryNode->GetEmissiveColorConnection(EmissiveColorUid, OutputName))
 		{
 			const UInterchangeMaterialExpressionFactoryNode* EmissiveColor = Cast<UInterchangeMaterialExpressionFactoryNode>(Arguments.NodeContainer->GetNode(EmissiveColorUid));
@@ -1537,12 +1574,12 @@ void UInterchangeMaterialFactory::SetupMaterial(UMaterial* Material, const FImpo
 			}
 		}
 	}
-	
+
 	// Normal
 	{
 		FString NormalUid;
 		FString OutputName;
-		
+
 		if (MaterialFactoryNode->GetNormalConnection(NormalUid, OutputName))
 		{
 			const UInterchangeMaterialExpressionFactoryNode* NormalNode = Cast<UInterchangeMaterialExpressionFactoryNode>(Arguments.NodeContainer->GetNode(NormalUid));
@@ -1627,7 +1664,7 @@ void UInterchangeMaterialFactory::SetupMaterial(UMaterial* Material, const FImpo
 							OpacityInput = Material->GetExpressionInputForProperty(MP_OpacityMask);
 						}
 					}
-					
+
 					if (OpacityInput)
 					{
 						OpacityExpression->ConnectExpression(OpacityInput, GetOutputIndex(*OpacityExpression, OutputName));
@@ -1739,7 +1776,7 @@ void UInterchangeMaterialFactory::SetupMaterial(UMaterial* Material, const FImpo
 				if (UMaterialExpression* ClearCoatNormalExpression = Builder.CreateExpressionsForNode(*ClearCoatNormalNode))
 				{
 					UMaterialExpression* ClearCoatNormalCustomOutput = CreateMaterialExpression(Material, nullptr, UMaterialExpressionClearCoatNormalCustomOutput::StaticClass());
-					
+
 					ClearCoatNormalExpression->ConnectExpression(ClearCoatNormalCustomOutput->GetInput(0), GetOutputIndex(*ClearCoatNormalExpression, OutputName));
 				}
 			}
diff --git a/Engine/Plugins/Interchange/Runtime/Source/Import/Private/OpenVDB/InterchangeOpenVDBTranslator.cpp b/Engine/Plugins/Interchange/Runtime/Source/Import/Private/OpenVDB/InterchangeOpenVDBTranslator.cpp
--- a/Engine/Plugins/Interchange/Runtime/Source/Import/Private/OpenVDB/InterchangeOpenVDBTranslator.cpp
+++ b/Engine/Plugins/Interchange/Runtime/Source/Import/Private/OpenVDB/InterchangeOpenVDBTranslator.cpp
@@ -26,6 +26,15 @@ static FAutoConsoleVariableRef CVarInterchangeEnableOpenVDBImport(
 	TEXT("Whether OpenVDB support is enabled via Interchange")
 );
 
+static int32 GMaxParallelFileReads = 16;
+static FAutoConsoleVariableRef CVarNumPerPrimLocks(
+	TEXT("Interchange.FeatureFlags.Import.OpenVDB.NumFileReadLocks"),
+	GMaxParallelFileReads,
+	TEXT("Maximum number of .vdb files that can be read concurrently. This uses system mutexes, which are limited and vary depending on your system."
+	),
+	ECVF_ReadOnly	 // We'll keep a static array of these for simplicity, so changing it during the session would be annoying to handle
+);
+
 namespace UE::InterchangeOpenVDBTranslator::Private
 {
 #if WITH_EDITOR
@@ -33,6 +42,16 @@ namespace UE::InterchangeOpenVDBTranslator::Private
 	{
 		TArray64<uint8> FileBytes;
 		TArray<FOpenVDBGridInfo> GridInfo;
+
+	private:
+		friend class UInterchangeOpenVDBTranslatorImpl;
+
+		// Index of our lock within the Impl's FileLocks. Having multiple locks allows concurrent file reads
+		int32 LockIndex = INDEX_NONE;
+
+		// We set this to true if we failed to load/open this file, so that we don't try again later and violate the
+		// assumption that the FileInfos are going to be read-only after they're returned from GetOrLoadFileInfo
+		bool bFailed = false;
 	};
 
 	EVolumeGridElementType GridTypeToInterchangeGridType(EOpenVDBGridType Type)
@@ -70,7 +89,115 @@ namespace UE::InterchangeOpenVDBTranslator::Private
 	class UInterchangeOpenVDBTranslatorImpl
 	{
 	public:
-		TMap<FString, FOpenVDBFileInfo> FileNameToInfo;
+		UInterchangeOpenVDBTranslatorImpl()
+		{
+			FileLocks = new FRWLock[GMaxParallelFileReads]();
+		}
+
+		~UInterchangeOpenVDBTranslatorImpl()
+		{
+			delete[] FileLocks;
+		}
+
+		UInterchangeOpenVDBTranslatorImpl(const UInterchangeOpenVDBTranslatorImpl& Other) = delete;
+		UInterchangeOpenVDBTranslatorImpl& operator=(const UInterchangeOpenVDBTranslatorImpl& Other) = delete;
+
+	public:
+		void ClearFileInfoEntries()
+		{
+			FWriteScopeLock WriteLock(FileNameToInfoLock);
+
+			FileNameToInfo.Reset();
+		}
+
+		void EnsureFileInfoEntries(const TArray<FString>& Filenames)
+		{
+			FWriteScopeLock WriteLock(FileNameToInfoLock);
+
+			for (const FString& Filename : Filenames)
+			{
+				if (FileNameToInfo.Contains(Filename))
+				{
+					continue;
+				}
+
+				TUniquePtr<FOpenVDBFileInfo> NewInfo = MakeUnique<FOpenVDBFileInfo>();
+				NewInfo->LockIndex = NextFileLockIndex++;
+
+				FileNameToInfo.Add(Filename, MoveTemp(NewInfo));
+			}
+		}
+
+		const FOpenVDBFileInfo* GetOrLoadFileInfo(const FString& Filename)
+		{
+			FOpenVDBFileInfo* Info = nullptr;
+			{
+				FReadScopeLock ReadLock(FileNameToInfoLock);
+
+				if (TUniquePtr<FOpenVDBFileInfo>* FoundInfo = FileNameToInfo.Find(Filename))
+				{
+					Info = FoundInfo->Get();
+				}
+			}
+			if (!Info)
+			{
+				return nullptr;
+			}
+
+			// Load on-demand if we need to
+			//
+			// Only try to do this once, and set bFailed to true if we failed, so that we can assume
+			// that nothing ever writes to this after we first try this on-demand load. With this assumption,
+			// we don't need to use read locks everywhere
+			{
+				FWriteScopeLock WriteLock(FileLocks[Info->LockIndex % GMaxParallelFileReads]);
+
+				if (!Info->bFailed && Info->FileBytes.Num() == 0)
+				{
+					if (!FPaths::FileExists(Filename))
+					{
+						UE_LOG(LogInterchangeImport, Error, TEXT("OpenVDB file could not be found: %s"), *Filename);
+						Info->bFailed = true;
+						return nullptr;
+					}
+
+					if (!FFileHelper::LoadFileToArray(Info->FileBytes, *Filename))
+					{
+						UE_LOG(LogInterchangeImport, Error, TEXT("OpenVDB file could not be loaded: %s"), *Filename);
+						Info->bFailed = true;
+						return nullptr;
+					}
+
+					const bool bCreateStrings = true;
+					if (!GetOpenVDBGridInfo(Info->FileBytes, bCreateStrings, &Info->GridInfo))
+					{
+						UE_LOG(LogInterchangeImport, Error, TEXT("Failed to read OpenVDB file: %s"), *Filename);
+						Info->bFailed = true;
+						return nullptr;
+					}
+				}
+			}
+
+			return Info;
+		}
+
+		const TMap<FString, TUniquePtr<FOpenVDBFileInfo>>& GetAllFileInfo() const
+		{
+			return FileNameToInfo;
+		}
+
+	private:
+		// We don't need this when translating files directly, as we can ensure FileNameToInfo only receives
+		// new entries on the game thread. If this translator is used indirectly however (i.e. through another
+		// translator like USDTranslator) we may be asked to return payload data for a file we never seen
+		// before and may need to open on-demand. Since payload retrieval can be multi-threaded, we may
+		// have multiple threads trying to add entries to FileNameToInfo at the same time, which is what
+		// this lock is for.
+		FRWLock FileNameToInfoLock;
+		TMap<FString, TUniquePtr<FOpenVDBFileInfo>> FileNameToInfo;
+
+		FRWLock* FileLocks = nullptr;
+		int32 NextFileLockIndex = 0;
 	};
 #else
 	class UInterchangeOpenVDBTranslatorImpl
@@ -146,13 +273,29 @@ bool UInterchangeOpenVDBTranslator::Translate(UInterchangeBaseNodeContainer& Nod
 		}
 	}
 
-	// Create all the entries we'll need in one thread
-	ImplPtr->FileNameToInfo.Reset();
-	for (const FString& Filename : Filenames)
+	// Group up multiple VDB files in a single animation ID, so that the SVT pipeline makes one animated SVT factory node out of the group.
+	//
+	// We try using the ID specified in the settings first, because it's handy to be able to drive this via the USD translator as
+	// the AnimationID will make it's way into the volume / grid node UIDs, which is cumbersome to patch up afterwards
+	FString AnimationID = Settings->AnimationID;
+	if (AnimationID.IsEmpty() && Filenames.Num() > 1)
 	{
-		ImplPtr->FileNameToInfo.Add(Filename);
+		FSHA1 SHA1;
+		for (const FString& Filename : Filenames)
+		{
+			SHA1.UpdateWithString(*Filename, Filename.Len());
+		}
+		SHA1.Final();
+
+		FSHAHash FilePathHash;
+		SHA1.GetHash(&FilePathHash.Hash[0]);
+		AnimationID = FilePathHash.ToString();
 	}
 
+	// Create all the entries we'll need in one thread
+	ImplPtr->ClearFileInfoEntries();
+	ImplPtr->EnsureFileInfoEntries(Filenames);
+
 	// Collect all the frame info in parallel
 	const int32 NumBatches = Filenames.Num();
 	ParallelFor(
@@ -160,35 +303,30 @@ bool UInterchangeOpenVDBTranslator::Translate(UInterchangeBaseNodeContainer& Nod
 		[ImplPtr, &Filenames](int32 FilenameIndex)
 		{
 			const FString& Filename = Filenames[FilenameIndex];
-			FOpenVDBFileInfo& FileInfo = ImplPtr->FileNameToInfo[Filename];
-
-			if (!FPaths::FileExists(Filename))
-			{
-				UE_LOG(LogInterchangeImport, Error, TEXT("OpenVDB file could not be found: %s"), *Filename);
-				return;
-			}
-
-			if (!FFileHelper::LoadFileToArray(FileInfo.FileBytes, *Filename))
-			{
-				UE_LOG(LogInterchangeImport, Error, TEXT("OpenVDB file could not be loaded: %s"), *Filename);
-				return;
-			}
-
-			const bool bCreateStrings = true;
-			if (!GetOpenVDBGridInfo(FileInfo.FileBytes, bCreateStrings, &FileInfo.GridInfo))
-			{
-				UE_LOG(LogInterchangeImport, Error, TEXT("Failed to read OpenVDB file: %s"), *Filename);
-				return;
-			}
+			ImplPtr->GetOrLoadFileInfo(Filename);
 		}
 	);
 
-	for (const TPair<FString, FOpenVDBFileInfo>& Pair : ImplPtr->FileNameToInfo)
+	// Note that Filenames are sorted inside FindOpenVDBSequenceFileNames, so this is also our volume frame index into the animation,
+	// if we do have any (e.g. filenames [tornado_23.vdb, tornado_47.vdb, tornado_77.vdb] --> indices [0, 1, 2])
+	for (int32 FilenameIndex = 0; FilenameIndex < Filenames.Num(); ++FilenameIndex)
 	{
-		const FString& Filename = Pair.Key;
-		const FOpenVDBFileInfo& FileInfo = Pair.Value;
+		const FString& Filename = Filenames[FilenameIndex];
+		const FOpenVDBFileInfo* FileInfo = ImplPtr->GetOrLoadFileInfo(Filenames[FilenameIndex]);
+		if (!FileInfo)
+		{
+			continue;
+		}
 
-		const FString VolumeNodeUid = FPaths::GetBaseFilename(Filename);
+		const static FString VolumePrefix = TEXT("\\Volume\\");
+		const FString BaseFilename = FPaths::GetBaseFilename(Filename);
+
+		// Full filename helps prevent name collisions.
+		// AnimationID because the same volume may be used for different animations in the same import (e.g. via USD). In that case,
+		// we want to emit a factory node for each animation, but it is nice to retain the correspondence between translated node
+		// and factory node UIDs of just having the extra factory node prefix. That directly implies we need a separate volume node
+		// for each animation ID, and so the animation ID must be part of the node UID
+		const FString VolumeNodeUid = VolumePrefix + Filename + TEXT("\\") + AnimationID;
 
 		// Generate node for this file
 		UInterchangeVolumeNode* VolumeNode = nullptr;
@@ -200,14 +338,20 @@ bool UInterchangeOpenVDBTranslator::Translate(UInterchangeBaseNodeContainer& Nod
 			}
 
 			VolumeNode = NewObject<UInterchangeVolumeNode>(&NodeContainer);
-			VolumeNode->InitializeNode(VolumeNodeUid, VolumeNodeUid, EInterchangeNodeContainerType::TranslatedAsset);
+			VolumeNode->InitializeNode(VolumeNodeUid, BaseFilename, EInterchangeNodeContainerType::TranslatedAsset);
 			VolumeNode->SetAssetName(VolumeNodeUid);
 			VolumeNode->SetCustomFileName(Filename);
+			if (!AnimationID.IsEmpty())
+			{
+				VolumeNode->SetCustomAnimationID(AnimationID);
+				VolumeNode->AddCustomFrameIndexInAnimation(FilenameIndex);
+			}
+
 			NodeContainer.AddNode(VolumeNode);
 		}
 
 		// Generate nodes for each volume grid
-		for (const FOpenVDBGridInfo& GridInfo : FileInfo.GridInfo)
+		for (const FOpenVDBGridInfo& GridInfo : FileInfo->GridInfo)
 		{
 			const FString& GridName = GridInfo.Name;
 			const FString GridNodeUid = VolumeNodeUid + TEXT("_") + GridName;
@@ -227,7 +371,7 @@ bool UInterchangeOpenVDBTranslator::Translate(UInterchangeBaseNodeContainer& Nod
 			GridNode->SetCustomGridActiveAABBMax(GridInfo.VolumeActiveAABBMax);
 			GridNode->SetCustomGridActiveDimensions(GridInfo.VolumeActiveDim);
 
-			VolumeNode->AddGridDependencyUid(GridNodeUid);
+			VolumeNode->AddCustomGridDependency(GridNodeUid);
 
 			NodeContainer.AddNode(GridNode);
 			NodeContainer.SetNodeParentUid(GridNodeUid, VolumeNodeUid);
@@ -249,7 +393,7 @@ void UInterchangeOpenVDBTranslator::ReleaseSource()
 		return;
 	}
 
-	ImplPtr->FileNameToInfo.Reset();
+	ImplPtr->ClearFileInfoEntries();
 
 	if (TranslatorSettings)
 	{
@@ -328,10 +472,17 @@ TOptional<UE::Interchange::FVolumePayloadData> UInterchangeOpenVDBTranslator::Ge
 		return {};
 	}
 
-	FOpenVDBFileInfo* FoundInfo = ImplPtr->FileNameToInfo.Find(PayloadKey.FileName);
-	if (!FoundInfo)
+	// Note: We shouldn't need to lock FoundInfo's mutex to read from it after we got it:
+	// We're not going to write to it after the initial load anyway
+	const FOpenVDBFileInfo* FoundInfo = nullptr;
 	{
-		return {};
+		ImplPtr->EnsureFileInfoEntries({PayloadKey.FileName});
+
+		FoundInfo = ImplPtr->GetOrLoadFileInfo(PayloadKey.FileName);
+		if (!FoundInfo)
+		{
+			return {};
+		}
 	}
 
 	// Convert our custom struct to the one used by ConvertOpenVDBToSparseVolumeTexture
@@ -359,7 +510,7 @@ TOptional<UE::Interchange::FVolumePayloadData> UInterchangeOpenVDBTranslator::Ge
 
 	UE::Interchange::FVolumePayloadData Result;
 	bool bSucess = ConvertOpenVDBToSparseVolumeTexture(
-		FoundInfo->FileBytes,
+		const_cast<TArray64<uint8>&>(FoundInfo->FileBytes),	   // This function needs non-const access but it shouldn't ever write to anything
 		ImportOptions,
 		PayloadKey.VolumeBoundsMin,
 		Result.TextureData,
diff --git a/Engine/Plugins/Interchange/Runtime/Source/Import/Private/Scene/InterchangeHeterogeneousVolumeActorFactory.cpp b/Engine/Plugins/Interchange/Runtime/Source/Import/Private/Scene/InterchangeHeterogeneousVolumeActorFactory.cpp
new file mode 100644
--- /dev/null
+++ b/Engine/Plugins/Interchange/Runtime/Source/Import/Private/Scene/InterchangeHeterogeneousVolumeActorFactory.cpp
@@ -0,0 +1,80 @@
+// Copyright Epic Games, Inc. All Rights Reserved.
+
+#include "Scene/InterchangeHeterogeneousVolumeActorFactory.h"
+
+#include "InterchangeHeterogeneousVolumeActorFactoryNode.h"
+#include "InterchangeImportLog.h"
+#include "Nodes/InterchangeBaseNodeContainer.h"
+#include "Nodes/InterchangeFactoryBaseNode.h"
+#include "Scene/InterchangeActorHelper.h"
+
+#include "Components/HeterogeneousVolumeComponent.h"
+#include "CoreMinimal.h"
+#include "Materials/MaterialInterface.h"
+#include "UObject/Object.h"
+#include "UObject/ObjectMacros.h"
+
+UClass* UInterchangeHeterogeneousVolumeActorFactory::GetFactoryClass() const
+{
+	return AHeterogeneousVolume::StaticClass();
+}
+
+UObject* UInterchangeHeterogeneousVolumeActorFactory::ProcessActor(
+	AActor& SpawnedActor,
+	const UInterchangeActorFactoryNode& FactoryNode,
+	const UInterchangeBaseNodeContainer& NodeContainer,
+	const FImportSceneObjectsParams& Params
+)
+{
+	AHeterogeneousVolume* VolumeActor = Cast<AHeterogeneousVolume>(&SpawnedActor);
+	if (!VolumeActor)
+	{
+		return nullptr;
+	}
+
+	UHeterogeneousVolumeComponent* VolumeComponent = Cast<UHeterogeneousVolumeComponent>(VolumeActor->GetRootComponent());
+	if (!VolumeComponent)
+	{
+		return nullptr;
+	}
+
+	const UInterchangeHeterogeneousVolumeActorFactoryNode* ActorFactoryNode = Cast<UInterchangeHeterogeneousVolumeActorFactoryNode>(&FactoryNode);
+	if (!ActorFactoryNode)
+	{
+		return nullptr;
+	}
+
+	// Check for a referenced material
+	UMaterialInterface* ReferencedMaterial = nullptr;
+	if (FString MaterialFactoryNodeUid; ActorFactoryNode->GetCustomVolumetricMaterialUid(MaterialFactoryNodeUid))
+	{
+		if (UInterchangeFactoryBaseNode* MaterialFactoryNode = NodeContainer.GetFactoryNode(MaterialFactoryNodeUid))
+		{
+			FSoftObjectPath ReferencedObject;
+			MaterialFactoryNode->GetCustomReferenceObject(ReferencedObject);
+
+			ReferencedMaterial = Cast<UMaterialInterface>(ReferencedObject.TryLoad());
+			if (!ReferencedMaterial)
+			{
+				UE_LOG(
+					LogInterchangeImport,
+					Warning,
+					TEXT("Failed to find material '%s' referenced by heterogeneous volume actor factory node '%s' ('%s')"),
+					*ReferencedObject.ToString(),
+					*ActorFactoryNode->GetUniqueID(),
+					*ActorFactoryNode->GetDisplayLabel()
+				);
+			}
+		}
+	}
+	if (ReferencedMaterial)
+	{
+		ReferencedMaterial->PostLoad();
+
+		// The component is hard-coded to handle only one material (it will check for index == 0)
+		const int ElementIndex = 0;
+		VolumeComponent->SetMaterial(0, ReferencedMaterial);
+	}
+
+	return VolumeComponent;
+}
diff --git a/Engine/Plugins/Interchange/Runtime/Source/Import/Private/Usd/InterchangeUsdTranslator.cpp b/Engine/Plugins/Interchange/Runtime/Source/Import/Private/Usd/InterchangeUsdTranslator.cpp
--- a/Engine/Plugins/Interchange/Runtime/Source/Import/Private/Usd/InterchangeUsdTranslator.cpp
+++ b/Engine/Plugins/Interchange/Runtime/Source/Import/Private/Usd/InterchangeUsdTranslator.cpp
@@ -5,8 +5,6 @@
 #include "Objects/USDInfoCache.h"
 #include "Objects/USDSchemaTranslator.h"
 #include "UnrealUSDWrapper.h"
-#include "Usd/InterchangeUsdContext.h"
-#include "Usd/InterchangeUsdDefinitions.h"
 #include "USDConversionUtils.h"
 #include "USDGeomMeshConversion.h"
 #include "USDLayerUtils.h"
@@ -15,6 +13,7 @@
 #include "USDMaterialXShaderGraph.h"
 #include "USDObjectUtils.h"
 #include "USDPrimConversion.h"
+#include "USDProjectSettings.h"
 #include "USDShadeConversion.h"
 #include "USDSkeletalDataConversion.h"
 #include "USDStageOptions.h"
@@ -36,30 +35,38 @@
 #include "UsdWrappers/UsdTyped.h"
 #include "UsdWrappers/UsdVariantSets.h"
 
-#include "Async/ParallelFor.h"
-#include "HAL/IConsoleManager.h"
 #include "InterchangeCameraNode.h"
 #include "InterchangeLightNode.h"
 #include "InterchangeManager.h"
+#include "InterchangeMaterialInstanceNode.h"
 #include "InterchangeMaterialReferenceNode.h"
 #include "InterchangeMeshNode.h"
 #include "InterchangeSceneNode.h"
 #include "InterchangeShaderGraphNode.h"
 #include "InterchangeTexture2DNode.h"
 #include "InterchangeTranslatorHelper.h"
-#include "Internationalization/Regex.h"
+#include "InterchangeVolumeNode.h"
 #include "Mesh/InterchangeMeshPayload.h"
+#include "Nodes/InterchangeUserDefinedAttribute.h"
+#include "OpenVDB/InterchangeOpenVDBTranslator.h"
+#include "Usd/InterchangeUsdContext.h"
+#include "Usd/InterchangeUsdDefinitions.h"
+#include "Volume/InterchangeVolumeDefinitions.h"
+
+#include "MaterialX/MaterialXUtils/MaterialXBase.h"
+#include "MaterialX/MaterialXUtils/MaterialXManager.h"
+
+#include "Async/ParallelFor.h"
+#include "Components/HeterogeneousVolumeComponent.h"
+#include "HAL/IConsoleManager.h"
+#include "Internationalization/Regex.h"
 #include "Modules/ModuleManager.h"
 #include "MovieSceneSection.h"
-#include "Nodes/InterchangeUserDefinedAttribute.h"
 #include "Rendering/SkeletalMeshLODImporterData.h"
 #include "StaticMeshAttributes.h"
 #include "UDIMUtilities.h"
 #include "UObject/GCObjectScopeGuard.h"
 
-#include "MaterialX/MaterialXUtils/MaterialXBase.h"
-#include "MaterialX/MaterialXUtils/MaterialXManager.h"
-
 #if USE_USD_SDK
 #include "USDIncludesStart.h"
 #include "pxr/usd/usdGeom/tokens.h"
@@ -89,11 +96,13 @@ static FAutoConsoleVariableRef CVarInterchangeEnableUSDLevelImport(
 
 namespace UE::InterchangeUsdTranslator::Private
 {
+#if USE_USD_SDK
 	const static FString AnimationPrefix = TEXT("\\Animation\\");
 	const static FString AnimationTrackPrefix = TEXT("\\AnimationTrack\\");
 	const static FString BonePrefix = TEXT("\\Bone\\");
 	const static FString CameraPrefix = TEXT("\\Camera\\");
 	const static FString LightPrefix = TEXT("\\Light\\");
+	const static FString VolumePrefix = TEXT("\\Volume\\");
 	const static FString LODPrefix = TEXT("\\LOD\\");
 	const static FString MaterialPrefix = TEXT("\\Material\\");
 	const static FString MaterialReferencePrefix = TEXT("\\MaterialReference\\");
@@ -258,16 +267,15 @@ namespace UE::InterchangeUsdTranslator::Private
 	class UInterchangeUSDTranslatorImpl
 	{
 	public:
-
 		/** Add a material instance to the node container, otherwise it will add a material if it comes from a Translator (for example coming from
 		 * MaterialX which cannot handle material instances) */
-		void AddMaterialNode(
+		FString AddMaterialNode(
 			const UE::FUsdPrim& Prim,
 			UInterchangeUsdTranslatorSettings* TranslatorSettings,
 			UInterchangeBaseNodeContainer& NodeContainer
 		);
 
-		void AddMeshNode(
+		FString AddMeshNode(
 			const UE::FUsdPrim& Prim,
 			UInterchangeBaseNodeContainer& NodeContainer,
 			const FTraversalInfo& Info,
@@ -282,9 +290,17 @@ namespace UE::InterchangeUsdTranslator::Private
 			FTraversalInfo Info
 		);
 
+		TArray<FString> AddVolumeNodes(
+			const UE::FUsdPrim& InPrim,
+			UInterchangeBaseNodeContainer& InNodeContainer,
+			FString& OutMaterialInstanceUid,
+			bool& bOutNeedsFrameTrack
+		);
+
 		/**
 		 * Search for a MaterialX file embedded in the file, in case there's none, create a ShaderGraph using MaterialX
-		 * @param GeomPropValueNames - Array of the <geompropvalue> node names from the shader graph that has been converted to <image> nodes (used for baking later on by the Factory)
+		 * @param GeomPropValueNames - Array of the <geompropvalue> node names from the shader graph that has been converted to <image> nodes (used
+		 * for baking later on by the Factory)
 		 */
 		bool AddMaterialXShaderGraph(
 			const UE::FUsdPrim& InPrim,
@@ -345,11 +361,9 @@ namespace UE::InterchangeUsdTranslator::Private
 		FString USDZFilePath;
 		FString DecompressedUSDZRoot;
 
-#if USE_USD_SDK
 		// On UInterchangeUSDTranslator::Translate we set this up based on our TranslatorSettings, and then
 		// we can reuse it (otherwise we have to keep converting the FNames into Tokens all the time)
 		UsdToUnreal::FUsdMeshConversionOptions CachedMeshConversionOptions;
-#endif
 
 		TMap<FString, UsdUtils::FUsdPrimMaterialAssignmentInfo> CachedMaterialAssignments;
 
@@ -368,12 +382,14 @@ namespace UE::InterchangeUsdTranslator::Private
 		// here for easy access when parsing the tracks
 		UInterchangeAnimationTrackSetNode* CurrentTrackSet = nullptr;
 
-		// Array of translators that we call in the GetTexturePayload, the key has no real meaning, it's just here to avoid having duplicates and
+		// Map of translators that we call in the GetTexturePayload, the key has no real meaning, it's just here to avoid having duplicates and
 		// calling several times the Translate function
-		TMap<FString, UInterchangeTranslatorBase*> Translators;
+		TMap<FString, TStrongObjectPtr<UInterchangeTranslatorBase>> Translators;
 
-	private:
+		// We stash the info we collected from each Volume prim path here, as we'll reuse it between translation and retrieving the payloads
+		TMap<FString, TArray<UsdUtils::FVolumePrimInfo>> PrimPathToVolumeInfo;
 
+	private:
 		struct FMaterialSlotMesh
 		{
 			FString MaterialSlotName;
@@ -381,8 +397,31 @@ namespace UE::InterchangeUsdTranslator::Private
 		};
 		TMap<FString, FString> MaterialUidToActualNodeUid;
 		TMap<FString, TArray<FMaterialSlotMesh>> PrimPathToSlotMeshNodes;
+
+		// Used within a translation. We cache these because we make a volume node *per .vdb file*, and on the USD side we may have
+		// any number of Volume prims internally using the same .vdb file, and we want to share these whenever possible
+		TMap<FString, TMap<FString, UInterchangeVolumeNode*>> VolumeFilepathToAnimationIDToNode;
 	};
 
+	// Adds a numbered suffix (if needed) to NodeUid to make sure there is nothing with that ID within NodeContainer.
+	// Does not add the ID to the NodeContainer.
+	void MakeNodeUidUniqueInContainer(FString& NodeUid, const UInterchangeBaseNodeContainer& NodeContainer)
+	{
+		if (!NodeContainer.IsNodeUidValid(NodeUid) || NodeUid == UInterchangeBaseNode::InvalidNodeUid())
+		{
+			return;
+		}
+
+		int32 Suffix = 0;
+		FString Result;
+		do
+		{
+			Result = FString::Printf(TEXT("%s_%d"), *NodeUid, Suffix++);
+		} while (NodeContainer.IsNodeUidValid(Result));
+
+		NodeUid = Result;
+	}
+
 	bool IsValidLODName(const FString& PrimName)
 	{
 		return PrimName.Len() > LODString.Len() && PrimName.StartsWith(LODString) && PrimName.RightChop(LODString.Len()).IsNumeric();
@@ -431,7 +470,6 @@ namespace UE::InterchangeUsdTranslator::Private
 		return false;
 	}
 
-#if USE_USD_SDK
 	FString HashAnimPayloadQuery(const Interchange::FAnimationPayloadQuery& Query)
 	{
 		FSHAHash Hash;
@@ -915,21 +953,22 @@ namespace UE::InterchangeUsdTranslator::Private
 			[&ShaderNodesTextureSampleUIDs](const FString&, UInterchangeShaderNode* ShaderNode)
 			{
 				using namespace UE::Interchange::Materials::Standard::Nodes;
-				
+
 				constexpr const TCHAR* GeomPropHandledAttribute = TEXT("GeomPropHandled");
 				if (FString ShaderType; ShaderNode->GetCustomShaderType(ShaderType) && ShaderType == TextureSample::Name.ToString())
 				{
 					if (bool bIsGeomProp; ShaderNode->GetBooleanAttribute(::MaterialX::Attributes::GeomPropImage, bIsGeomProp) && bIsGeomProp)
 					{
 						// We only want the geomprops that were not already assigned to the mesh node
-						if(bool bGeomPropHandled; !ShaderNode->GetBooleanAttribute(GeomPropHandledAttribute, bGeomPropHandled) || !bGeomPropHandled)
+						if (bool bGeomPropHandled; !ShaderNode->GetBooleanAttribute(GeomPropHandledAttribute, bGeomPropHandled) || !bGeomPropHandled)
 						{
 							ShaderNode->AddBooleanAttribute(GeomPropHandledAttribute, true);
 							ShaderNodesTextureSampleUIDs.Emplace(ShaderNode->GetUniqueID());
 						}
 					}
 				}
-			});
+			}
+		);
 
 		// Now we need to check if we have to set the slot of the mesh nodes here
 		if (TArray<FMaterialSlotMesh>* SlotMeshes = PrimPathToSlotMeshNodes.Find(MaterialUid))
@@ -945,8 +984,14 @@ namespace UE::InterchangeUsdTranslator::Private
 						for (int32 Index = 0; Index < GeomPropValues.Num(); ++Index)
 						{
 							MaterialSlotMesh.MeshNode->AddStringAttribute(USD::Primvar::Name + FString::FromInt(Index), GeomPropValues[Index].Name);
-							MaterialSlotMesh.MeshNode->AddBooleanAttribute(USD::Primvar::TangentSpace + FString::FromInt(Index), GeomPropValues[Index].bTangentSpace);
-							MaterialSlotMesh.MeshNode->AddStringAttribute(USD::Primvar::ShaderNodeTextureSample + FString::FromInt(Index), ShaderNodesTextureSampleUIDs[Index]);
+							MaterialSlotMesh.MeshNode->AddBooleanAttribute(
+								USD::Primvar::TangentSpace + FString::FromInt(Index),
+								GeomPropValues[Index].bTangentSpace
+							);
+							MaterialSlotMesh.MeshNode->AddStringAttribute(
+								USD::Primvar::ShaderNodeTextureSample + FString::FromInt(Index),
+								ShaderNodesTextureSampleUIDs[Index]
+							);
 						}
 					}
 
@@ -986,7 +1031,7 @@ namespace UE::InterchangeUsdTranslator::Private
 					if (Translator)
 					{
 						Translator->Translate(NodeContainer);
-						Translators.Add(File, Translator);
+						Translators.Emplace(File, Translator);
 					}
 				}
 			}
@@ -1129,8 +1174,10 @@ namespace UE::InterchangeUsdTranslator::Private
 					*Prim.GetPrimPath().GetString(),
 					*UsdUtils::StringifyMap(PrimvarToUVIndex),
 					Value.Primvar.IsEmpty()
-					? TEXT(" Is your UsdUVTexture Shader missing the 'inputs:st' attribute? (It specifies which UV set to sample the texture with)")
-					: TEXT("")
+						? TEXT(
+							  " Is your UsdUVTexture Shader missing the 'inputs:st' attribute? (It specifies which UV set to sample the texture with)"
+						  )
+						: TEXT("")
 				);
 			}
 
@@ -1183,7 +1230,7 @@ namespace UE::InterchangeUsdTranslator::Private
 		const FString* BaseParameterName = nullptr;
 	};
 
-	void UInterchangeUSDTranslatorImpl::AddMaterialNode(
+	FString UInterchangeUSDTranslatorImpl::AddMaterialNode(
 		const UE::FUsdPrim& Prim,
 		UInterchangeUsdTranslatorSettings* TranslatorSettings,
 		UInterchangeBaseNodeContainer& NodeContainer
@@ -1202,8 +1249,7 @@ namespace UE::InterchangeUsdTranslator::Private
 			TOptional<FString> UnrealContentPath = UsdUtils::GetUnrealSurfaceOutput(Prim);
 			if (UnrealContentPath.IsSet())
 			{
-				UE::InterchangeUsdTranslator::Private::AddUnrealMaterialReferenceNodeIfNeeded(NodeContainer, UnrealContentPath.GetValue());
-				return;
+				return UE::InterchangeUsdTranslator::Private::AddUnrealMaterialReferenceNodeIfNeeded(NodeContainer, UnrealContentPath.GetValue());
 			}
 		}
 
@@ -1214,14 +1260,15 @@ namespace UE::InterchangeUsdTranslator::Private
 		// Check if Node already exist with this ID
 		if (const UInterchangeShaderGraphNode* Node = Cast<UInterchangeShaderGraphNode>(NodeContainer.GetNode(MaterialUid)))
 		{
-			return;
+			return MaterialUid;
 		}
 
 		// we only create material instances if we didn't find any MaterialX instances (including a shader graph)
-		if (TArray<FUsdMaterialXShaderGraph::FGeomProp> GeomPropValues; AddMaterialXShaderGraph(Prim, TranslatorSettings, NodeContainer, GeomPropValues))
+		if (TArray<FUsdMaterialXShaderGraph::FGeomProp> GeomPropValues;
+			AddMaterialXShaderGraph(Prim, TranslatorSettings, NodeContainer, GeomPropValues))
 		{
 			SetMaterialSlotDependencies(MaterialPrimName, MaterialUid, NodeContainer, GeomPropValues, true);
-			return;
+			return MaterialUid;
 		}
 
 		UInterchangeShaderGraphNode* MaterialNode = NewObject<UInterchangeShaderGraphNode>(&NodeContainer);
@@ -1274,6 +1321,8 @@ namespace UE::InterchangeUsdTranslator::Private
 		{
 			MaterialNode->AddBooleanAttribute(UE::Interchange::USD::NeedsVTParentIdentifier, true);
 		}
+
+		return MaterialUid;
 	}
 
 	FString UInterchangeUSDTranslatorImpl::GetTextureSourcePath(const FString& TexturePathOnDisk)
@@ -1536,11 +1585,11 @@ namespace UE::InterchangeUsdTranslator::Private
 		CurrentLODSceneNodes.Reset();
 	}
 
-	void AddLightNode(const UE::FUsdPrim& Prim, UInterchangeBaseNodeContainer& NodeContainer)
+	FString AddLightNode(const UE::FUsdPrim& Prim, UInterchangeBaseNodeContainer& NodeContainer)
 	{
-#if USE_USD_SDK
+		FString NodeUid;
 
-		FString NodeUid = LightPrefix + Prim.GetPrimPath().GetString();
+		NodeUid = LightPrefix + Prim.GetPrimPath().GetString();
 		FString NodeName(Prim.GetName().ToString());
 
 		// Ref. UsdToUnreal::ConvertLight
@@ -1669,13 +1718,15 @@ namespace UE::InterchangeUsdTranslator::Private
 		}
 		// #ueent_todo:
 		// DomeLight -> SkyLight
-#endif	  // USE_USD_SDK
+
+		return NodeUid;
 	}
 
-	void AddCameraNode(const UE::FUsdPrim& Prim, UInterchangeBaseNodeContainer& NodeContainer)
+	FString AddCameraNode(const UE::FUsdPrim& Prim, UInterchangeBaseNodeContainer& NodeContainer)
 	{
-#if USE_USD_SDK
-		FString NodeUid = CameraPrefix + Prim.GetPrimPath().GetString();
+		FString NodeUid;
+
+		NodeUid = CameraPrefix + Prim.GetPrimPath().GetString();
 		FString NodeName(Prim.GetName().ToString());
 
 		UInterchangePhysicalCameraNode* CameraNode = NewObject<UInterchangePhysicalCameraNode>(&NodeContainer);
@@ -1703,7 +1754,8 @@ namespace UE::InterchangeUsdTranslator::Private
 		CameraNode->SetCustomSensorHeight(SensorHeight);
 
 		// Focus distance and FStop not currently supported
-#endif	  // USE_USD_SDK
+
+		return NodeUid;
 	}
 
 	void AddMorphTargetNodes(
@@ -1723,7 +1775,7 @@ namespace UE::InterchangeUsdTranslator::Private
 		const FString MeshPrimPath = MeshPrim.GetPrimPath().GetString();
 
 		TFunction<void(const FString&, int32, const FString&)> AddMorphTargetNode =
-		[&MeshNode, &MeshPrimPath, &NodeContainer, &Info](const FString& MorphTargetName, int32 BlendShapeIndex, const FString& InbetweenName)
+			[&MeshNode, &MeshPrimPath, &NodeContainer, &Info](const FString& MorphTargetName, int32 BlendShapeIndex, const FString& InbetweenName)
 		{
 			// Note: We identify a blend shape by its Mesh prim path and the blend shape index, even though
 			// the blend shape itself is a full standalone prim. This is for two reasons:
@@ -1768,7 +1820,7 @@ namespace UE::InterchangeUsdTranslator::Private
 		}
 	}
 
-	void UInterchangeUSDTranslatorImpl::AddMeshNode(
+	FString UInterchangeUSDTranslatorImpl::AddMeshNode(
 		const UE::FUsdPrim& Prim,
 		UInterchangeBaseNodeContainer& NodeContainer,
 		const FTraversalInfo& Info,
@@ -1782,12 +1834,13 @@ namespace UE::InterchangeUsdTranslator::Private
 		// Check if Node already exist with this ID
 		if (const UInterchangeMeshNode* Node = Cast<UInterchangeMeshNode>(NodeContainer.GetNode(NodeUid)))
 		{
-			return;
+			return NodeUid;
 		}
 
 		// Fill in the MeshNode itself
 		const bool bIsAnimated = UsdUtils::IsAnimatedMesh(Prim);
-		UInterchangeMeshNode* MeshNode = bIsAnimated ? NewObject<UInterchangeGeometryCacheNode>(&NodeContainer) : NewObject<UInterchangeMeshNode>(&NodeContainer);
+		UInterchangeMeshNode* MeshNode = bIsAnimated ? NewObject<UInterchangeGeometryCacheNode>(&NodeContainer)
+													 : NewObject<UInterchangeMeshNode>(&NodeContainer);
 		MeshNode->InitializeNode(NodeUid, NodeName, EInterchangeNodeContainerType::TranslatedAsset);
 		MeshNode->SetAssetName(NodeName);
 		const bool bIsSkinned = static_cast<bool>(Info.ClosestParentSkelRootPath) && Prim.HasAPI(TEXT("SkelBindingAPI"));
@@ -1802,7 +1855,7 @@ namespace UE::InterchangeUsdTranslator::Private
 			// We are currently not supporting Skinned Primitive Shapes.
 			// In theory a Skinned Mesh needs Joint influences and weights information provided,
 			// However, there does not seem to be a ruleset against a Primitive Shape having SkelBindingAPI set.
-			// Which means, that on a theoretical level there could be such a scenario. 
+			// Which means, that on a theoretical level there could be such a scenario.
 			ensureMsgf(!bIsSkinned, TEXT("Unexpected scenario: Primitive Shape is skinned."));
 
 			// For Primitive Shapes we add PrimitiveShapePrefix for the PayLoadKey,
@@ -1941,6 +1994,7 @@ namespace UE::InterchangeUsdTranslator::Private
 		}
 
 		NodeContainer.AddNode(MeshNode);
+		return NodeUid;
 	}
 
 	void UInterchangeUSDTranslatorImpl::AddLODMeshNodes(	//
@@ -2025,6 +2079,441 @@ namespace UE::InterchangeUsdTranslator::Private
 		}
 	}
 
+	void AddVolumeCustomAttributesToNode(const UsdUtils::FVolumePrimInfo& VolumePrimInfo, UInterchangeVolumeNode* VolumeNode)
+	{
+		using namespace UE::Interchange::USD;
+
+		// Convert from the {'velocity': {'AttributesA.R': 'X', 'AttributesA.G': 'Y', 'AttributesA.B': 'Z'}} style of mapping
+		// from grid info into {"AttributesA.X": "velocity_0", "AttributesA.G": "velocity_1", "AttributesA.B": "velocity_2"}
+		// mapping into the VolumeNode custom attributes
+		for (const TPair<FString, TMap<FString, FString>>& GridPair : VolumePrimInfo.GridNameToChannelComponentMapping)
+		{
+			const FString& GridName = GridPair.Key;	   // "density", "temperature", etc.
+
+			const TMap<FString, FString>& AttributesChannelToGridChannel = GridPair.Value;
+			for (const TPair<FString, FString>& AssignmentPair : AttributesChannelToGridChannel)
+			{
+				const FString& AttributeChannel = AssignmentPair.Key;	 // "AttributesA.R", "AttributesB.B", etc.
+				const FString& GridChannel = AssignmentPair.Value;		 // "X", "Y", "Z", etc.
+
+				const static TMap<FString, FString> AttributeChannelToAttributeKey{
+					{TEXT("AttributesA.R"), SparseVolumeTexture::AttributesAChannelR},
+					{TEXT("AttributesA.G"), SparseVolumeTexture::AttributesAChannelG},
+					{TEXT("AttributesA.B"), SparseVolumeTexture::AttributesAChannelB},
+					{TEXT("AttributesA.A"), SparseVolumeTexture::AttributesAChannelA},
+					{TEXT("AttributesB.R"), SparseVolumeTexture::AttributesBChannelR},
+					{TEXT("AttributesB.G"), SparseVolumeTexture::AttributesBChannelG},
+					{TEXT("AttributesB.B"), SparseVolumeTexture::AttributesBChannelB},
+					{TEXT("AttributesB.A"), SparseVolumeTexture::AttributesBChannelA},
+				};
+
+				const static TMap<FString, FString> GridChannelToComponentIndex{
+					{TEXT("X"), TEXT("0")},
+					{TEXT("Y"), TEXT("1")},
+					{TEXT("Z"), TEXT("2")},
+					{TEXT("W"), TEXT("3")},
+					{TEXT("R"), TEXT("0")},
+					{TEXT("G"), TEXT("1")},
+					{TEXT("B"), TEXT("2")},
+					{TEXT("A"), TEXT("3")},
+				};
+
+				const FString* FoundAttributeKey = AttributeChannelToAttributeKey.Find(AttributeChannel);
+				if (!FoundAttributeKey)
+				{
+					UE_LOG(LogUsd, Warning, TEXT("Failing to parse unreal:SVT:mappedAttributeChannels value '%s'"), *AttributeChannel);
+					continue;
+				}
+
+				const FString* FoundComponentIndex = GridChannelToComponentIndex.Find(GridChannel);
+				if (!FoundComponentIndex)
+				{
+					UE_LOG(LogUsd, Warning, TEXT("Failing to parse unreal:SVT:mappedGridComponents value '%s'"), *GridChannel);
+					continue;
+				}
+
+				VolumeNode->AddStringAttribute(
+					*FoundAttributeKey,
+					GridName + UE::Interchange::Volume::GridNameAndComponentIndexSeparator + *FoundComponentIndex
+				);
+			}
+		}
+
+		// Convert texture format
+		static_assert((int)UsdUtils::ESparseVolumeAttributesFormat::Unorm8 == (int)EInterchangeSparseVolumeTextureFormat::Unorm8);
+		static_assert((int)UsdUtils::ESparseVolumeAttributesFormat::Float16 == (int)EInterchangeSparseVolumeTextureFormat::Float16);
+		static_assert((int)UsdUtils::ESparseVolumeAttributesFormat::Float32 == (int)EInterchangeSparseVolumeTextureFormat::Float32);
+		if (VolumePrimInfo.AttributesAFormat.IsSet())
+		{
+			VolumeNode->AddInt32Attribute(SparseVolumeTexture::AttributesAFormat, static_cast<int32>(VolumePrimInfo.AttributesAFormat.GetValue()));
+		}
+		if (VolumePrimInfo.AttributesBFormat.IsSet())
+		{
+			VolumeNode->AddInt32Attribute(SparseVolumeTexture::AttributesBFormat, static_cast<int32>(VolumePrimInfo.AttributesBFormat.GetValue()));
+		}
+	}
+
+	UInterchangeMaterialInstanceNode* CreateVolumetricMaterialInstanceNode(
+		const UE::FUsdPrim& VolumePrim,
+		UInterchangeBaseNodeContainer& NodeContainer,
+		const UsdToUnreal::FUsdMeshConversionOptions& ConversionOptions
+	)
+	{
+		FString ParentMaterialPath;
+
+		// Priority 1: Explicit material assignment on the Volume prim (Unreal materials).
+		{
+			double TimeCode = UsdUtils::GetDefaultTimeCode();	 // Not relevant as material bindings can't be time sampled
+			bool bProvideMaterialIndices = false;
+			UsdUtils::FUsdPrimMaterialAssignmentInfo Assignments = UsdUtils::GetPrimMaterialAssignments(
+				VolumePrim,
+				TimeCode,
+				bProvideMaterialIndices,
+				UnrealIdentifiers::Unreal,	  // Unreal materials should always be on the unreal render context
+				ConversionOptions.MaterialPurpose
+			);
+
+			for (const UsdUtils::FUsdPrimMaterialSlot& Slot : Assignments.Slots)
+			{
+				if (Slot.AssignmentType == UsdUtils::EPrimAssignmentType::UnrealMaterial)
+				{
+					ParentMaterialPath = Slot.MaterialSource;
+					break;
+				}
+			}
+		}
+
+		// Priority 2: The USD default volumetric material on the UsdProjectSettings.
+		bool bIsFallback = false;
+		if (ParentMaterialPath.IsEmpty())
+		{
+			if (const UUsdProjectSettings* ProjectSettings = GetDefault<UUsdProjectSettings>())
+			{
+				bIsFallback = true;
+				ParentMaterialPath = ProjectSettings->ReferenceDefaultSVTMaterial.ToString();
+			}
+		}
+
+		// Priority 3: Hard-coded fallback volumetric material that ships with the engine.
+		if (ParentMaterialPath.IsEmpty())
+		{
+			bIsFallback = true;
+			ParentMaterialPath = TEXT("/Engine/EngineMaterials/SparseVolumeMaterial.SparseVolumeMaterial");
+		}
+
+		FString MaterialDisplayLabel = bIsFallback ? TEXT("USDVolumetricFallbackMaterial") : FPaths::GetBaseFilename(ParentMaterialPath);
+		FString MaterialNodeUid = MaterialPrefix + MaterialDisplayLabel;
+		MakeNodeUidUniqueInContainer(MaterialNodeUid, NodeContainer);
+
+		// We'll always spawn a new material instance for each volume prim. Realistically a stage is only going to have a handful
+		// of volumes at most, and material instances should be pretty cheap. This should be a more predictable result for the user,
+		// and it prevents us from needing some bespoke code to reuse these material instance nodes depending on their volume assignment
+		UInterchangeMaterialInstanceNode* MaterialInstance = NewObject<UInterchangeMaterialInstanceNode>(&NodeContainer);
+		MaterialInstance->InitializeNode(MaterialNodeUid, MaterialDisplayLabel, EInterchangeNodeContainerType::TranslatedAsset);
+		MaterialInstance->SetCustomParent(ParentMaterialPath);
+
+		NodeContainer.AddNode(MaterialInstance);
+		return MaterialInstance;
+	}
+
+	TArray<FString> UInterchangeUSDTranslatorImpl::AddVolumeNodes(
+		const UE::FUsdPrim& InPrim,
+		UInterchangeBaseNodeContainer& InNodeContainer,
+		FString& OutMaterialInstanceUid,
+		bool& bOutNeedsFrameTrack
+	)
+	{
+		using namespace UE::Interchange::USD;
+
+		// Each Volume prim can reference multiple FieldAsset prims. Each FieldAsset itself can point to a particular grid within
+		// a .vdb file... USD is probably too flexible here, allowing us to reference grids from separate .vdb files in the same Volume prim,
+		// or letting us refer to the same grid more than once, etc.
+		//
+		// Our end goal is to make each .vdb file into a single SparseVolumeTexture, combining all the grids that need to be read from it.
+		// We'll do our best to satisfy all the requirements, but emit some warnings if we fail on an edge case. Then we want to spawn a single
+		// HeterogeneousVolumeActor for each Volume prim, generate an instance of the right material, and assign all these generated SVTs to it.
+		//
+		// In Interchange terms, this means that we'll emit a single InterchangeVolumeNode for each .vdb file, but will
+		// emit a new UInterchangeVolumeGridNode for each grid reference within that file. Note that USD's flexibility means we may have
+		// separate Volume prims all referencing the same shared FieldAsset prim, so we need to presume a InterchangeVolumeNode for this .vdb
+		// file has potentially already been created when parsing another Volume prim...
+		//
+		// Finally, since we may use the same .vdb file in multiple animations, and we want to end up with separate animated SVTs, we need
+		// separate factory nodes. If we want to keep the expected mapping of factory node / volume node Uids (just an added "Factory_" prefix)
+		// this means we need a separate volume node *per animation*, so we'll also use "AnimationIDs" to differentiate them
+
+		TStrongObjectPtr<UInterchangeBaseNodeContainer> TempContainer = nullptr;
+		TStrongObjectPtr<UInterchangeOpenVDBTranslator> Translator = nullptr;
+
+		// This is collected by path hash here because for animated SVTs we want to still have a single UsdUtils::FVolumePrimInfo for each
+		// group of animated volume frames, since they will become a separate SVT. If we just collected them by filepath we could run into
+		// trouble if we had a volume prim with 3 frames starting at "file.vdb", and a separate volume prim that just wants one frame, "frame.vdb"
+		TMap<FString, UsdUtils::FVolumePrimInfo> VolumeInfoByFilePathHash = UsdUtils::GetVolumeInfoByFilePathHash(InPrim);
+
+		TSet<FString> VolumeAssetNodeUids;
+		TMap<FString, FString> VolumeFieldNameToNodeUids;
+
+		bOutNeedsFrameTrack = false;
+
+		// Stash these as we may need this info later when retrieving animation pipelines
+		TArray<UsdUtils::FVolumePrimInfo>& CollectedInfoForPrim = PrimPathToVolumeInfo.FindOrAdd(InPrim.GetPrimPath().GetString());
+
+		for (const TPair<FString, UsdUtils::FVolumePrimInfo>& Pair : VolumeInfoByFilePathHash)
+		{
+			const FString& AnimationID = Pair.Key;
+			const UsdUtils::FVolumePrimInfo& VolumePrimInfo = Pair.Value;
+			CollectedInfoForPrim.Add(VolumePrimInfo);
+
+			TArray<FString> VDBFilePaths;
+			VDBFilePaths.Reserve(VolumePrimInfo.TimeSamplePathIndices.Num() + 1);
+
+			// In case we have both timeSamples and different default opinion, add the default opinion as the first frame so that's what it shows
+			// on the level. The LevelSequence Frame track will factor this in, and have the LevelSequence only go through the TimeSamplePaths frames
+			// though.
+			bool bInsertedDefaultOpinion = false;
+			if (VolumePrimInfo.TimeSamplePaths.Num() > 0 && VolumePrimInfo.TimeSamplePaths[0] != VolumePrimInfo.SourceVDBFilePath)
+			{
+				VDBFilePaths.Add(VolumePrimInfo.SourceVDBFilePath);
+				bInsertedDefaultOpinion = true;
+			}
+			// No time samples at all
+			else if (VolumePrimInfo.TimeSamplePaths.Num() == 0)
+			{
+				VDBFilePaths.Add(VolumePrimInfo.SourceVDBFilePath);
+			}
+
+			// Add the file paths going through TimeSamplePathIndices because it's possible that GetVolumeInfoByFilePathHash
+			// deduplicated volume frames already. It's fine to add duplicate entries to VDBFilePaths though, because we'll check
+			// for an existing volume node for that path every time anyway
+			for (int32 PathIndex : VolumePrimInfo.TimeSamplePathIndices)
+			{
+				VDBFilePaths.Add(VolumePrimInfo.TimeSamplePaths[PathIndex]);
+			}
+
+			if (VolumePrimInfo.TimeSamplePathTimeCodes.Num() > 0)
+			{
+				bOutNeedsFrameTrack = true;
+			}
+
+			// First volume is special as that is what will "become" the animated factory node if we have animation. We'll also only stash
+			// our custom attributes on this first volume node. The filepaths are sorted according to timeSamples, so this is always the first
+			// frame of animation, or the default opinion if we have that
+			UInterchangeVolumeNode* FirstVolumeNode = nullptr;
+
+			for (int32 Index = 0; Index < VDBFilePaths.Num(); ++Index)
+			{
+				const FString& FilePath = VDBFilePaths[Index];
+
+				TMap<FString, UInterchangeVolumeNode*>& AnimationIDToVolumeNode = VolumeFilepathToAnimationIDToNode.FindOrAdd(FilePath);
+				UInterchangeVolumeNode* VolumeNode = AnimationIDToVolumeNode.FindRef(AnimationID);
+
+				// Need to translate the volume for this animation ID.
+				// Note: It may seem wasteful to translate the same volume more than once in case it is used by multiple animations,
+				// but keep in mind that:
+				//  - Multiple animations for the same volume frame in the same import will realistically never happen in practice;
+				//  - The VDB translator will cache the read file bytes from the first translation;
+				//  - "Translating" the volume just involves returning some header information, which should be pretty fast;
+				//  - Doing this saves us from having to manually duplicate other volume and grid nodes, and manually patching up
+				//    their unique IDs and / or making some sort of mistake;
+				if (!VolumeNode)
+				{
+					UInterchangeSourceData* SourceData = UInterchangeManager::CreateSourceData(FilePath);
+
+					if (!Translator || !Translator->CanImportSourceData(SourceData))
+					{
+						// Check and create an OpenVDB translator directly for the file. Ideally we'd switch the cvar
+						// ourselves and just use UInterchangeManager::GetTranslatorForSourceData(), however we can only toggle
+						// cvars from the game thread, and translation can happen on any thread it seems
+						const bool bIncludeDot = false;
+						const FString Extension = FPaths::GetExtension(FilePath, bIncludeDot);
+						if (Extension.ToLower() == TEXT("vdb"))
+						{
+							Translator.Reset(NewObject<UInterchangeOpenVDBTranslator>(GetTransientPackage(), NAME_None));
+							Translator->SourceData = SourceData;
+						}
+					}
+					if (!Translator)
+					{
+						UE_LOG(LogUsd, Warning, TEXT("Failed to retrieve a translator for filepath '%s'"), *FilePath);
+						continue;
+					}
+
+					if (UInterchangeOpenVDBTranslatorSettings* Settings = Cast<UInterchangeOpenVDBTranslatorSettings>(Translator->GetSettings()))
+					{
+						// We never want to discover new .vdb files via the OpenVDB translator for animations.
+						// If we should have animations via USD they will be described on the USD file explicictly
+						Settings->bTranslateAdjacentNumberedFiles = false;
+
+						// If the volume prim describes an animation, let's add the same AnimationID to the volume nodes that the translator
+						// will output, so that the SVT pipeline groups them up into a single factory node
+						Settings->AnimationID = VolumePrimInfo.TimeSamplePathTimeCodes.Num() > 0 ? AnimationID : FString{};
+					}
+
+					if (!TempContainer)
+					{
+						TempContainer.Reset(NewObject<UInterchangeBaseNodeContainer>());
+					}
+					TempContainer->Reset();	   // Empty the container of nodes
+
+					Translator->SourceData = SourceData;
+					Translator->Translate(*TempContainer);
+
+					// Cache that we used this translator for this filepath. If we keep it, we don't have to open the file again to
+					// retrieve the payload. Note: This is likely the same translator we used for all .vdb files during this import
+					Translators.Add(FilePath, Translator);
+
+					// Move the nodes into our own node container. We do this because we must iterate to discover the volume node
+					// (as we have no guarantees about its NodeUid), and we don't want to do a full loop over our own NodeContainer
+					// that may have other Volume nodes and many more nodes overall.
+					//
+					// Also note the slight exploit: We can get non-const access to the translated nodes in this way, which we need
+					// in order to call InNodeContainer::Add
+					TempContainer->IterateNodes(
+						[&InNodeContainer, &VolumeNode](const FString& NodeUid, UInterchangeBaseNode* Node)
+						{
+							InNodeContainer.AddNode(Node);
+
+							if (UInterchangeVolumeNode* CastNode = Cast<UInterchangeVolumeNode>(Node))
+							{
+								if (VolumeNode == nullptr)
+								{
+									// Note that since we're not discovering new .vdbs, we expect exactly one UInterchangeVolumeNode to be generated
+									VolumeNode = CastNode;
+								}
+								else
+								{
+									UE_LOG(
+										LogUsd,
+										Warning,
+										TEXT("Found unexpected volume node '%s' (display label '%s')"),
+										*NodeUid,
+										*CastNode->GetDisplayLabel()
+									);
+								}
+							}
+						}
+					);
+
+					if (VolumeNode)
+					{
+						// Remove any animation index that the OpenVDB translator may have set on this volume node,
+						// as we want to control that explicitly from here
+						TArray<int32> ExistingAnimationIndices;
+						VolumeNode->GetCustomFrameIndicesInAnimation(ExistingAnimationIndices);
+						for (const int32 ExistingIndex : ExistingAnimationIndices)
+						{
+							VolumeNode->RemoveCustomFrameIndexInAnimation(ExistingIndex);
+						}
+
+						AnimationIDToVolumeNode.Add(AnimationID, VolumeNode);
+					}
+				}
+
+				if (!VolumeNode)
+				{
+					UE_LOG(LogUsd, Warning, TEXT("Failed to produce a volume node from file '%s'"), *FilePath);
+					continue;
+				}
+
+				VolumeNode->AddCustomFrameIndexInAnimation(Index);
+
+				if (!FirstVolumeNode)
+				{
+					FirstVolumeNode = VolumeNode;
+				}
+
+				VolumeAssetNodeUids.Add(VolumeNode->GetUniqueID());
+			}
+
+			if (FirstVolumeNode)
+			{
+				// Collect all of our custom-schema-based assignment info to be used as as custom attributes on the first volume node.
+				// The USD Pipeline will handle these, and move them into the factory nodes.
+				AddVolumeCustomAttributesToNode(VolumePrimInfo, FirstVolumeNode);
+
+				for (const FString& FieldName : VolumePrimInfo.VolumeFieldNames)
+				{
+					VolumeFieldNameToNodeUids.Add(FieldName, FirstVolumeNode->GetUniqueID());
+				}
+			}
+		}
+
+		// Setup a new material instance node for this volume, which will be used by the Heterogeneous Volume Actor we'll also
+		// spawn for this volume prim.
+		//
+		// Note that there's only one material slot per actor, but that we really do need some kind of material to be put in there,
+		// so that we can at least assign our SVTs somewhere.
+		UInterchangeMaterialInstanceNode* MaterialInstance = CreateVolumetricMaterialInstanceNode(
+			InPrim,
+			InNodeContainer,
+			CachedMeshConversionOptions
+		);
+
+		// Assign SVTs as material parameters on our new material instance
+		if (ensure(MaterialInstance))
+		{
+			OutMaterialInstanceUid = MaterialInstance->GetUniqueID();
+
+			TMultiMap<FString, FString> MaterialParameterToFieldName = UsdUtils::GetVolumeMaterialParameterToFieldNameMap(InPrim);
+
+			// Prim doesn't have the attributes specifying an explicit material parameter name to volume mapping (this is probably
+			// the more common case)
+			if (MaterialParameterToFieldName.Num() == 0)
+			{
+				// Consider that the field names may match material parameter names. We can't tell if that's the case or not from here,
+				// and this may cause us to assign a VolumeUid more than once (as multiple fields may target the same SVT), but the
+				// USDPipeline will clean this up
+				for (const TPair<FString, FString>& Pair : VolumeFieldNameToNodeUids)
+				{
+					const FString& FieldName = Pair.Key;
+					const FString& VolumeUid = Pair.Value;
+
+					MaterialInstance->AddTextureParameterValue(UE::Interchange::USD::VolumeFieldNameMaterialParameterPrefix + FieldName, VolumeUid);
+				}
+			}
+			// Prim has custom attributes specifying exactly which volume should be assigned to which material parameter
+			else
+			{
+				TMap<FString, FString> ParameterNameToVolume;
+				ParameterNameToVolume.Reserve(MaterialParameterToFieldName.Num());
+
+				for (const TPair<FString, FString>& Pair : MaterialParameterToFieldName)
+				{
+					const FString& ParameterName = Pair.Key;
+					const FString& FieldName = Pair.Value;
+
+					const FString* FoundVolumeUid = VolumeFieldNameToNodeUids.Find(FieldName);
+					if (!FoundVolumeUid)
+					{
+						continue;
+					}
+
+					// Show a warning in case we have conflicting assignments, because the legacy schema translator also did it
+					if (FString* FoundAssignedVolume = ParameterNameToVolume.Find(ParameterName))
+					{
+						if (*FoundAssignedVolume != *FoundVolumeUid)
+							UE_LOG(
+								LogUsd,
+								Warning,
+								TEXT(
+									"Trying to assign different Sparse Volume Textures to the same material parameter '%s' on the material instantiated for Volume '%s' and field name '%s'! Only a single texture can be assigned to a material parameter at a time."
+								),
+								*ParameterName,
+								**FoundVolumeUid,
+								*FieldName
+							);
+						continue;
+					}
+					ParameterNameToVolume.Add(ParameterName, *FoundVolumeUid);
+
+					MaterialInstance->AddTextureParameterValue(ParameterName, *FoundVolumeUid);
+				}
+			}
+		}
+
+		return VolumeAssetNodeUids.Array();
+	}
+
 	void AddTrackSetNode(UInterchangeUSDTranslatorImpl& Impl, UInterchangeBaseNodeContainer& NodeContainer)
 	{
 		// For now we only want a single track set (i.e. LevelSequence) per stage.
@@ -2081,6 +2570,39 @@ namespace UE::InterchangeUsdTranslator::Private
 		Impl.CurrentTrackSet->AddCustomAnimationTrackUid(AnimTrackNodeUid);
 	}
 
+	void AddPropertyAnimationNode(
+		const FString& SceneNodeUid,
+		const FString& UEPropertyName,
+		EInterchangePropertyTracks TrackType,
+		EInterchangeAnimationPayLoadType PayloadType,
+		UInterchangeUSDTranslatorImpl& Impl,
+		UInterchangeBaseNodeContainer& NodeContainer
+	)
+	{
+		// We don't use the USD attribute path here because we want one unique node per UE track name,
+		// so that if e.g. both "intensity" and "exposure" are animated we make a single track for
+		// the Intensity UE property
+		const FString UniquePath = SceneNodeUid + TEXT("\\") + UEPropertyName;
+		const FString AnimTrackNodeUid = AnimationTrackPrefix + UniquePath;
+
+		const UInterchangeAnimationTrackNode* ExistingNode = Cast<UInterchangeAnimationTrackNode>(NodeContainer.GetNode(AnimTrackNodeUid));
+		if (ExistingNode)
+		{
+			return;
+		}
+
+		UInterchangeAnimationTrackNode* AnimTrackNode = NewObject<UInterchangeAnimationTrackNode>(&NodeContainer);
+		AnimTrackNode->InitializeNode(AnimTrackNodeUid, UniquePath, EInterchangeNodeContainerType::TranslatedAsset);
+		AnimTrackNode->SetCustomActorDependencyUid(*SceneNodeUid);
+		AnimTrackNode->SetCustomPropertyTrack(TrackType);
+		AnimTrackNode->SetCustomAnimationPayloadKey(UniquePath, PayloadType);
+
+		NodeContainer.AddNode(AnimTrackNode);
+
+		AddTrackSetNode(Impl, NodeContainer);
+		Impl.CurrentTrackSet->AddCustomAnimationTrackUid(AnimTrackNodeUid);
+	}
+
 	void AddPropertyAnimationNodes(const UE::FUsdPrim& Prim, UInterchangeUSDTranslatorImpl& Impl, UInterchangeBaseNodeContainer& NodeContainer)
 	{
 		using namespace UE::InterchangeUsdTranslator::Private;
@@ -2104,42 +2626,66 @@ namespace UE::InterchangeUsdTranslator::Private
 			const FName AttrTypeName = Attr.GetTypeName();
 			const bool bIsBoolTrack = AttrTypeName == TEXT("bool") || AttrTypeName == TEXT("token");	// Visibility is a token track
 
-			TArray<FName> UEAttrNames = UsdUtils::GetPropertiesForAttribute(Prim, Attr.GetName().ToString());
-			for (const FName& UEAttrName : UEAttrNames)
+			TArray<FName> UEPropertyNames = UsdUtils::GetPropertiesForAttribute(Prim, Attr.GetName().ToString());
+			for (const FName& UEPropertyName : UEPropertyNames)
 			{
-				const EInterchangePropertyTracks* FoundTrackType = PropertyNameToTrackType.Find(UEAttrName);
+				const EInterchangePropertyTracks* FoundTrackType = PropertyNameToTrackType.Find(UEPropertyName);
 				if (!FoundTrackType)
 				{
 					continue;
 				}
 
-				// We don't use the USD attribute path here because we want one unique node per UE track name,
-				// so that if e.g. both "intensity" and "exposure" are animated we make a single track for
-				// the Intensity UE property
-				const FString UniquePath = PrimPath + TEXT("\\") + UEAttrName.ToString();
-				const FString AnimTrackNodeUid = AnimationTrackPrefix + UniquePath;
-
-				const UInterchangeAnimationTrackNode* ExistingNode = Cast<UInterchangeAnimationTrackNode>(NodeContainer.GetNode(AnimTrackNodeUid));
-				if (ExistingNode)
-				{
-					continue;
-				}
+				const EInterchangeAnimationPayLoadType PayloadType = bIsBoolTrack ? EInterchangeAnimationPayLoadType::STEPCURVE
+																				  : EInterchangeAnimationPayLoadType::CURVE;
 
-				UInterchangeAnimationTrackNode* AnimTrackNode = NewObject<UInterchangeAnimationTrackNode>(&NodeContainer);
-				AnimTrackNode->InitializeNode(AnimTrackNodeUid, UniquePath, EInterchangeNodeContainerType::TranslatedAsset);
-				AnimTrackNode->SetCustomActorDependencyUid(*PrimPath);
-				AnimTrackNode->SetCustomPropertyTrack(*FoundTrackType);
-				AnimTrackNode->SetCustomAnimationPayloadKey(
-					UniquePath,
-					bIsBoolTrack ? EInterchangeAnimationPayLoadType::STEPCURVE : EInterchangeAnimationPayLoadType::CURVE
-				);
+				AddPropertyAnimationNode(PrimPath, UEPropertyName.ToString(), *FoundTrackType, PayloadType, Impl, NodeContainer);
+			}
+		}
+	}
 
-				NodeContainer.AddNode(AnimTrackNode);
+	// Some of the volume prim info is meant for the HeterogeneousVolume (HV) actor and the volumetric material,
+	// so we need to add it to the scene node (it's possible separate HV actors with different values for these
+	// end up sharing identical volume nodes)
+	void AddVolumeSceneNodeAttributes(
+		const UE::FUsdPrim& Prim,
+		UInterchangeSceneNode* SceneNode,
+		TArray<FString> AssetNodeUids,
+		const FString& VolumeMaterialInstanceUid,
+		bool bNeedsAnimationTrack,
+		UInterchangeUSDTranslatorImpl& Impl,
+		UInterchangeBaseNodeContainer& NodeContainer
+	)
+	{
+		if (!SceneNode)
+		{
+			return;
+		}
+		const FString& SceneNodeUid = SceneNode->GetUniqueID();
 
-				AddTrackSetNode(Impl, NodeContainer);
-				Impl.CurrentTrackSet->AddCustomAnimationTrackUid(AnimTrackNodeUid);
+		// Target the scene node from all the volume nodes to make it easy to find everything we need on the USD Pipeline
+		for (const FString& AssetNodeUid : AssetNodeUids)
+		{
+			if (const UInterchangeVolumeNode* VolumeNode = Cast<const UInterchangeVolumeNode>(NodeContainer.GetNode(AssetNodeUid)))
+			{
+				VolumeNode->AddTargetNodeUid(SceneNodeUid);
 			}
 		}
+
+		// Set our volumetric material as a "material override" directly on the scene node, which the USD Pipeline will also use
+		if (!VolumeMaterialInstanceUid.IsEmpty())
+		{
+			SceneNode->SetSlotMaterialDependencyUid(UE::Interchange::Volume::VolumetricMaterial, VolumeMaterialInstanceUid);
+		}
+
+		if (bNeedsAnimationTrack)
+		{
+			// Ideally we'd write some step curves, but Interchange doesn't support float step curves
+			const EInterchangeAnimationPayLoadType PayloadType = EInterchangeAnimationPayLoadType::CURVE;
+			const EInterchangePropertyTracks TrackType = EInterchangePropertyTracks::HeterogeneousVolumeFrame;
+			const static FName UEPropertyName = GET_MEMBER_NAME_CHECKED(UHeterogeneousVolumeComponent, Frame);
+
+			AddPropertyAnimationNode(SceneNodeUid, UEPropertyName.ToString(), TrackType, PayloadType, Impl, NodeContainer);
+		}
 	}
 
 	void AddSkeletalAnimationNode(
@@ -2487,7 +3033,7 @@ namespace UE::InterchangeUsdTranslator::Private
 			FWriteScopeLock ScopedInfoWriteLock{TranslatorImpl.CachedTraversalInfoLock};
 			TranslatorImpl.NodeUidToCachedTraversalInfo.Add(SkeletonPrimNodeUid, Info);
 		}
-#endif
+#endif	  // WITH_EDITOR
 	}
 
 	template<typename T>
@@ -2612,36 +3158,38 @@ namespace UE::InterchangeUsdTranslator::Private
 		// Do this before generating other nodes as they may need the updated info
 		Info.UpdateWithCurrentPrim(Prim);
 
-		// Generate asset node if applicable
-		const FString* Prefix = nullptr;
+		FString VolumeMaterialInstanceUid;
+		bool bNeedsVolumeTrack = false;
+
+		// Generate asset node(s) if applicable
+		TArray<FString> AssetNodeUids;
 		if (Prim.IsA(TEXT("Material")))
 		{
-			Prefix = &MaterialPrefix;
-			TranslatorImpl.AddMaterialNode(Prim, TranslatorSettings, NodeContainer);
+			AssetNodeUids = {TranslatorImpl.AddMaterialNode(Prim, TranslatorSettings, NodeContainer)};
 		}
 		else if (Prim.IsA(TEXT("Mesh")))
 		{
-			Prefix = &MeshPrefix;
-			TranslatorImpl.AddMeshNode(Prim, NodeContainer, Info);
+			AssetNodeUids = {TranslatorImpl.AddMeshNode(Prim, NodeContainer, Info)};
 		}
 		else if (Prim.IsA(TEXT("Camera")))
 		{
-			Prefix = &CameraPrefix;
-			AddCameraNode(Prim, NodeContainer);
+			AssetNodeUids = {AddCameraNode(Prim, NodeContainer)};
 		}
 		else if (Prim.HasAPI(TEXT("LightAPI")))
 		{
-			Prefix = &LightPrefix;
-			AddLightNode(Prim, NodeContainer);
+			AssetNodeUids = {AddLightNode(Prim, NodeContainer)};
 		}
 		else if (Prim.IsA(TEXT("Gprim")) && !Prim.IsA(TEXT("PointBased")) && !Prim.IsA(TEXT("Volume")))
 		{
-			//PointBased prims are currently not supported apart from Meshes (which is taken care of in a previous branch).
-			//Volumes are also not currently supported.
+			// PointBased prims are currently not supported apart from Meshes (which is taken care of in a previous branch).
+			// Volumes are also not currently supported.
 
-			Prefix = &PrimitiveShapePrefix;
 			constexpr bool bPrimitiveShape_True = true;
-			TranslatorImpl.AddMeshNode(Prim, NodeContainer, Info, bPrimitiveShape_True);
+			AssetNodeUids = {TranslatorImpl.AddMeshNode(Prim, NodeContainer, Info, bPrimitiveShape_True)};
+		}
+		else if (Prim.IsA(TEXT("Volume")))
+		{
+			AssetNodeUids = TranslatorImpl.AddVolumeNodes(Prim, NodeContainer, VolumeMaterialInstanceUid, bNeedsVolumeTrack);
 		}
 
 		// Only prims that require rendering (and have a renderable parent) get a scene node.
@@ -2707,7 +3255,7 @@ namespace UE::InterchangeUsdTranslator::Private
 				TranslatorImpl.AddLODMeshNodes(Prim, NodeContainer, *SceneNode, TranslatorSettings, Info);
 			}
 
-			if (Info.bInsideLOD && Prefix == &MeshPrefix)
+			if (Info.bInsideLOD && Prim.IsA(TEXT("Mesh")))
 			{
 				TranslatorImpl.CurrentLODSceneNodes.Add(SceneNode);
 			}
@@ -2718,11 +3266,23 @@ namespace UE::InterchangeUsdTranslator::Private
 				AddSkeletonNodes(Prim, TranslatorImpl, *SceneNode, NodeContainer, Info);
 			}
 
-			// Connect scene node and asset node
-			if (Prefix)
+			if (Prim.IsA(TEXT("Volume")))
+			{
+				AddVolumeSceneNodeAttributes(
+					Prim,
+					SceneNode,
+					AssetNodeUids,
+					VolumeMaterialInstanceUid,
+					bNeedsVolumeTrack,
+					TranslatorImpl,
+					NodeContainer
+				);
+			}
+
+			// Connect scene node and primary asset node
+			if (AssetNodeUids.Num() != 0)
 			{
-				const FString AssetNodeUid = *Prefix + SceneNodeUid;
-				SceneNode->SetCustomAssetInstanceUid(AssetNodeUid);
+				SceneNode->SetCustomAssetInstanceUid(AssetNodeUids[0]);
 			}
 
 			// Connect parent and child scene nodes
@@ -2930,7 +3490,7 @@ namespace UE::InterchangeUsdTranslator::Private
 			// to resolve all the bindings. The cached assignments come from AddMeshNode step, where
 			// we switch the active variant on the current stage and so get nice material bindings that
 			// resolve normally.
-		
+
 			// Note that we can't even use the info cache here, because it wouldn't have cached info
 			// about the inactive LOD variants
 			if (UsdUtils::FUsdPrimMaterialAssignmentInfo* CachedInfo = Impl.CachedMaterialAssignments.Find(PrimPath))
@@ -3012,8 +3572,77 @@ namespace UE::InterchangeUsdTranslator::Private
 		return UsdUtils::ApplyBlendShape(OutMeshDescription, BlendShape.GetPrim(), Weight, InbetweenName);
 	}
 
+	// Volume animations need a special property reader because in USD it's an animation of the file path attribute within
+	// the asset prims, while in UE we want a float track flipping through the volume frame indices
+	UsdToUnreal::FPropertyTrackReader CreateVolumeTrackReader(
+		const FString& InPrimPath,
+		const UInterchangeUSDTranslatorImpl& InImpl,
+		TArray<double>& OutTimeSampleUnion
+	)
+	{
+		UsdToUnreal::FPropertyTrackReader Result;
+
+		const TArray<UsdUtils::FVolumePrimInfo>* CollectedInfoForPrim = InImpl.PrimPathToVolumeInfo.Find(InPrimPath);
+		if (!CollectedInfoForPrim)
+		{
+			return Result;
+		}
+
+		const UsdUtils::FVolumePrimInfo* AnimatedInfo = nullptr;
+
+		for (const UsdUtils::FVolumePrimInfo& Info : *CollectedInfoForPrim)
+		{
+			if (Info.TimeSamplePathTimeCodes.Num() > 0)
+			{
+				if (AnimatedInfo != nullptr)
+				{
+					UE_LOG(
+						LogUsd,
+						Log,
+						TEXT(
+							"Only one animated SparseVolumeTexture can be driven via LevelSequences for each prim, for now. Prim '%s' has multiple, so the animation may not be correct."
+						),
+						*InPrimPath
+					);
+				}
+
+				AnimatedInfo = &Info;
+			}
+		}
+
+		if (AnimatedInfo)
+		{
+			// Detect whether we inserted the default opinion volume as the first frame of the animation
+			bool bInsertedDefaultOpinion = AnimatedInfo->TimeSamplePaths.Num() > 0
+										   && AnimatedInfo->TimeSamplePaths[0] != AnimatedInfo->SourceVDBFilePath;
+
+			OutTimeSampleUnion = AnimatedInfo->TimeSamplePathTimeCodes;
+			Result.FloatReader = [AnimatedInfo, bInsertedDefaultOpinion](double TimeCode) -> float
+			{
+				int32 FrameIndex = 0;
+				for (; FrameIndex + 1 < AnimatedInfo->TimeSamplePathTimeCodes.Num(); ++FrameIndex)
+				{
+					if (AnimatedInfo->TimeSamplePathTimeCodes[FrameIndex + 1] > TimeCode)
+					{
+						break;
+					}
+				}
+				FrameIndex = FMath::Clamp(FrameIndex, 0, AnimatedInfo->TimeSamplePathTimeCodes.Num() - 1);
+
+				if (bInsertedDefaultOpinion)
+				{
+					FrameIndex += 1;
+				}
+
+				return static_cast<float>(FrameIndex);
+			};
+		}
+
+		return Result;
+	}
+
 	bool GetPropertyAnimationCurvePayloadData(
-		const UE::FUsdStage& UsdStage,
+		const UInterchangeUSDTranslatorImpl& Impl,
 		const FString& PayloadKey,
 		Interchange::FAnimationPayloadData& OutPayloadData
 	)
@@ -3026,7 +3655,7 @@ namespace UE::InterchangeUsdTranslator::Private
 			return false;
 		}
 
-		UE::FUsdPrim Prim = UsdStage.GetPrimAtPath(UE::FSdfPath{*PrimPath});
+		UE::FUsdPrim Prim = Impl.UsdStage.GetPrimAtPath(UE::FSdfPath{*PrimPath});
 		FName UEPropertyName = *UEPropertyNameStr;
 		if (!Prim || UEPropertyName == NAME_None)
 		{
@@ -3034,30 +3663,40 @@ namespace UE::InterchangeUsdTranslator::Private
 		}
 
 		TArray<double> TimeSampleUnion;
-		TArray<UE::FUsdAttribute> Attrs = UsdUtils::GetAttributesForProperty(Prim, UEPropertyName);
-		bool bSuccess = UE::FUsdAttribute::GetUnionedTimeSamples(Attrs, TimeSampleUnion);
-		if (!bSuccess)
+		UsdToUnreal::FPropertyTrackReader Reader;
+
+		if (Prim.IsA(TEXT("Volume")) && UEPropertyNameStr == GET_MEMBER_NAME_CHECKED(UHeterogeneousVolumeComponent, Frame))
 		{
-			return false;
+			Reader = CreateVolumeTrackReader(PrimPath, Impl, TimeSampleUnion);
+		}
+		else
+		{
+			TArray<UE::FUsdAttribute> Attrs = UsdUtils::GetAttributesForProperty(Prim, UEPropertyName);
+			bool bSuccess = UE::FUsdAttribute::GetUnionedTimeSamples(Attrs, TimeSampleUnion);
+			if (!bSuccess)
+			{
+				return false;
+			}
+
+			const bool bIgnorePrimLocalTransform = false;
+			Reader = UsdToUnreal::CreatePropertyTrackReader(Prim, UEPropertyName, bIgnorePrimLocalTransform);
 		}
 
-		const bool bIgnorePrimLocalTransform = false;
-		UsdToUnreal::FPropertyTrackReader Reader = UsdToUnreal::CreatePropertyTrackReader(Prim, UEPropertyName, bIgnorePrimLocalTransform);
 		if (Reader.BoolReader)
 		{
-			return ReadBools(UsdStage, TimeSampleUnion, Reader.BoolReader, OutPayloadData);
+			return ReadBools(Impl.UsdStage, TimeSampleUnion, Reader.BoolReader, OutPayloadData);
 		}
 		else if (Reader.ColorReader)
 		{
-			return ReadColors(UsdStage, TimeSampleUnion, Reader.ColorReader, OutPayloadData);
+			return ReadColors(Impl.UsdStage, TimeSampleUnion, Reader.ColorReader, OutPayloadData);
 		}
 		else if (Reader.FloatReader)
 		{
-			return ReadFloats(UsdStage, TimeSampleUnion, Reader.FloatReader, OutPayloadData);
+			return ReadFloats(Impl.UsdStage, TimeSampleUnion, Reader.FloatReader, OutPayloadData);
 		}
 		else if (Reader.TransformReader)
 		{
-			return ReadTransforms(UsdStage, TimeSampleUnion, Reader.TransformReader, OutPayloadData);
+			return ReadTransforms(Impl.UsdStage, TimeSampleUnion, Reader.TransformReader, OutPayloadData);
 		}
 
 		return false;
@@ -3349,6 +3988,10 @@ namespace UE::InterchangeUsdTranslator::Private
 
 		return true;
 	}
+#else
+	class UInterchangeUSDTranslatorImpl
+	{
+	};
 #endif	  // USE_USD_SDK
 }	 // namespace UE::InterchangeUsdTranslator::Private
 
@@ -3400,6 +4043,9 @@ bool UInterchangeUSDTranslator::Translate(UInterchangeBaseNodeContainer& NodeCon
 	using namespace UE::InterchangeUsdTranslator::Private;
 	using namespace UsdToUnreal;
 
+	// Reset impl as we don't want to share internal state with a previous translation, as the file
+	// (or our settings) may have changed, which should lead to different nodes / data
+	Impl = MakeUnique<UE::InterchangeUsdTranslator::Private::UInterchangeUSDTranslatorImpl>();
 	UInterchangeUSDTranslatorImpl* ImplPtr = Impl.Get();
 	if (!ImplPtr)
 	{
@@ -3536,21 +4182,21 @@ bool UInterchangeUSDTranslator::Translate(UInterchangeBaseNodeContainer& NodeCon
 
 void UInterchangeUSDTranslator::ReleaseSource()
 {
+#if USE_USD_SDK
 	UE::InterchangeUsdTranslator::Private::UInterchangeUSDTranslatorImpl* ImplPtr = Impl.Get();
-	if (!ImplPtr)
+	if (ImplPtr)
 	{
-		return;
-	}
+		ImplPtr->TranslationContext.Reset();	// This holds a reference to the stage
+		ImplPtr->UsdStage = UE::FUsdStage{};
+		ImplPtr->InfoCache = nullptr;
+		ImplPtr->CurrentTrackSet = nullptr;
 
-	ImplPtr->TranslationContext.Reset();	// This holds a reference to the stage
-	ImplPtr->UsdStage = UE::FUsdStage{};
-	ImplPtr->InfoCache = nullptr;
-	ImplPtr->CurrentTrackSet = nullptr;
-
-	{
-		FWriteScopeLock Lock{ImplPtr->PrimPathToVariantToStageLock};
-		ImplPtr->PrimPathToVariantToStage.Reset();
+		{
+			FWriteScopeLock Lock{ImplPtr->PrimPathToVariantToStageLock};
+			ImplPtr->PrimPathToVariantToStage.Reset();
+		}
 	}
+#endif	  // USE_USD_SDK
 
 	if (TranslatorSettings)
 	{
@@ -3576,12 +4222,6 @@ UInterchangeTranslatorSettings* UInterchangeUSDTranslator::GetSettings() const
 {
 	using namespace UE::InterchangeUsdTranslator::Private;
 
-	UInterchangeUSDTranslatorImpl* ImplPtr = Impl.Get();
-	if (!ImplPtr)
-	{
-		return nullptr;
-	}
-
 	if (!TranslatorSettings)
 	{
 		TranslatorSettings = DuplicateObject<UInterchangeUsdTranslatorSettings>(
@@ -3600,12 +4240,6 @@ void UInterchangeUSDTranslator::SetSettings(const UInterchangeTranslatorSettings
 {
 	using namespace UE::InterchangeUsdTranslator::Private;
 
-	UInterchangeUSDTranslatorImpl* ImplPtr = Impl.Get();
-	if (!ImplPtr)
-	{
-		return;
-	}
-
 	if (TranslatorSettings)
 	{
 		TranslatorSettings->ClearFlags(RF_Standalone);
@@ -3763,6 +4397,7 @@ TOptional<UE::Interchange::FImportImage> UInterchangeUSDTranslator::GetTexturePa
 	using namespace UE::InterchangeUsdTranslator::Private;
 
 	TOptional<UE::Interchange::FImportImage> TexturePayloadData;
+#if USE_USD_SDK
 
 	UInterchangeUSDTranslatorImpl* ImplPtr = Impl.Get();
 	if (!ImplPtr)
@@ -3770,7 +4405,6 @@ TOptional<UE::Interchange::FImportImage> UInterchangeUSDTranslator::GetTexturePa
 		return TexturePayloadData;
 	}
 
-#if USE_USD_SDK
 	FString FilePath;
 	TextureGroup TextureGroup;
 	bool bDecoded = DecodeTexturePayloadKey(PayloadKey, FilePath, TextureGroup);
@@ -3795,16 +4429,15 @@ TOptional<UE::Interchange::FImportImage> UInterchangeUSDTranslator::GetTexturePa
 			TexturePayloadData->CompressionSettings = TextureGroup == TEXTUREGROUP_WorldNormalMap ? TC_Normalmap : TC_Default;
 		}
 	}
-#endif	  // USE_USD_SDK
 
 	// We did not find a suitable Payload in USD Translator, let's find one in one of the Translators (MaterialX for the moment)
 	// The best way would be to have a direct association between the payload and the right Translator, but we don't have a suitable way of knowing
 	// which Payload belongs to which Translator So let's just loop over them all
 	if (!TexturePayloadData.IsSet())
 	{
-		for (const TPair<FString, UInterchangeTranslatorBase*>& Pair : Impl->Translators)
+		for (const TPair<FString, TStrongObjectPtr<UInterchangeTranslatorBase>>& Pair : Impl->Translators)
 		{
-			if (IInterchangeTexturePayloadInterface* TexturePayloadInterface = Cast<IInterchangeTexturePayloadInterface>(Pair.Value))
+			if (IInterchangeTexturePayloadInterface* TexturePayloadInterface = Cast<IInterchangeTexturePayloadInterface>(Pair.Value.Get()))
 			{
 				TexturePayloadData = TexturePayloadInterface->GetTexturePayloadData(PayloadKey, AlternateTexturePath);
 				if (TexturePayloadData.IsSet())
@@ -3822,6 +4455,7 @@ TOptional<UE::Interchange::FImportImage> UInterchangeUSDTranslator::GetTexturePa
 		TexturePayloadData = GetMaterialXTexturePayloadData(PayloadKey, AlternateTexturePath, this, *ImplPtr);
 	}
 
+#endif	  // USE_USD_SDK
 	return TexturePayloadData;
 }
 
@@ -3832,15 +4466,15 @@ TOptional<UE::Interchange::FImportBlockedImage> UInterchangeUSDTranslator::GetBl
 {
 	using namespace UE::InterchangeUsdTranslator::Private;
 
+	UE::Interchange::FImportBlockedImage BlockData;
+
+#if USE_USD_SDK
 	UInterchangeUSDTranslatorImpl* ImplPtr = Impl.Get();
 	if (!ImplPtr)
 	{
 		return {};
 	}
 
-	UE::Interchange::FImportBlockedImage BlockData;
-
-#if USE_USD_SDK
 	FString FilePath;
 	TextureGroup TextureGroup;
 	bool bDecoded = DecodeTexturePayloadKey(PayloadKey, FilePath, TextureGroup);
@@ -3930,6 +4564,7 @@ TArray<UE::Interchange::FAnimationPayloadData> UInterchangeUSDTranslator::GetAni
 	// This is the results we return
 	TArray<UE::Interchange::FAnimationPayloadData> AnimationPayloads;
 
+#if USE_USD_SDK
 	// Maps to help sorting the queries by payload type
 	TArray<int32> BakeQueryIndexes;
 	TArray<TArray<UE::Interchange::FAnimationPayloadData>> BakeAnimationPayloads;
@@ -3952,8 +4587,6 @@ TArray<UE::Interchange::FAnimationPayloadData> UInterchangeUSDTranslator::GetAni
 		}
 	}
 
-#if USE_USD_SDK
-
 	// Import the Baked curve payloads
 	if (BakeQueryIndexes.Num() > 0)
 	{
@@ -4019,7 +4652,7 @@ TArray<UE::Interchange::FAnimationPayloadData> UInterchangeUSDTranslator::GetAni
 				// We're fine handling these in isolation (currently GetAnimationPayloadData is called with
 				// a single query at a time for these): Emit a separate task for each right away
 				FAnimationPayloadData AnimationPayLoadData{PayloadQuery.SceneNodeUniqueID, PayloadQuery.PayloadKey};
-				if (GetPropertyAnimationCurvePayloadData(Impl->UsdStage, PayloadQuery.PayloadKey.UniqueId, AnimationPayLoadData))
+				if (GetPropertyAnimationCurvePayloadData(*Impl, PayloadQuery.PayloadKey.UniqueId, AnimationPayLoadData))
 				{
 					CurveAnimationPayloads[PayloadIndex].Emplace(AnimationPayLoadData);
 				}
@@ -4083,4 +4716,19 @@ TArray<UE::Interchange::FAnimationPayloadData> UInterchangeUSDTranslator::GetAni
 	return AnimationPayloads;
 }
 
+TOptional<UE::Interchange::FVolumePayloadData> UInterchangeUSDTranslator::GetVolumePayloadData(const UE::Interchange::FVolumePayloadKey& PayloadKey
+) const
+{
+#if USE_USD_SDK
+	TStrongObjectPtr<UInterchangeTranslatorBase> ExistingTranslator = Impl->Translators.FindRef(PayloadKey.FileName);
+	const IInterchangeVolumePayloadInterface* VolumeInterface = Cast<IInterchangeVolumePayloadInterface>(ExistingTranslator.Get());
+	if (VolumeInterface)
+	{
+		return VolumeInterface->GetVolumePayloadData(PayloadKey);
+	}
+#endif	  // USE_USD_SDK
+
+	return {};
+}
+
 #undef LOCTEXT_NAMESPACE
diff --git a/Engine/Plugins/Interchange/Runtime/Source/Import/Private/Volume/InterchangeSparseVolumeTextureFactory.cpp b/Engine/Plugins/Interchange/Runtime/Source/Import/Private/Volume/InterchangeSparseVolumeTextureFactory.cpp
--- a/Engine/Plugins/Interchange/Runtime/Source/Import/Private/Volume/InterchangeSparseVolumeTextureFactory.cpp
+++ b/Engine/Plugins/Interchange/Runtime/Source/Import/Private/Volume/InterchangeSparseVolumeTextureFactory.cpp
@@ -72,7 +72,7 @@ namespace UE::InterchangeSparseVolumeTextureFactory::Private
 		}
 
 		TArray<FString> GridNodeUids;
-		VolumeNode->GetGridDependencies(GridNodeUids);
+		VolumeNode->GetCustomGridDependecies(GridNodeUids);
 
 		TArray<const UInterchangeVolumeGridNode*> Result;
 		Result.Reserve(GridNodeUids.Num());
@@ -133,12 +133,7 @@ namespace UE::InterchangeSparseVolumeTextureFactory::Private
 	}
 #endif	  // WITH_EDITOR
 
-	// Actually fills in the provided SVT asset with the payload / source data
-	void SetupSparseVolumeTexture(
-		USparseVolumeTexture* SparseVolumeTexture,
-		TArray<TOptional<UE::Interchange::FVolumePayloadData>>& ProcessedPayloads,
-		bool bIsReimport
-	)
+	void FillSparseVolumeTextureWithPayloadData(USparseVolumeTexture* SparseVolumeTexture, TArray<FVolumePayload>& ProcessedPayloads)
 	{
 		UStreamableSparseVolumeTexture* Streamable = Cast<UStreamableSparseVolumeTexture>(SparseVolumeTexture);
 		int32 NumPayloads = ProcessedPayloads.Num();
@@ -153,14 +148,62 @@ namespace UE::InterchangeSparseVolumeTextureFactory::Private
 			return;
 		}
 
-		for (TOptional<UE::Interchange::FVolumePayloadData>& ProcessedPayload : ProcessedPayloads)
+		// Setup some temp structure to help us sort the animation frames, if any
+		struct FPayloadAndIndex
+		{
+			FVolumePayload* Payload = nullptr;
+			int32 FrameIndex;
+		};
+
+		TArray<FPayloadAndIndex> PayloadAndIndices;
+		PayloadAndIndices.Reserve(ProcessedPayloads.Num());
+
+		for (FVolumePayload& VolumePayload : ProcessedPayloads)
+		{
+			if (VolumePayload.PayloadFrameIndices.Num() == 0)
+			{
+				// No animation index --> Not animated, so pretend it's the earliest frame possible
+				PayloadAndIndices.Emplace(&VolumePayload, TNumericLimits<int32>::Lowest());
+			}
+			else
+			{
+				for (const int32 FrameIndex : VolumePayload.PayloadFrameIndices)
+				{
+					PayloadAndIndices.Emplace(&VolumePayload, FrameIndex);
+				}
+			}
+		}
+
+		PayloadAndIndices.Sort(
+			[](const FPayloadAndIndex& LHS, const FPayloadAndIndex& RHS)
+			{
+				if (LHS.FrameIndex == RHS.FrameIndex)
+				{
+					// Fallback compare on the volume node UID for consistency, or else the order could change for identical imports
+					return LHS.Payload->VolumeNodeUid < RHS.Payload->VolumeNodeUid;
+				}
+
+				return LHS.FrameIndex < RHS.FrameIndex;
+			}
+		);
+
+		// Note that we have no deduplication past this point: If a frame is repeated in the animated SVT, it will be added to the
+		// SVT twice (we did only fetch the payload for it once though).
+		//
+		// We could maybe add this in the future, but it's probably a premature optimization that is just more prone to lead to confusion
+		// than anything else. This because the way we compensate for removing those duplicated frames from the SVT is by playing the
+		// unique frame multiple times on the LevelSequence. We can't help with that if the import is asset only (or if LevelSequences
+		// aren't being imported, etc.), so it's probably best to just always add all the frames and make the animated SVT asset work
+		// as intended on its own
+		for (FPayloadAndIndex& PayloadAndIndex : PayloadAndIndices)
 		{
-			if (!ProcessedPayload.IsSet())
+			TOptional<UE::Interchange::FVolumePayloadData>& PayloadData = PayloadAndIndex.Payload->PayloadData;
+			if (!PayloadData.IsSet())
 			{
 				continue;
 			}
 
-			bSucess = Streamable->AppendFrame(ProcessedPayload->TextureData, ProcessedPayload->Transform);
+			bSucess = Streamable->AppendFrame(PayloadData->TextureData, PayloadData->Transform);
 			if (!bSucess)
 			{
 				return;
@@ -203,7 +246,7 @@ namespace UE::InterchangeSparseVolumeTextureFactory::Private
 			// If we're importing multiple volume nodes, it's implied that they're separate frames of the same animated SparseVolumeTexture,
 			// and that means they should have the same grid and assignment, so picking any volume node should do
 			TArray<FString> GridNodeUids;
-			VolumeNodes[0]->GetGridDependencies(GridNodeUids);
+			VolumeNodes[0]->GetCustomGridDependecies(GridNodeUids);
 
 			OrderedGridNodes.Reserve(GridNodeUids.Num());
 			for (const FString& GridNodeUid : GridNodeUids)
@@ -228,13 +271,13 @@ namespace UE::InterchangeSparseVolumeTextureFactory::Private
 			EInterchangeSparseVolumeTextureFormat Format;
 
 			Result.Attributes[0].Format = EInterchangeSparseVolumeTextureFormat::Float16;
-			if (FactoryNode->GetCustomAttributesADataType(Format))
+			if (FactoryNode->GetCustomAttributesAFormat(Format))
 			{
 				Result.Attributes[0].Format = Format;
 			}
 
 			Result.Attributes[1].Format = EInterchangeSparseVolumeTextureFormat::Float16;
-			if (FactoryNode->GetCustomAttributesBDataType(Format))
+			if (FactoryNode->GetCustomAttributesBFormat(Format))
 			{
 				Result.Attributes[1].Format = Format;
 			}
@@ -545,6 +588,40 @@ UInterchangeFactoryBase::FImportAssetResult UInterchangeSparseVolumeTextureFacto
 			return ImportAssetResult;
 		}
 
+#if WITH_EDITOR
+		// We have to collect all the info we can from the previous asset before we stomp it for reimport, as we will need
+		// to apply it back to the new asset
+		TStrongObjectPtr<UAssetImportData> OldAssetImportData = nullptr;
+#endif	  // WITH_EDITOR
+		TArray<TStrongObjectPtr<UAssetUserData>> OldAssetUserData;
+
+		if (UStreamableSparseVolumeTexture* OriginalTexture = Cast<UStreamableSparseVolumeTexture>(Arguments.ReimportObject))
+		{
+#if WITH_EDITOR
+			if (OriginalTexture->AssetImportData)
+			{
+				OldAssetImportData.Reset(OriginalTexture->AssetImportData);
+
+				UObject* NewOuter = GetTransientPackage();
+				FName NewName = MakeUniqueObjectName(NewOuter, OriginalTexture->AssetImportData->GetClass());
+				ensure(OldAssetImportData->Rename(*NewName.ToString(), NewOuter, REN_DontCreateRedirectors | REN_NonTransactional | REN_DoNotDirty));
+
+				OriginalTexture->AssetImportData = nullptr;
+			}
+#endif	  // WITH_EDITOR
+
+			OldAssetUserData.Reserve(OriginalTexture->AssetUserData.Num());
+			for (UAssetUserData* UserData : OriginalTexture->AssetUserData)
+			{
+				OldAssetUserData.Emplace(UserData);
+
+				UObject* NewOuter = GetTransientPackage();
+				FName NewName = MakeUniqueObjectName(NewOuter, UserData->GetClass());
+				ensure(UserData->Rename(*NewName.ToString(), NewOuter, REN_DontCreateRedirectors | REN_NonTransactional | REN_DoNotDirty));
+			}
+			OldAssetUserData.Reset();
+		}
+
 		SparseVolumeTexture = NewObject<USparseVolumeTexture>(Arguments.Parent, ClassToUse, *Arguments.AssetName, RF_Public | RF_Standalone);
 		if (!SparseVolumeTexture)
 		{
@@ -552,6 +629,34 @@ UInterchangeFactoryBase::FImportAssetResult UInterchangeSparseVolumeTextureFacto
 			return ImportAssetResult;
 		}
 
+#if WITH_EDITOR
+		// Let's set these back right away and pretend they were always there, as that's what Interchange will expect of a reimported asset
+		if (UStreamableSparseVolumeTexture* NewTexture = Cast<UStreamableSparseVolumeTexture>(SparseVolumeTexture))
+		{
+			if (OldAssetImportData)
+			{
+				NewTexture->AssetImportData = OldAssetImportData.Get();
+
+				UObject* NewOuter = NewTexture;
+				FName NewName = MakeUniqueObjectName(NewOuter, OldAssetImportData->GetClass());
+				ensure(OldAssetImportData->Rename(*NewName.ToString(), NewOuter, REN_DontCreateRedirectors | REN_NonTransactional | REN_DoNotDirty));
+			}
+
+			if (OldAssetUserData.Num() > 0)
+			{
+				NewTexture->AssetUserData.Reserve(OldAssetUserData.Num() + NewTexture->AssetUserData.Num());
+				for (const TStrongObjectPtr<UAssetUserData>& UserData : OldAssetUserData)
+				{
+					NewTexture->AssetUserData.Add(UserData.Get());
+
+					UObject* NewOuter = NewTexture;
+					FName NewName = MakeUniqueObjectName(NewOuter, UserData->GetClass());
+					ensure(UserData->Rename(*NewName.ToString(), NewOuter, REN_DontCreateRedirectors | REN_NonTransactional | REN_DoNotDirty));
+				}
+			}
+		}
+#endif	  // WITH_EDITOR
+
 		ImportAssetResult.ImportedObject = SparseVolumeTexture;
 	}
 
@@ -656,14 +761,24 @@ UInterchangeFactoryBase::FImportAssetResult UInterchangeSparseVolumeTextureFacto
 	FIntVector3 VolumeBoundsMax;
 	ComputeExpandedVolumeBounds(VolumeNodes, Arguments.NodeContainer, VolumeBoundsMin, VolumeBoundsMax);
 
-	// Collect the payload keys from each volume node
-	bool bAtLeastOneValidPayload = false;
+	bool bAtLeastOneValidPayloadKey = false;
+
+	TSet<int32> FrameIndices;
+
 	TArray<FVolumePayloadKey> PayloadKeys;
 	PayloadKeys.Reserve(VolumeNodes.Num());
+
+	TArray<FVolumePayload> VolumePayloads;
+	VolumePayloads.Reserve(VolumeNodes.Num());
+
+	// Collect the payload keys from each volume node
 	for (const UInterchangeVolumeNode* VolumeNode : VolumeNodes)
 	{
 		// Emplace even if we fail so that we keep a one-to-one match with the VolumeNodes array
 		FVolumePayloadKey& Key = PayloadKeys.Emplace_GetRef();
+		FVolumePayload& Payload = VolumePayloads.Emplace_GetRef();
+
+		Payload.VolumeNodeUid = VolumeNode->GetUniqueID();
 
 		FString FileName;
 		bool bSuccess = VolumeNode->GetCustomFileName(FileName);
@@ -672,12 +787,25 @@ UInterchangeFactoryBase::FImportAssetResult UInterchangeSparseVolumeTextureFacto
 			continue;
 		}
 
-		bAtLeastOneValidPayload = true;
+		bAtLeastOneValidPayloadKey = true;
 		Key.FileName = FileName;
 		Key.AssignmentInfo = AssignmentInfo.GetValue();
 		Key.VolumeBoundsMin = VolumeBoundsMin;
+
+		VolumeNode->GetCustomFrameIndicesInAnimation(Payload.PayloadFrameIndices);
+
+		if (Payload.PayloadFrameIndices.Num() == 0)
+		{
+			// No animated frames --> Volume node wants to become a static SVT. Let's consider it as one frame,
+			// as that is how we'll process its payload later in FillSparseVolumeTextureWithPayloadData()
+			FrameIndices.Add(TNumericLimits<int32>::Lowest());
+		}
+		else
+		{
+			FrameIndices.Append(Payload.PayloadFrameIndices);
+		}
 	}
-	if (!bAtLeastOneValidPayload)
+	if (!bAtLeastOneValidPayloadKey)
 	{
 		HandleFailure(
 			LOCTEXT("SparseVolumeTextureFactory_InvalidPayloadKey", "None of the translated UInterchangeVolumeGridNodes has a valid payload key.")
@@ -685,32 +813,27 @@ UInterchangeFactoryBase::FImportAssetResult UInterchangeSparseVolumeTextureFacto
 		return ImportAssetResult;
 	}
 
-	// TODO: Handle this properly when implementing the USD translator's usage of the OpenVDB translator
-	AlternateTexturePath.Reset();
-
 	// Fetch the actual payloads in parallel (presumably they're all different files, but it should work even in the edge case
 	// that the same file shows up more than once somehow)
-	TArray<TOptional<FVolumePayloadData>> PayloadData;
-	PayloadData.SetNum(VolumeNodes.Num());
 	const int32 NumBatches = PayloadKeys.Num();
 	ParallelFor(
 		NumBatches,
-		[&PayloadKeys, &PayloadData, TranslatorInterface](int32 Index)
+		[&PayloadKeys, &VolumePayloads, TranslatorInterface](int32 Index)
 		{
-			PayloadData[Index] = TranslatorInterface->GetVolumePayloadData(PayloadKeys[Index]);
+			VolumePayloads[Index].PayloadData = TranslatorInterface->GetVolumePayloadData(PayloadKeys[Index]);
 		}
 	);
 
-	bool bAtLeastOnePayloadIsValid = false;
-	for (const TOptional<FVolumePayloadData>& PayloadDataEntry : PayloadData)
+	bool bAtLeastOneValidPayloadData = false;
+	for (const FVolumePayload& VolumePayload : VolumePayloads)
 	{
-		if (PayloadDataEntry.IsSet())
+		if (VolumePayload.PayloadData.IsSet())
 		{
-			bAtLeastOnePayloadIsValid = true;
+			bAtLeastOneValidPayloadData = true;
 			break;
 		}
 	}
-	if (!bAtLeastOnePayloadIsValid)
+	if (!bAtLeastOneValidPayloadData)
 	{
 		HandleFailure(LOCTEXT(
 			"SparseVolumeTextureFactory_Async_CannotRetrievePayload",
@@ -720,14 +843,14 @@ UInterchangeFactoryBase::FImportAssetResult UInterchangeSparseVolumeTextureFacto
 	}
 
 	// Check if the payload is valid for the Texture
-	bool bCanSetup = bAtLeastOnePayloadIsValid;
+	bool bCanSetup = bAtLeastOneValidPayloadData;
 	if (UStaticSparseVolumeTexture* StaticVolumeTexture = Cast<UStaticSparseVolumeTexture>(SparseVolumeTexture))
 	{
-		bCanSetup = bCanSetup && PayloadKeys.Num() == 1;
+		bCanSetup = bCanSetup && FrameIndices.Num() == 1;
 	}
 	else if (UAnimatedSparseVolumeTexture* AnimatedVolumeTexture = Cast<UAnimatedSparseVolumeTexture>(SparseVolumeTexture))
 	{
-		bCanSetup = bCanSetup && PayloadKeys.Num() >= 1;
+		bCanSetup = bCanSetup && FrameIndices.Num() >= 1;
 	}
 	if (!bCanSetup)
 	{
@@ -749,7 +872,7 @@ UInterchangeFactoryBase::FImportAssetResult UInterchangeSparseVolumeTextureFacto
 	}
 #endif	  // WITH_EDITOR
 
-	ProcessedPayloads = MoveTemp(PayloadData);
+	ProcessedPayloads = MoveTemp(VolumePayloads);
 	return ImportAssetResult;
 }
 
@@ -762,106 +885,99 @@ UInterchangeFactoryBase::FImportAssetResult UInterchangeSparseVolumeTextureFacto
 	using namespace UE::InterchangeSparseVolumeTextureFactory::Private;
 
 	FImportAssetResult ImportAssetResult;
-	ImportAssetResult.bIsFactorySkipAsset = bSkipImport;
 
-	if (!Arguments.AssetNode)
+	// If we're not meant to modify the asset just return early
+	if (bSkipImport)
 	{
+		ImportAssetResult.bIsFactorySkipAsset = true;
 		return ImportAssetResult;
 	}
 
-	UObject* ExistingAsset = nullptr;
-	FSoftObjectPath ReferenceObject;
-	if (Arguments.AssetNode->GetCustomReferenceObject(ReferenceObject))
-	{
-		ExistingAsset = ReferenceObject.TryLoad();
-	}
-	if (!ExistingAsset)
+	// We only handle SVT factory nodes
+	UInterchangeSparseVolumeTextureFactoryNode* FactoryNode = Cast<UInterchangeSparseVolumeTextureFactoryNode>(Arguments.AssetNode);
+	if (!FactoryNode)
 	{
 		return ImportAssetResult;
 	}
 
-	if (USparseVolumeTexture* SparseVolumeTexture = Cast<USparseVolumeTexture>(ExistingAsset))
+	// Get the asset we're currently importing
+	USparseVolumeTexture* ImportedSparseVolumeTexture = nullptr;
 	{
-		ImportAssetResult.ImportedObject = SparseVolumeTexture;
-
-		if (!HasValidPayloads())
+		FSoftObjectPath ReferenceObject;
+		if (FactoryNode->GetCustomReferenceObject(ReferenceObject))
 		{
-			if (!Arguments.ReimportObject)
-			{
-				// Abandon asset
-				ImportAssetResult.ImportedObject = nullptr;
-				SparseVolumeTexture->RemoveFromRoot();
-				SparseVolumeTexture->MarkAsGarbage();
-				Arguments.AssetNode->SetCustomReferenceObject(nullptr);
-			}
+			ImportedSparseVolumeTexture = Cast<USparseVolumeTexture>(ReferenceObject.TryLoad());
+		}
 
-			Arguments.AssetNode->SetEnabled(false);
-			bSkipImport = true;
-			ImportAssetResult.bIsFactorySkipAsset = bSkipImport;
+		if (!ImportedSparseVolumeTexture)
+		{
+			return ImportAssetResult;
 		}
 	}
 
-	return ImportAssetResult;
-}
-
-void UInterchangeSparseVolumeTextureFactory::SetupObject_GameThread(const FSetupObjectParams& Arguments)
-{
-	TRACE_CPUPROFILER_EVENT_SCOPE(UInterchangeSparseVolumeTextureFactory::SetupObject_GameThread);
+	const bool bIsReimport = Arguments.ReimportObject != nullptr;
 
-	using namespace UE::InterchangeSparseVolumeTextureFactory::Private;
-
-	if (bSkipImport)
+	if (HasValidPayloads())
 	{
-		return;
+		ImportAssetResult.ImportedObject = ImportedSparseVolumeTexture;
 	}
+	// Abandon asset if it failed to import
+	else
+	{
+		// Dispose of asset if it's not a reimport
+		if (!bIsReimport)
+		{
+			ImportedSparseVolumeTexture->RemoveFromRoot();
+			ImportedSparseVolumeTexture->MarkAsGarbage();
+			FactoryNode->SetCustomReferenceObject(nullptr);
+		}
 
-	UInterchangeSparseVolumeTextureFactoryNode* SparseVolumeTextureFactoryNode = Cast<UInterchangeSparseVolumeTextureFactoryNode>(
-		Arguments.FactoryNode
-	);
+		FactoryNode->SetEnabled(false);
 
-	USparseVolumeTexture* SparseVolumeTexture = Cast<USparseVolumeTexture>(Arguments.ImportedObject);
+		bSkipImport = true;
+		ImportAssetResult.ImportedObject = nullptr;
+		ImportAssetResult.bIsFactorySkipAsset = bSkipImport;
+		return ImportAssetResult;
+	}
 
-	if (SparseVolumeTexture && SparseVolumeTextureFactoryNode && HasValidPayloads())
-	{
-		SetupSparseVolumeTexture(SparseVolumeTexture, ProcessedPayloads, Arguments.bIsReimport);
+	FillSparseVolumeTextureWithPayloadData(ImportedSparseVolumeTexture, ProcessedPayloads);
 
 #if WITH_EDITOR
-		if (!Arguments.bIsReimport)
-		{
-			SparseVolumeTextureFactoryNode->ApplyAllCustomAttributeToObject(SparseVolumeTexture);
-		}
-		else
+	if (bIsReimport)
+	{
+		if (UStreamableSparseVolumeTexture* OriginalTexture = Cast<UStreamableSparseVolumeTexture>(Arguments.ReimportObject))
 		{
-			UInterchangeAssetImportData* InterchangeAssetImportData = nullptr;
+			if (UInterchangeAssetImportData* AssetImportData = Cast<UInterchangeAssetImportData>(OriginalTexture->AssetImportData))
 			{
-				UStreamableSparseVolumeTexture* Streamable = Cast<UStreamableSparseVolumeTexture>(SparseVolumeTexture);
-				InterchangeAssetImportData = Cast<UInterchangeAssetImportData>(Streamable->AssetImportData);
-			}
+				UInterchangeFactoryBaseNode* PreviousNode = AssetImportData->GetStoredFactoryNode(AssetImportData->NodeUniqueID);
 
-			UInterchangeFactoryBaseNode* PreviousNode = nullptr;
-			if (InterchangeAssetImportData)
-			{
-				PreviousNode = InterchangeAssetImportData->GetStoredFactoryNode(InterchangeAssetImportData->NodeUniqueID);
-			}
+				// Create a factory node filled with the property values we can extract from our new imported asset
+				UInterchangeFactoryBaseNode* CurrentNode = NewObject<UInterchangeSparseVolumeTextureFactoryNode>(GetTransientPackage());
+				UInterchangeBaseNode::CopyStorage(FactoryNode, CurrentNode);
+				CurrentNode->FillAllCustomAttributeFromObject(ImportedSparseVolumeTexture);
 
-			UInterchangeFactoryBaseNode* CurrentNode = NewObject<UInterchangeSparseVolumeTextureFactoryNode>(GetTransientPackage());
-			UInterchangeBaseNode::CopyStorage(SparseVolumeTextureFactoryNode, CurrentNode);
-			CurrentNode->FillAllCustomAttributeFromObject(SparseVolumeTexture);
-
-			UE::Interchange::FFactoryCommon::ApplyReimportStrategyToAsset(
-				SparseVolumeTexture,
-				PreviousNode,
-				CurrentNode,
-				SparseVolumeTextureFactoryNode
-			);
+				UE::Interchange::FFactoryCommon::ApplyReimportStrategyToAsset(ImportedSparseVolumeTexture, PreviousNode, CurrentNode, FactoryNode);
+			}
 		}
-#endif	  // WITH_EDITOR
 	}
-	else if (!Arguments.bIsReimport)
+	else
 	{
-		// Import failed, so let's abandon the asset
-		SparseVolumeTexture->RemoveFromRoot();
-		SparseVolumeTexture->MarkAsGarbage();
+		FactoryNode->ApplyAllCustomAttributeToObject(ImportedSparseVolumeTexture);
+	}
+#endif	  // WITH_EDITOR
+
+	return ImportAssetResult;
+}
+
+void UInterchangeSparseVolumeTextureFactory::SetupObject_GameThread(const FSetupObjectParams& Arguments)
+{
+	TRACE_CPUPROFILER_EVENT_SCOPE(UInterchangeSparseVolumeTextureFactory::SetupObject_GameThread);
+
+	using namespace UE::InterchangeSparseVolumeTextureFactory::Private;
+
+	if (bSkipImport)
+	{
+		return;
 	}
 
 	Super::SetupObject_GameThread(Arguments);
@@ -871,7 +987,7 @@ void UInterchangeSparseVolumeTextureFactory::SetupObject_GameThread(const FSetup
 	//
 	// The Streamable derived class is the one that stores the AssetImportData, but all SVT classes that are actual assets
 	// derive it, so it should be safe to assume we have a Streamable here
-	UStreamableSparseVolumeTexture* Streamable = Cast<UStreamableSparseVolumeTexture>(SparseVolumeTexture);
+	UStreamableSparseVolumeTexture* Streamable = Cast<UStreamableSparseVolumeTexture>(Arguments.ImportedObject);
 	if (Streamable && Arguments.SourceData && HasValidPayloads())
 	{
 		// We must call the Update of the asset source file in the main thread because UAssetImportData::Update execute some delegate we do not
@@ -888,13 +1004,6 @@ void UInterchangeSparseVolumeTextureFactory::SetupObject_GameThread(const FSetup
 		SetImportAssetDataParameters.SourceFiles = MoveTemp(SourceFiles);
 
 		Streamable->AssetImportData = UE::Interchange::FFactoryCommon::SetImportAssetData(SetImportAssetDataParameters);
-
-		if (AlternateTexturePath.IsSet())
-		{
-			// A file different from the source data has been used to create the texture.
-			// Set this file's path as the first file name
-			Streamable->AssetImportData->AddFileName(AlternateTexturePath.GetValue(), 0);
-		}
 	}
 #endif	  // WITH_EDITOR
 }
@@ -955,9 +1064,9 @@ void UInterchangeSparseVolumeTextureFactory::ClearBackupSourceData(const UObject
 
 bool UInterchangeSparseVolumeTextureFactory::HasValidPayloads() const
 {
-	for (const TOptional<UE::Interchange::FVolumePayloadData>& Payload : ProcessedPayloads)
+	for (const FVolumePayload& Payload : ProcessedPayloads)
 	{
-		if (Payload.IsSet())
+		if (Payload.PayloadData.IsSet())
 		{
 			return true;
 		}
diff --git a/Engine/Plugins/Interchange/Runtime/Source/Import/Public/OpenVDB/InterchangeOpenVDBTranslator.h b/Engine/Plugins/Interchange/Runtime/Source/Import/Public/OpenVDB/InterchangeOpenVDBTranslator.h
--- a/Engine/Plugins/Interchange/Runtime/Source/Import/Public/OpenVDB/InterchangeOpenVDBTranslator.h
+++ b/Engine/Plugins/Interchange/Runtime/Source/Import/Public/OpenVDB/InterchangeOpenVDBTranslator.h
@@ -34,6 +34,14 @@ public:
 	UPROPERTY(EditAnywhere, Category = "OpenVDB Translator")
 	bool bTranslateAdjacentNumberedFiles;
 
+	/**
+	 * Identifier set on the generated volume nodes (via SetCustomAnimationID) if they are part of an animation, to identify said animation.
+	 * Volume nodes that share the same AnimationID belong to the same volume animation.
+	 * This can be left empty to have an ID automatically generated, or can be set to something to force it to that value.
+	 */
+	UPROPERTY()
+	FString AnimationID;
+
 public:
 	UInterchangeOpenVDBTranslatorSettings();
 };
diff --git a/Engine/Plugins/Interchange/Runtime/Source/Import/Public/Scene/InterchangeHeterogeneousVolumeActorFactory.h b/Engine/Plugins/Interchange/Runtime/Source/Import/Public/Scene/InterchangeHeterogeneousVolumeActorFactory.h
new file mode 100644
--- /dev/null
+++ b/Engine/Plugins/Interchange/Runtime/Source/Import/Public/Scene/InterchangeHeterogeneousVolumeActorFactory.h
@@ -0,0 +1,36 @@
+// Copyright Epic Games, Inc. All Rights Reserved.
+
+#pragma once
+
+#include "Scene/InterchangeActorFactory.h"
+
+#include "InterchangeHeterogeneousVolumeActorFactory.generated.h"
+
+UCLASS(BlueprintType)
+class INTERCHANGEIMPORT_API UInterchangeHeterogeneousVolumeActorFactory : public UInterchangeActorFactory
+{
+	GENERATED_BODY()
+
+public:
+	//////////////////////////////////////////////////////////////////////////
+	// Interchange factory base interface begin
+
+	virtual UClass* GetFactoryClass() const override;
+
+	// Interchange factory base interface end
+	//////////////////////////////////////////////////////////////////////////
+
+protected:
+	//////////////////////////////////////////////////////////////////////////
+	// Interchange actor factory interface begin
+
+	virtual UObject* ProcessActor(
+		class AActor& SpawnedActor,
+		const UInterchangeActorFactoryNode& FactoryNode,
+		const UInterchangeBaseNodeContainer& NodeContainer,
+		const FImportSceneObjectsParams& Params
+	) override;
+
+	// Interchange actor factory interface end
+	//////////////////////////////////////////////////////////////////////////
+};
diff --git a/Engine/Plugins/Interchange/Runtime/Source/Import/Public/Usd/InterchangeUsdTranslator.h b/Engine/Plugins/Interchange/Runtime/Source/Import/Public/Usd/InterchangeUsdTranslator.h
--- a/Engine/Plugins/Interchange/Runtime/Source/Import/Public/Usd/InterchangeUsdTranslator.h
+++ b/Engine/Plugins/Interchange/Runtime/Source/Import/Public/Usd/InterchangeUsdTranslator.h
@@ -7,11 +7,13 @@
 #include "Mesh/InterchangeMeshPayloadInterface.h"
 #include "Texture/InterchangeBlockedTexturePayloadInterface.h"
 #include "Texture/InterchangeTexturePayloadInterface.h"
-#include "UObject/Object.h"
-#include "UObject/ObjectMacros.h"
+#include "Volume/InterchangeVolumePayloadInterface.h"
 
 #include "USDStageOptions.h"
 
+#include "UObject/Object.h"
+#include "UObject/ObjectMacros.h"
+
 #include "InterchangeUsdTranslator.generated.h"
 
 enum class EUsdInterpolationType : uint8;
@@ -72,6 +74,7 @@ class UInterchangeUSDTranslator
 	, public IInterchangeTexturePayloadInterface
 	, public IInterchangeBlockedTexturePayloadInterface
 	, public IInterchangeAnimationPayloadInterface
+	, public IInterchangeVolumePayloadInterface
 {
 	GENERATED_BODY()
 
@@ -110,10 +113,12 @@ public:
 	virtual TArray<UE::Interchange::FAnimationPayloadData> GetAnimationPayloadData(
 		const TArray<UE::Interchange::FAnimationPayloadQuery>& PayloadQueries
 	) const override;
+
+	virtual TOptional<UE::Interchange::FVolumePayloadData> GetVolumePayloadData(const UE::Interchange::FVolumePayloadKey& PayloadKey) const override;
 	/** End Interchange payload interfaces */
 
 private:
-	TUniquePtr<UE::InterchangeUsdTranslator::Private::UInterchangeUSDTranslatorImpl> Impl;
+	mutable TUniquePtr<UE::InterchangeUsdTranslator::Private::UInterchangeUSDTranslatorImpl> Impl;
 
 	UPROPERTY(Transient, DuplicateTransient)
 	mutable TObjectPtr<UInterchangeUsdTranslatorSettings> TranslatorSettings = nullptr;
diff --git a/Engine/Plugins/Interchange/Runtime/Source/Import/Public/Volume/InterchangeSparseVolumeTextureFactory.h b/Engine/Plugins/Interchange/Runtime/Source/Import/Public/Volume/InterchangeSparseVolumeTextureFactory.h
--- a/Engine/Plugins/Interchange/Runtime/Source/Import/Public/Volume/InterchangeSparseVolumeTextureFactory.h
+++ b/Engine/Plugins/Interchange/Runtime/Source/Import/Public/Volume/InterchangeSparseVolumeTextureFactory.h
@@ -3,6 +3,7 @@
 #pragma once
 
 #include "InterchangeFactoryBase.h"
+#include "Volume/InterchangeVolumePayloadData.h"
 
 #include "CoreMinimal.h"
 #include "EditorFramework/AssetImportData.h"
@@ -13,12 +14,19 @@
 #include "InterchangeSparseVolumeTextureFactory.generated.h"
 
 class UInterchangeSparseVolumeTextureFactoryNode;
+class USparseVolumeTexture;
 namespace UE::Interchange
 {
-	struct FVolumePayloadData;
 	struct FVolumePayloadKey;
 }
 
+struct FVolumePayload
+{
+	FString VolumeNodeUid;
+	TArray<int32> PayloadFrameIndices;
+	TOptional<UE::Interchange::FVolumePayloadData> PayloadData;
+};
+
 UCLASS(BlueprintType)
 class INTERCHANGEIMPORT_API UInterchangeSparseVolumeTextureFactory : public UInterchangeFactoryBase
 {
@@ -46,8 +54,6 @@ private:
 	bool HasValidPayloads() const;
 
 private:
-	TOptional<FString> AlternateTexturePath;
-
 #if WITH_EDITORONLY_DATA
 	// The data for the source files will be collecetd here during the import during an async task,
 	// and later be added to AssetImportData
@@ -55,7 +61,7 @@ private:
 #endif
 
 	// Where we retain the processed payloads between the different factory interface function calls
-	TArray<TOptional<UE::Interchange::FVolumePayloadData>> ProcessedPayloads;
+	TArray<FVolumePayload> ProcessedPayloads;
 
 	// We turn this to true when we run into an existing asset that we shouldn't overwrite/modify (e.g. when reimporting
 	// some other asset type, like reimporting a material and finding an existing SVT that it wants to use)
diff --git a/Engine/Plugins/Interchange/Runtime/Source/Nodes/Private/InterchangeVolumeNode.cpp b/Engine/Plugins/Interchange/Runtime/Source/Nodes/Private/InterchangeVolumeNode.cpp
--- a/Engine/Plugins/Interchange/Runtime/Source/Nodes/Private/InterchangeVolumeNode.cpp
+++ b/Engine/Plugins/Interchange/Runtime/Source/Nodes/Private/InterchangeVolumeNode.cpp
@@ -4,16 +4,27 @@
 
 namespace UE::Interchange
 {
-	const FAttributeKey& FInterchangeVolumeNodeStaticData::GetGridDependenciesBaseKey()
+	const FAttributeKey& FInterchangeVolumeNodeStaticData::GetCustomGridDependeciesBaseKey()
 	{
 		static FAttributeKey AttributeKey(TEXT("__GridDependencies__"));
 		return AttributeKey;
 	}
+
+	const FAttributeKey& FInterchangeVolumeNodeStaticData::GetCustomFrameIndicesInAnimationBaseKey()
+	{
+		static FAttributeKey AttributeKey(TEXT("__FrameIndexInAnimation__"));
+		return AttributeKey;
+	}
 }
 
 UInterchangeVolumeNode::UInterchangeVolumeNode()
 {
-	GridDependencies.Initialize(Attributes, UE::Interchange::FInterchangeVolumeNodeStaticData::GetGridDependenciesBaseKey().ToString());
+	GridDependencies.Initialize(Attributes, UE::Interchange::FInterchangeVolumeNodeStaticData::GetCustomGridDependeciesBaseKey().ToString());
+
+	IndexInVolumeAnimation.Initialize(
+		Attributes,
+		UE::Interchange::FInterchangeVolumeNodeStaticData::GetCustomFrameIndicesInAnimationBaseKey().ToString()
+	);
 }
 
 FString UInterchangeVolumeNode::GetTypeName() const
@@ -31,31 +42,61 @@ bool UInterchangeVolumeNode::SetCustomFileName(const FString& AttributeValue)
 	IMPLEMENT_NODE_ATTRIBUTE_SETTER_NODELEGATE(FileName, FString);
 }
 
-int32 UInterchangeVolumeNode::GetGridDependeciesCount() const
+int32 UInterchangeVolumeNode::GetCustomGridDependeciesCount() const
 {
 	return GridDependencies.GetCount();
 }
 
-void UInterchangeVolumeNode::GetGridDependencies(TArray<FString>& OutDependencies) const
+void UInterchangeVolumeNode::GetCustomGridDependecies(TArray<FString>& OutDependencies) const
 {
 	GridDependencies.GetItems(OutDependencies);
 }
 
-void UInterchangeVolumeNode::GetGridDependency(const int32 Index, FString& OutDependency) const
+void UInterchangeVolumeNode::GetCustomGridDependency(const int32 Index, FString& OutDependency) const
 {
 	GridDependencies.GetItem(Index, OutDependency);
 }
 
-bool UInterchangeVolumeNode::AddGridDependencyUid(const FString& DependencyUid)
+bool UInterchangeVolumeNode::AddCustomGridDependency(const FString& DependencyUid)
 {
 	return GridDependencies.AddItem(DependencyUid);
 }
 
-bool UInterchangeVolumeNode::RemoveGridDependencyUid(const FString& DependencyUid)
+bool UInterchangeVolumeNode::RemoveCustomGridDependency(const FString& DependencyUid)
 {
 	return GridDependencies.RemoveItem(DependencyUid);
 }
 
+bool UInterchangeVolumeNode::GetCustomAnimationID(FString& AttributeValue) const
+{
+	IMPLEMENT_NODE_ATTRIBUTE_GETTER(AnimationID, FString);
+}
+
+bool UInterchangeVolumeNode::SetCustomAnimationID(const FString& AttributeValue)
+{
+	IMPLEMENT_NODE_ATTRIBUTE_SETTER_NODELEGATE(AnimationID, FString);
+}
+
+void UInterchangeVolumeNode::GetCustomFrameIndicesInAnimation(TArray<int32>& OutVolumeIndices) const
+{
+	IndexInVolumeAnimation.GetItems(OutVolumeIndices);
+}
+
+void UInterchangeVolumeNode::GetCustomFrameIndexInAnimation(int32 IndexIndex, int32& OutIndex) const
+{
+	IndexInVolumeAnimation.GetItem(IndexIndex, OutIndex);
+}
+
+bool UInterchangeVolumeNode::AddCustomFrameIndexInAnimation(int32 Index)
+{
+	return IndexInVolumeAnimation.AddItem(Index);
+}
+
+bool UInterchangeVolumeNode::RemoveCustomFrameIndexInAnimation(int32 FrameIndex)
+{
+	return IndexInVolumeAnimation.RemoveItem(FrameIndex);
+}
+
 FString UInterchangeVolumeGridNode::GetTypeName() const
 {
 	return TEXT("VolumeGrid");
diff --git a/Engine/Plugins/Interchange/Runtime/Source/Nodes/Public/InterchangeAnimationDefinitions.h b/Engine/Plugins/Interchange/Runtime/Source/Nodes/Public/InterchangeAnimationDefinitions.h
--- a/Engine/Plugins/Interchange/Runtime/Source/Nodes/Public/InterchangeAnimationDefinitions.h
+++ b/Engine/Plugins/Interchange/Runtime/Source/Nodes/Public/InterchangeAnimationDefinitions.h
@@ -340,9 +340,11 @@ enum class EInterchangePropertyTracks : int32
 	StaticMeshForcedLodModel,
 	StaticMeshReverseCulling,
 	StaticMeshWorldPositionOffsetDisableDistance,
-
 	StaticMesh,
 
+	/** Heterogeneous Volume Component Properties */
+	HeterogeneousVolumeFrame,
+
 	None = -1 UMETA(hidden),
 	Visibility UE_DEPRECATED(5.6, "Renamed to 'ActorHiddenInGame'") = ActorHiddenInGame UMETA(hidden)
 };
diff --git a/Engine/Plugins/Interchange/Runtime/Source/Nodes/Public/InterchangeVolumeNode.h b/Engine/Plugins/Interchange/Runtime/Source/Nodes/Public/InterchangeVolumeNode.h
--- a/Engine/Plugins/Interchange/Runtime/Source/Nodes/Public/InterchangeVolumeNode.h
+++ b/Engine/Plugins/Interchange/Runtime/Source/Nodes/Public/InterchangeVolumeNode.h
@@ -11,7 +11,8 @@ namespace UE::Interchange
 {
 	struct INTERCHANGENODES_API FInterchangeVolumeNodeStaticData : public FBaseNodeStaticData
 	{
-		static const FAttributeKey& GetGridDependenciesBaseKey();
+		static const FAttributeKey& GetCustomGridDependeciesBaseKey();
+		static const FAttributeKey& GetCustomFrameIndicesInAnimationBaseKey();
 	};
 }
 
@@ -53,28 +54,63 @@ public:
 	 * (in other words, returns the number of volume grids contained in this file)
 	 */
 	UFUNCTION(BlueprintCallable, Category = "Interchange | Node | Volume")
-	int32 GetGridDependeciesCount() const;
+	int32 GetCustomGridDependeciesCount() const;
 
 	/** Gets the Node IDs of UInterchangeVolumeGridNodes declared as dependencies by this volume node */
 	UFUNCTION(BlueprintCallable, Category = "Interchange | Node | Volume")
-	void GetGridDependencies(TArray<FString>& OutDependencies) const;
+	void GetCustomGridDependecies(TArray<FString>& OutDependencies) const;
 
 	/** Gets the Node ID of an UInterchangeVolumeGridNode dependency */
 	UFUNCTION(BlueprintCallable, Category = "Interchange | Node | Volume")
-	void GetGridDependency(const int32 Index, FString& OutDependency) const;
+	void GetCustomGridDependency(const int32 Index, FString& OutDependency) const;
 
 	/** Sets the Node ID of an UInterchangeVolumeGridNode dependency */
 	UFUNCTION(BlueprintCallable, Category = "Interchange | Node | Volume")
-	bool AddGridDependencyUid(const FString& DependencyUid);
+	bool AddCustomGridDependency(const FString& DependencyUid);
 
 	/** Removes the Node ID of an UInterchangeVolumeGridNode dependency */
 	UFUNCTION(BlueprintCallable, Category = "Interchange | Node | Volume")
-	bool RemoveGridDependencyUid(const FString& DependencyUid);
+	bool RemoveCustomGridDependency(const FString& DependencyUid);
+
+	/**
+	 * Gets an identifier that is shared by all volume nodes that correspond to the same animation
+	 * (i.e. every volume node within the animation will have the same AnimationID)
+	 */
+	 UFUNCTION(BlueprintCallable, Category = "Interchange | Node | Volume")
+	 bool GetCustomAnimationID(FString& AttributeValue) const;
+
+	 /**
+	  * Sets an identifier that is shared by all volume nodes that correspond to the same animation
+	  * (i.e. every volume node within the animation will have the same AnimationID)
+	  */
+	UFUNCTION(BlueprintCallable, Category = "Interchange | Node | Volume")
+	bool SetCustomAnimationID(const FString& AttributeValue);
+
+	/**
+	 * Get the frame indices for which this volume is displayed within the animation that it belongs to, if any
+	 * (e.g. if this had [2, 3] then frames 2 and 3 of the animation with AnimationID should display this volume)
+	 */
+	UFUNCTION(BlueprintCallable, Category = "Interchange | Node | Volume")
+	void GetCustomFrameIndicesInAnimation(TArray<int32>& OutAnimationIndices) const;
+
+	/** Gets one of the frame indices for which this volume is displayed within the animation that it belongs to, if any */
+	UFUNCTION(BlueprintCallable, Category = "Interchange | Node | Volume")
+	void GetCustomFrameIndexInAnimation(int32 IndexIndex, int32& OutIndex) const;
+
+	/** Adds a frame index for which this volume is displayed within the animation that it belongs to, if any */
+	UFUNCTION(BlueprintCallable, Category = "Interchange | Node | Volume")
+	bool AddCustomFrameIndexInAnimation(int32 Index);
+
+	/** Removes a frame index for which this volume is displayed within the animation that it belongs to, if any */
+	UFUNCTION(BlueprintCallable, Category = "Interchange | Node | Volume")
+	bool RemoveCustomFrameIndexInAnimation(int32 Index);
 
 private:
 	IMPLEMENT_NODE_ATTRIBUTE_KEY(FileName);
+	IMPLEMENT_NODE_ATTRIBUTE_KEY(AnimationID);
 
 	UE::Interchange::TArrayAttributeHelper<FString> GridDependencies;
+	UE::Interchange::TArrayAttributeHelper<int32> IndexInVolumeAnimation;
 };
 
 UENUM(Blueprintable)
diff --git a/Engine/Plugins/Interchange/Runtime/Source/Pipelines/InterchangePipelines.Build.cs b/Engine/Plugins/Interchange/Runtime/Source/Pipelines/InterchangePipelines.Build.cs
--- a/Engine/Plugins/Interchange/Runtime/Source/Pipelines/InterchangePipelines.Build.cs
+++ b/Engine/Plugins/Interchange/Runtime/Source/Pipelines/InterchangePipelines.Build.cs
@@ -41,6 +41,7 @@ namespace UnrealBuildTool.Rules
 					"StaticMeshDescription",
 					"UnrealUSDWrapper",
 					"USDClasses",
+					"USDUtilities",
 				}
 			);
 
diff --git a/Engine/Plugins/Interchange/Runtime/Source/Pipelines/Private/InterchangeGenericScenesPipeline.cpp b/Engine/Plugins/Interchange/Runtime/Source/Pipelines/Private/InterchangeGenericScenesPipeline.cpp
--- a/Engine/Plugins/Interchange/Runtime/Source/Pipelines/Private/InterchangeGenericScenesPipeline.cpp
+++ b/Engine/Plugins/Interchange/Runtime/Source/Pipelines/Private/InterchangeGenericScenesPipeline.cpp
@@ -1,4 +1,4 @@
-// Copyright Epic Games, Inc. All Rights Reserved. 
+// Copyright Epic Games, Inc. All Rights Reserved.
 
 #include "InterchangeGenericScenesPipeline.h"
 
@@ -11,6 +11,7 @@
 #include "InterchangeDecalActorFactoryNode.h"
 #include "InterchangeDecalNode.h"
 #include "InterchangeEditorUtilitiesBase.h"
+#include "InterchangeHeterogeneousVolumeActorFactoryNode.h"
 #include "InterchangeImportReset.h"
 #include "InterchangeLevelFactoryNode.h"
 #include "InterchangeLevelInstanceActorFactoryNode.h"
@@ -31,8 +32,9 @@
 #include "InterchangeStaticMeshFactoryNode.h"
 #include "InterchangeStaticMeshLodDataNode.h"
 #include "InterchangeVariantSetNode.h"
+#include "InterchangeVolumeNode.h"
 #include "Nodes/InterchangeUserDefinedAttribute.h"
-
+#include "Volume/InterchangeVolumeDefinitions.h"
 
 #include "Animation/SkeletalMeshActor.h"
 #include "CineCameraActor.h"
@@ -292,7 +294,7 @@ void UInterchangeGenericLevelPipeline::ExecutePipeline(UInterchangeBaseNodeConta
 	{
 		TargetWorld = GWorld->GetCurrentLevel()->GetWorld();
 	}
-	
+
 	UInterchangeLevelFactoryNode* ParentLevelFactoryNode = nullptr;
 	LevelInstanceActorFactoryNode = nullptr;
 
@@ -449,7 +451,7 @@ void UInterchangeGenericLevelPipeline::ExecutePipeline(UInterchangeBaseNodeConta
 					if (SpecializeTypes.Contains(UE::Interchange::FSceneNodeStaticData::GetJointSpecializeTypeString()))
 					{
 						if(!Cached_ActiveJointUids.Contains(SceneNode->GetUniqueID()))
-						{ 
+						{
 							bSkipNode = false;
 						}
 						else
@@ -539,7 +541,7 @@ void UInterchangeGenericLevelPipeline::ExecuteSceneNodePreImport(const FTransfor
 			}
 		}
 	}
-	
+
 	if (!bRootJointNode)
 	{
 		FString AssetInstanceUid;
@@ -574,7 +576,7 @@ void UInterchangeGenericLevelPipeline::ExecuteSceneNodePreImport(const FTransfor
 						}
 					}
 					return false;
-				});	
+				});
 			}
 			//Skinned mesh are added when the bRootJointNode is true.
 			//In this case we dont want to add an empty staticmesh actor.
@@ -831,6 +833,10 @@ UInterchangeActorFactoryNode* UInterchangeGenericLevelPipeline::CreateActorFacto
 		{
 			return NewObject<UInterchangeDecalActorFactoryNode>(BaseNodeContainer, NAME_None);
 		}
+		else if (TranslatedAssetNode->IsA<UInterchangeVolumeNode>())
+		{
+			return NewObject<UInterchangeHeterogeneousVolumeActorFactoryNode>(BaseNodeContainer, NAME_None);
+		}
 	}
 
 	return NewObject<UInterchangeActorFactoryNode>(BaseNodeContainer, NAME_None);
@@ -1035,7 +1041,7 @@ void UInterchangeGenericLevelPipeline::SetUpFactoryNode(UInterchangeActorFactory
 			{
 				PhysicalCameraFactoryNode->SetCustomFocusMethod(bEnableDepthOfField ? ECameraFocusMethod::Manual : ECameraFocusMethod::DoNotOverride);
 			}
-			
+
 		}
 	}
 	else if (const UInterchangeStandardCameraNode* CameraNode = Cast<UInterchangeStandardCameraNode>(TranslatedAssetNode))
@@ -1123,7 +1129,7 @@ void UInterchangeGenericLevelPipeline::SetUpFactoryNode(UInterchangeActorFactory
 		{
 			DecalActorFactory->SetCustomSortOrder(SortOrder);
 		}
-		
+
 		bool bHasMaterialPathName = false;
 		FString DecalMaterialPathName;
 		if (DecalNode->GetCustomDecalMaterialPathName(DecalMaterialPathName))
@@ -1140,6 +1146,29 @@ void UInterchangeGenericLevelPipeline::SetUpFactoryNode(UInterchangeActorFactory
 			DecalActorFactory->AddFactoryDependencyUid(MaterialFactoryUid);
 		}
 	}
+	else if (const UInterchangeVolumeNode* VolumeNode = Cast<UInterchangeVolumeNode>(TranslatedAssetNode))
+	{
+		UInterchangeHeterogeneousVolumeActorFactoryNode* FactoryNode = Cast<UInterchangeHeterogeneousVolumeActorFactoryNode>(ActorFactoryNode);
+		ensure(FactoryNode);
+
+		// Setup a material binding if we have one (we always should, but let's not ensure here as the user
+		// may have custom pipelines, etc.)
+		{
+			TMap<FString, FString> SlotNameToMaterialPath;
+			SceneNode->GetSlotMaterialDependencies(SlotNameToMaterialPath);
+
+			FString MaterialFactoryNodeUid;
+			if (FString* FoundUid = SlotNameToMaterialPath.Find(UE::Interchange::Volume::VolumetricMaterial))
+			{
+				MaterialFactoryNodeUid = UInterchangeFactoryBaseNode::BuildFactoryNodeUid(*FoundUid);
+			}
+
+			if (!MaterialFactoryNodeUid.IsEmpty())
+			{
+				FactoryNode->SetCustomVolumetricMaterialUid(MaterialFactoryNodeUid);
+			}
+		}
+	}
 }
 
 void UInterchangeGenericLevelPipeline::ExecuteSceneVariantSetNodePreImport(const UInterchangeSceneVariantSetsNode& SceneVariantSetNode)
@@ -1174,7 +1203,7 @@ void UInterchangeGenericLevelPipeline::ExecuteSceneVariantSetNodePreImport(const
 
 			for (const FString& DependencyNodeUid : DependencyNodeUids)
 			{
-				
+
 				const FString DependencyFactoryNodeUid = UInterchangeFactoryBaseNode::BuildFactoryNodeUid(DependencyNodeUid);
 				FactoryNode->AddFactoryDependencyUid(DependencyFactoryNodeUid);
 
@@ -1385,7 +1414,7 @@ void UInterchangeGenericLevelPipeline::ExecutePostImportPipeline(const UIntercha
 void UInterchangeGenericLevelPipeline::ExecutePostBroadcastPipeline(const UInterchangeBaseNodeContainer* InBaseNodeContainer, const FString& NodeKey, UObject* CreatedAsset, bool bIsAReimport)
 {
 	Super::ExecutePostBroadcastPipeline(InBaseNodeContainer, NodeKey, CreatedAsset, bIsAReimport);
-	
+
 	if (!CreatedAsset || !ensure(IsInGameThread()))
 	{
 		return;
@@ -1417,7 +1446,7 @@ void UInterchangeGenericLevelPipeline::ExecutePostBroadcastPipeline(const UInter
 					{
 						if (APackedLevelActor* PackedLevelActor = Cast<APackedLevelActor>(CurrentLevelInstanceActor))
 						{
-							
+
 							PackedLevelActor->UpdateLevelInstanceFromWorldAsset();
 						}
 						else
@@ -1487,7 +1516,7 @@ bool UInterchangeGenericLevelPipeline::FPostPipelineImportData::UpdateLevelInsta
 	}
 
 	bool& bWorldUpdate = Worlds.FindChecked(ReferencedWorld);
-	
+
 	if (!bWorldUpdate)
 	{
 		if (UInterchangeEditorUtilitiesBase* EditorUtilities = UInterchangeManager::GetInterchangeManager().GetEditorUtilities())
diff --git a/Engine/Plugins/Interchange/Runtime/Source/Pipelines/Private/InterchangeSparseVolumeTexturePipeline.cpp b/Engine/Plugins/Interchange/Runtime/Source/Pipelines/Private/InterchangeSparseVolumeTexturePipeline.cpp
--- a/Engine/Plugins/Interchange/Runtime/Source/Pipelines/Private/InterchangeSparseVolumeTexturePipeline.cpp
+++ b/Engine/Plugins/Interchange/Runtime/Source/Pipelines/Private/InterchangeSparseVolumeTexturePipeline.cpp
@@ -5,6 +5,7 @@
 #include "InterchangePipelineLog.h"
 #include "InterchangeSparseVolumeTextureFactoryNode.h"
 #include "InterchangeVolumeNode.h"
+#include "Nodes/InterchangeUserDefinedAttribute.h"
 #include "Volume/InterchangeVolumeDefinitions.h"
 
 #include "Engine/Texture.h"
@@ -92,7 +93,7 @@ namespace UE::Interchange::Private
 		TArray<const UInterchangeVolumeGridNode*> GridNodes;
 		{
 			TArray<FString> GridNodeUids;
-			VolumeNode->GetGridDependencies(GridNodeUids);
+			VolumeNode->GetCustomGridDependecies(GridNodeUids);
 
 			for (int32 GridIndex = 0; GridIndex < GridNodeUids.Num(); ++GridIndex)
 			{
@@ -143,18 +144,18 @@ namespace UE::Interchange::Private
 		const bool bOptimizedDensityAssignment = DensityGridIndex != INDEX_NONE && (NumNonDensity <= 4 && NumNonDensity != 3);
 		if (bOptimizedDensityAssignment)
 		{
-			VolumeFactoryNode->SetCustomAttributesADataType(EInterchangeSparseVolumeTextureFormat::Unorm8);
+			VolumeFactoryNode->SetCustomAttributesAFormat(EInterchangeSparseVolumeTextureFormat::Unorm8);
 			VolumeFactoryNode->SetCustomAttributesAChannelX(DensityGridName + GridNameAndComponentIndexSeparator + TEXT("0"));
 
-			VolumeFactoryNode->SetCustomAttributesBDataType(EInterchangeSparseVolumeTextureFormat::Float16);
+			VolumeFactoryNode->SetCustomAttributesBFormat(EInterchangeSparseVolumeTextureFormat::Float16);
 
 			SetterIndex = 4;	// Start at SetCustomAttributesBChannelX() instead, as our AttributesA texture will hold just the density
 		}
 		else
 		{
-			VolumeFactoryNode->SetCustomAttributesADataType(EInterchangeSparseVolumeTextureFormat::Float16);
+			VolumeFactoryNode->SetCustomAttributesAFormat(EInterchangeSparseVolumeTextureFormat::Float16);
 
-			VolumeFactoryNode->SetCustomAttributesBDataType(EInterchangeSparseVolumeTextureFormat::Float16);
+			VolumeFactoryNode->SetCustomAttributesBFormat(EInterchangeSparseVolumeTextureFormat::Float16);
 		}
 
 		// Actually distribute the remaining grid/components across the channels in order
@@ -357,112 +358,135 @@ void UInterchangeSparseVolumeTexturePipeline::ExecutePipeline(
 
 	TArray<UInterchangeSparseVolumeTextureFactoryNode*> CreatedFactoryNodes;
 
-	// We need to group up all the volume nodes that seem to refer to the same sequence, according to their filename.
-	// These could have been manually added by the user somehow, or discovered via the OpenVDB translator.
-	if (bImportAnimatedSparseVolumeTextures)
+	struct FNodeAndAnimationIndex
 	{
-		struct FNodeAndTrailingIndex
-		{
-			const UInterchangeVolumeNode* Node;
-			int32 Index;
-		};
+		const UInterchangeVolumeNode* Node;
+		int32 Index;
+	};
 
-		TMap<FString, TArray<FNodeAndTrailingIndex>> NodesWithSamePrefix;
-		NodesWithSamePrefix.Reserve(VolumeNodes.Num());
+	// Group up volume nodes by animation ID
+	//
+	// Note: A volume may show up in multiple animation IDs, but that's supported.
+	TSet<UInterchangeVolumeNode*> VolumeNodesWithNoAnimationID;
+	TMap<FString, TArray<FNodeAndAnimationIndex>> AnimationIDToVolumeNodes;
+	for (UInterchangeVolumeNode* VolumeNode : VolumeNodes)
+	{
+		FString AnimationID;
 
-		// Group up volume nodes by the same filename prefix
-		for (UInterchangeVolumeNode* VolumeNode : VolumeNodes)
+		// Animated volume
+		if (bImportAnimatedSparseVolumeTextures && VolumeNode->GetCustomAnimationID(AnimationID) && !AnimationID.IsEmpty())
 		{
-			FString FileName;
-			bool bSuccess = VolumeNode->GetCustomFileName(FileName);
-			if (!bSuccess || FileName.IsEmpty())
+			TArray<int32> AnimationIndices;
+			VolumeNode->GetCustomFrameIndicesInAnimation(AnimationIndices);
+
+			for (const int32 Index : AnimationIndices)
 			{
-				continue;
+				AnimationIDToVolumeNodes.FindOrAdd(AnimationID).Emplace(VolumeNode, Index);
 			}
-			FileName = FPaths::GetBaseFilename(FileName);
+		}
+		// Static volume
+		else
+		{
+			VolumeNodesWithNoAnimationID.Add(VolumeNode);
+		}
+	}
 
-			// Note: Suffix may be INDEX_NONE, but that should be alright (they'll just be grouped together
-			// at the front, if any)
-			FString Prefix;
-			int32 Suffix;
-			SplitNumberedSuffix(FileName, Prefix, Suffix);
+	// Create static factory nodes for ungrouped volume nodes (no animation id)
+	for (const UInterchangeVolumeNode* VolumeNode : VolumeNodesWithNoAnimationID)
+	{
+		const FString FactoryNodeUid = UInterchangeFactoryBaseNode::BuildFactoryNodeUid(VolumeNode->GetUniqueID());
 
-			FNodeAndTrailingIndex& NodeAndIndex = NodesWithSamePrefix.FindOrAdd(Prefix).Emplace_GetRef();
-			NodeAndIndex.Node = VolumeNode;
-			NodeAndIndex.Index = Suffix;
+		UInterchangeSparseVolumeTextureFactoryNode* FactoryNode = CreateTextureFactoryNode(
+			VolumeNode->GetDisplayLabel(),
+			FactoryNodeUid,
+			BaseNodeContainer
+		);
+		if (!FactoryNode)
+		{
+			continue;
 		}
 
-		for (TPair<FString, TArray<FNodeAndTrailingIndex>>& Pair : NodesWithSamePrefix)
-		{
-			TArray<FNodeAndTrailingIndex>& NodeAndIndices = Pair.Value;
-			if (NodeAndIndices.Num() == 0)
-			{
-				continue;
-			}
+		CreatedFactoryNodes.Add(FactoryNode);
 
-			// Sort them according to their suffix integers (note: We need all this because just sorting the filename
-			// strings is not enough in case the files have numbered suffixes without zero padding.
-			// i.e. lexicographic sort would produce an order like ["tornado_1", "tornado_12", "tornado_2", "tornado_24", etc.])
-			NodeAndIndices.Sort(
-				[](const FNodeAndTrailingIndex& LHS, const FNodeAndTrailingIndex& RHS)
-				{
-					return LHS.Index < RHS.Index;
-				}
-			);
+		const bool bAddSourceNodeName = false;
+		UInterchangeUserDefinedAttributesAPI::DuplicateAllUserDefinedAttribute(VolumeNode, FactoryNode, bAddSourceNodeName);
 
-			const FString& Prefix = Pair.Key;
-			const UInterchangeVolumeNode* FirstVolume = NodeAndIndices[0].Node;
+		FactoryNode->AddTargetNodeUid(VolumeNode->GetUniqueID());
+		VolumeNode->AddTargetNodeUid(FactoryNode->GetUniqueID());
 
-			FString DisplayLabel;
-			if (NodeAndIndices.Num() == 1)
-			{
-				DisplayLabel = FirstVolume->GetDisplayLabel();
-			}
-			else
-			{
-				// Splitting the numbered suffix likely left us with something like "tornado_" as the prefix.
-				// If we're going to use that as the asset name, let's strip any separators that were left at the end like that
-				DisplayLabel = RemoveTrailingSeparators(Prefix);
-			}
+		SetupDefaultOpenVDBGridAssignment(FactoryNode, BaseNodeContainer);
+	}
 
-			const FString FactoryNodeUid = UInterchangeFactoryBaseNode::BuildFactoryNodeUid(FirstVolume->GetUniqueID());
+	// Create animated factory nodes for each animation ID
+	for (TPair<FString, TArray<FNodeAndAnimationIndex>>& Pair : AnimationIDToVolumeNodes)
+	{
+		const FString& AnimationID = Pair.Key;
+		TArray<FNodeAndAnimationIndex>& NodeAndIndices = Pair.Value;
+		if (NodeAndIndices.Num() == 0)
+		{
+			continue;
+		}
 
-			UInterchangeSparseVolumeTextureFactoryNode* FactoryNode = CreateTextureFactoryNode(DisplayLabel, FactoryNodeUid, BaseNodeContainer);
-			if (!FactoryNode)
+		// Sort them according to their animation indices
+		NodeAndIndices.Sort(
+			[](const FNodeAndAnimationIndex& LHS, const FNodeAndAnimationIndex& RHS)
 			{
-				continue;
-			}
-			CreatedFactoryNodes.Add(FactoryNode);
+				if (LHS.Index == RHS.Index)
+				{
+					// Fallback for a consistent order in case the animation IDs collide
+					return LHS.Node->GetUniqueID() < RHS.Node->GetUniqueID();
+				}
 
-			for (const FNodeAndTrailingIndex& NodeAndIndex : NodeAndIndices)
-			{
-				FactoryNode->AddTargetNodeUid(NodeAndIndex.Node->GetUniqueID());
-				NodeAndIndex.Node->AddTargetNodeUid(FactoryNode->GetUniqueID());
+				return LHS.Index < RHS.Index;
 			}
+		);
+
+		const UInterchangeVolumeNode* FirstVolume = NodeAndIndices[0].Node;
 
-			SetupDefaultOpenVDBGridAssignment(FactoryNode, BaseNodeContainer);
+		FString FileName;
+		bool bSuccess = FirstVolume->GetCustomFileName(FileName);
+		if (!bSuccess || FileName.IsEmpty())
+		{
+			continue;
 		}
-	}
-	// No grouping: Each volume node becomes it's own factory node
-	else
-	{
-		for (UInterchangeVolumeNode* VolumeNode : VolumeNodes)
+		FileName = FPaths::GetBaseFilename(FileName);				// e.g. "tornado_223"
+
+		FString Prefix;												// e.g. "tornado_"
+		int32 NumberSuffix;											// e.g. 223
+		SplitNumberedSuffix(FileName, Prefix, NumberSuffix);
+		FString DisplayLabel = RemoveTrailingSeparators(Prefix);	// e.g. "tornado"
+
+		const FString FactoryNodeUid = UInterchangeFactoryBaseNode::BuildFactoryNodeUid(FirstVolume->GetUniqueID());
+
+		UInterchangeSparseVolumeTextureFactoryNode* FactoryNode = CreateTextureFactoryNode(DisplayLabel, FactoryNodeUid, BaseNodeContainer);
+		if (!FactoryNode)
 		{
-			const FString DisplayLabel = VolumeNode->GetDisplayLabel();
-			const FString NodeUid = UInterchangeFactoryBaseNode::BuildFactoryNodeUid(VolumeNode->GetUniqueID());
+			continue;
+		}
+		CreatedFactoryNodes.Add(FactoryNode);
 
-			UInterchangeSparseVolumeTextureFactoryNode* FactoryNode = CreateTextureFactoryNode(DisplayLabel, NodeUid, BaseNodeContainer);
-			if (!FactoryNode)
+		const bool bAddSourceNodeName = false;
+		UInterchangeUserDefinedAttributesAPI::DuplicateAllUserDefinedAttribute(FirstVolume, FactoryNode, bAddSourceNodeName);
+
+		// Just record this for information's sake, as we won't really use this for anything after this point
+		FactoryNode->SetCustomAnimationID(AnimationID);
+
+		TSet<const UInterchangeVolumeNode*> AddedNodes;
+		for (const FNodeAndAnimationIndex& NodeAndIndex : NodeAndIndices)
+		{
+			// We may have multiple FNodeAndAnimationIndex for the same node, if the same volume frame shows up
+			// multiple times in an animation. We don't want to add it as a target multiple times though
+			if (AddedNodes.Contains(NodeAndIndex.Node))
 			{
 				continue;
 			}
-			CreatedFactoryNodes.Add(FactoryNode);
+			AddedNodes.Add(NodeAndIndex.Node);
 
-			FactoryNode->AddTargetNodeUid(VolumeNode->GetUniqueID());
-			VolumeNode->AddTargetNodeUid(FactoryNode->GetUniqueID());
-
-			SetupDefaultOpenVDBGridAssignment(FactoryNode, BaseNodeContainer);
+			FactoryNode->AddTargetNodeUid(NodeAndIndex.Node->GetUniqueID());
+			NodeAndIndex.Node->AddTargetNodeUid(FactoryNode->GetUniqueID());
 		}
+
+		SetupDefaultOpenVDBGridAssignment(FactoryNode, BaseNodeContainer);
 	}
 
 	// Set an override asset name if we have exactly one factory node
diff --git a/Engine/Plugins/Interchange/Runtime/Source/Pipelines/Private/InterchangeUSDPipeline.cpp b/Engine/Plugins/Interchange/Runtime/Source/Pipelines/Private/InterchangeUSDPipeline.cpp
--- a/Engine/Plugins/Interchange/Runtime/Source/Pipelines/Private/InterchangeUSDPipeline.cpp
+++ b/Engine/Plugins/Interchange/Runtime/Source/Pipelines/Private/InterchangeUSDPipeline.cpp
@@ -2,17 +2,22 @@
 
 #include "InterchangeUSDPipeline.h"
 
+#include "InterchangeHeterogeneousVolumeActorFactoryNode.h"
 #include "InterchangeManager.h"
 #include "InterchangeMaterialDefinitions.h"
 #include "InterchangeMaterialFactoryNode.h"
-#include "InterchangeMeshNode.h"
 #include "InterchangeMeshFactoryNode.h"
+#include "InterchangeMeshNode.h"
+#include "InterchangePipelineLog.h"
 #include "InterchangeSceneNode.h"
 #include "InterchangeShaderGraphNode.h"
 #include "InterchangeSkeletalMeshLodDataNode.h"
+#include "InterchangeSparseVolumeTextureFactoryNode.h"
 #include "InterchangeStaticMeshLodDataNode.h"
-#include "InterchangePipelineLog.h"
+#include "InterchangeVolumeNode.h"
 #include "Nodes/InterchangeBaseNodeContainer.h"
+#include "Nodes/InterchangeUserDefinedAttribute.h"
+#include "Usd/InterchangeUsdDefinitions.h"
 
 #include "DynamicMesh/DynamicMesh3.h"
 #include "Engine/StaticMesh.h"
@@ -28,18 +33,20 @@
 #include "StaticMeshOperations.h"
 
 #include "UnrealUSDWrapper.h"
+#include "USDConversionUtils.h"
 #include "USDMaterialUtils.h"
-#include "Usd/InterchangeUsdDefinitions.h"
+#include "USDProjectSettings.h"
 #include "UsdWrappers/SdfPath.h"
 
 #if WITH_EDITORONLY_DATA
 #include "AssetUtils/Texture2DBuilder.h"
-#endif // WITH_EDITORONLY_DATA
+#endif	  // WITH_EDITORONLY_DATA
 
 #include UE_INLINE_GENERATED_CPP_BY_NAME(InterchangeUSDPipeline)
 
 namespace UE::InterchangeUsdPipeline::Private
 {
+#if USE_USD_SDK
 	const FString& GetPseudoRootTranslatedNodeUid()
 	{
 		const static FString Result = UE::FSdfPath::AbsoluteRootPath().GetString();
@@ -281,6 +288,9 @@ namespace UE::InterchangeUsdPipeline::Private
 		UInterchangeMaterialInstanceFactoryNode* NewFactoryNode = NewObject<UInterchangeMaterialInstanceFactoryNode>(&NodeContainer);
 		NewFactoryNode->InitializeNode(MaterialFactoryNodeUid, MaterialFactoryNodeName, EInterchangeNodeContainerType::FactoryData);
 
+		const bool bAddSourceNodeName = false;
+		UInterchangeUserDefinedAttributesAPI::DuplicateAllUserDefinedAttribute(&ShaderGraphNode, NewFactoryNode, bAddSourceNodeName);
+
 		SetCustomParentMaterial(ShaderGraphNode, *NewFactoryNode);
 
 		// Pick material instance class
@@ -428,6 +438,286 @@ namespace UE::InterchangeUsdPipeline::Private
 		}
 	}
 
+	// Move custom attribute info from first volume nodes into corresponding factory nodes.
+	// This info originally came from the USD codeless SparseVolumeTextureAPI schema.
+	void MoveCustomAttributesToVolumeFactoryNode(const UInterchangeVolumeNode* VolumeNode, UInterchangeSparseVolumeTextureFactoryNode* FactoryNode)
+	{
+		using namespace UE::Interchange;
+		using namespace UE::Interchange::USD;
+
+		using StringSetterFunc = decltype(&UInterchangeSparseVolumeTextureFactoryNode::SetCustomAttributesAChannelX);
+		using EnumSetterFunc = decltype(&UInterchangeSparseVolumeTextureFactoryNode::SetCustomAttributesAFormat);
+
+		const static TMap<FAttributeKey, StringSetterFunc> StringSetters{
+			{FAttributeKey{SparseVolumeTexture::AttributesAChannelR}, &UInterchangeSparseVolumeTextureFactoryNode::SetCustomAttributesAChannelX},
+			{FAttributeKey{SparseVolumeTexture::AttributesAChannelG}, &UInterchangeSparseVolumeTextureFactoryNode::SetCustomAttributesAChannelY},
+			{FAttributeKey{SparseVolumeTexture::AttributesAChannelB}, &UInterchangeSparseVolumeTextureFactoryNode::SetCustomAttributesAChannelZ},
+			{FAttributeKey{SparseVolumeTexture::AttributesAChannelA}, &UInterchangeSparseVolumeTextureFactoryNode::SetCustomAttributesAChannelW},
+			{FAttributeKey{SparseVolumeTexture::AttributesBChannelR}, &UInterchangeSparseVolumeTextureFactoryNode::SetCustomAttributesBChannelX},
+			{FAttributeKey{SparseVolumeTexture::AttributesBChannelG}, &UInterchangeSparseVolumeTextureFactoryNode::SetCustomAttributesBChannelY},
+			{FAttributeKey{SparseVolumeTexture::AttributesBChannelB}, &UInterchangeSparseVolumeTextureFactoryNode::SetCustomAttributesBChannelZ},
+			{FAttributeKey{SparseVolumeTexture::AttributesBChannelA}, &UInterchangeSparseVolumeTextureFactoryNode::SetCustomAttributesBChannelW},
+		};
+
+		const static TMap<FAttributeKey, EnumSetterFunc> EnumSetters{
+			{FAttributeKey{SparseVolumeTexture::AttributesAFormat}, &UInterchangeSparseVolumeTextureFactoryNode::SetCustomAttributesAFormat},
+			{FAttributeKey{SparseVolumeTexture::AttributesBFormat}, &UInterchangeSparseVolumeTextureFactoryNode::SetCustomAttributesBFormat},
+		};
+
+		TArray<FAttributeKey> AttributeKeys;
+		VolumeNode->GetAttributeKeys(AttributeKeys);
+
+		for (const FAttributeKey& AttributeKey : AttributeKeys)
+		{
+			if (StringSetterFunc FoundStringSetter = StringSetters.FindRef(AttributeKey))
+			{
+				FString Value;
+				if (VolumeNode->GetStringAttribute(AttributeKey.ToString(), Value))
+				{
+					(FactoryNode->*FoundStringSetter)(Value);
+				}
+			}
+			else if (EnumSetterFunc FoundEnumSetter = EnumSetters.FindRef(AttributeKey))
+			{
+				int32 Value = 0;
+				if (VolumeNode->GetInt32Attribute(AttributeKey.ToString(), Value))
+				{
+					(FactoryNode->*FoundEnumSetter)(static_cast<EInterchangeSparseVolumeTextureFormat>(Value));
+				}
+			}
+		}
+	}
+
+	void ProcessHeterogeneousVolumeSceneNode(const UInterchangeSceneNode* SceneNode, UInterchangeBaseNodeContainer& NodeContainer)
+	{
+		// Find the corresponding actor factory node
+		UInterchangeHeterogeneousVolumeActorFactoryNode* ActorFactoryNode = nullptr;
+		{
+			TArray<FString> TargetNodeUids;
+			SceneNode->GetTargetNodeUids(TargetNodeUids);
+
+			for (const FString& TargetNodeUid : TargetNodeUids)
+			{
+				UInterchangeFactoryBaseNode* FactoryBaseNode = NodeContainer.GetFactoryNode(TargetNodeUid);
+				ActorFactoryNode = Cast<UInterchangeHeterogeneousVolumeActorFactoryNode>(FactoryBaseNode);
+				if (ActorFactoryNode)
+				{
+					break;
+				}
+			}
+		}
+
+		// Find the material it is using
+		UInterchangeMaterialInstanceFactoryNode* MaterialFactoryNode = nullptr;
+		{
+			TMap<FString, FString> OutMaterialDependencies;
+			SceneNode->GetSlotMaterialDependencies(OutMaterialDependencies);
+
+			if (FString* MaterialInstanceUid = OutMaterialDependencies.Find(UE::Interchange::Volume::VolumetricMaterial))
+			{
+				const FString MaterialFactoryNodeUid = UInterchangeFactoryBaseNode::BuildFactoryNodeUid(*MaterialInstanceUid);
+				MaterialFactoryNode = Cast<UInterchangeMaterialInstanceFactoryNode>(NodeContainer.GetFactoryNode(MaterialFactoryNodeUid));
+			}
+		}
+		if (!MaterialFactoryNode)
+		{
+			return;
+		}
+
+		// Make sure we spawn our actor after the material has finished generating, as we need to assign the
+		// material to the actor
+		if (ActorFactoryNode && MaterialFactoryNode)
+		{
+			ActorFactoryNode->AddFactoryDependencyUid(MaterialFactoryNode->GetUniqueID());
+		}
+
+		// Fixup the SVT assignment on the material, if needed
+		//
+		// Reference: CollectMaterialParameterTextureAssignment from USDVolVolumeTranslator.cpp
+		// We don't refactor/reuse the function as the logic/data is slightly different
+		//
+		// The USDTranslator has already translated into string attributes the best SVT to material parameter assignment
+		// we could come up with so far, which includes the fallback of interpreting the field names as material parameter names.
+		// In here, we check if it used that fallback, and if so we check if it will work or not, and if not we use another fallback.
+		// The intent here is to make sure that we get *something* assigned to the SVT material, like the legacy schema translator does
+		{
+			TMap<FString, FString> CleanAttributeNamesToVolumes;
+			bool bIsFallbackCase = false;
+
+			TArray<UE::Interchange::FAttributeKey> AllAttributeKeys;
+			MaterialFactoryNode->GetAttributeKeys(AllAttributeKeys);
+			for (const UE::Interchange::FAttributeKey& AttributeKey : AllAttributeKeys)
+			{
+				const UE::Interchange::EAttributeTypes AttributeType = MaterialFactoryNode->GetAttributeType(AttributeKey);
+				if (AttributeType != UE::Interchange::EAttributeTypes::String)
+				{
+					continue;
+				}
+
+				const FString AttributeKeyString = AttributeKey.ToString();
+				if (!AttributeKeyString.Contains(UE::Interchange::USD::VolumeFieldNameMaterialParameterPrefix))
+				{
+					continue;
+				}
+
+				bIsFallbackCase = true;
+
+				FString AttributeValue;
+				if (!MaterialFactoryNode->GetStringAttribute(AttributeKeyString, AttributeValue))
+				{
+					continue;
+				}
+
+				// e.g. go from "Inputs:USD_FieldName_density:Value" to "USD_FieldName_density"
+				FString CleanAttributeName = UInterchangeShaderPortsAPI::MakeInputName(AttributeKeyString);
+
+				// e.g. go "USD_FieldName_density" to "density"
+				CleanAttributeName.RemoveFromStart(UE::Interchange::USD::VolumeFieldNameMaterialParameterPrefix);
+
+				CleanAttributeNamesToVolumes.Add(CleanAttributeName, AttributeValue);
+
+				// We definitely don't want the fallback attrs to continue past this pipeline, as they have no meaning for Interchange itself
+				ensure(MaterialFactoryNode->RemoveAttribute(AttributeKeyString));
+			}
+			if (!bIsFallbackCase)
+			{
+				// If the translator didn't emit any of these VolumeFieldNameMaterialParameterPrefix attributes then we're not
+				// in a fallback material assignment case and can just stop now
+				return;
+			}
+
+			// Get the actual SVT parameter names from the material we're using on this actor
+			TArray<FString> SparseVolumeTextureParameterNames;
+			{
+				FString ParentMaterialContentPath;
+				if (!MaterialFactoryNode->GetCustomParent(ParentMaterialContentPath))
+				{
+					return;
+				}
+
+				FSoftObjectPath ReferencedObject{ParentMaterialContentPath};
+
+				// This is why this entire pipeline needs to be on the game thread
+				UMaterial* ParentMaterial = nullptr;
+				if (UMaterialInterface* MaterialInterface = Cast<UMaterialInterface>(ReferencedObject.TryLoad()))
+				{
+					ParentMaterial = MaterialInterface->GetMaterial();
+				}
+				if (!ParentMaterial)
+				{
+					return;
+				}
+
+				SparseVolumeTextureParameterNames = UsdUtils::GetSparseVolumeTextureParameterNames(ParentMaterial);
+			}
+
+			// Compensate for how parameter names are usually upper case in UE and field names are lower case in USD
+			TMap<FString, FString> LowercaseParamNamesToOriginal;
+			LowercaseParamNamesToOriginal.Reserve(SparseVolumeTextureParameterNames.Num());
+			for (const FString& ParamName : SparseVolumeTextureParameterNames)
+			{
+				LowercaseParamNamesToOriginal.Add(ParamName.ToLower(), ParamName);
+			}
+
+			// Check if our volume prim field names actually match material parameters (if we disregard casing)
+			bool bFoundFieldNameMatch = false;
+			for (const TPair<FString, FString>& AttrPair : CleanAttributeNamesToVolumes)
+			{
+				const FString& CleanAttributeName = AttrPair.Key;
+				const FString& VolumeUid = AttrPair.Value;
+
+				if (FString* FoundParamName = LowercaseParamNamesToOriginal.Find(CleanAttributeName.ToLower()))
+				{
+					bFoundFieldNameMatch = true;
+
+					const FString ParameterKey = UInterchangeShaderPortsAPI::MakeInputValueKey(*FoundParamName);
+					MaterialFactoryNode->AddStringAttribute(ParameterKey, VolumeUid);
+				}
+			}
+			if (bFoundFieldNameMatch)
+			{
+				// If we found any kind of match with the field names let's just take that
+				return;
+			}
+
+			// As a final fallback case just assign SVTs to the material parameters in alphabetical order like the legacy
+			// schema translator does
+
+			TArray<FString> VolumeUids;
+			{
+				TSet<FString> UniqueVolumeIds;
+				UniqueVolumeIds.Reserve(CleanAttributeNamesToVolumes.Num());
+				for (const TPair<FString, FString>& ParamPair : CleanAttributeNamesToVolumes)
+				{
+					UniqueVolumeIds.Add(ParamPair.Value);
+				}
+				VolumeUids = UniqueVolumeIds.Array();
+			}
+
+			SparseVolumeTextureParameterNames.Sort();
+			VolumeUids.Sort();
+
+			for (int32 Index = 0; Index < VolumeUids.Num() && Index < SparseVolumeTextureParameterNames.Num(); ++Index)
+			{
+				const FString& VolumeUid = VolumeUids[Index];
+				const FString& ParameterName = SparseVolumeTextureParameterNames[Index];
+				const FString ParameterKey = UInterchangeShaderPortsAPI::MakeInputValueKey(ParameterName);
+
+				MaterialFactoryNode->AddStringAttribute(ParameterKey, VolumeUid);
+			}
+		}
+	}
+
+	void ProcessVolumeNodes(UInterchangeBaseNodeContainer& NodeContainer)
+	{
+		using namespace UE::Interchange;
+		using namespace UE::Interchange::USD;
+
+		TMap<const UInterchangeSceneNode*, const UInterchangeVolumeNode*> SceneNodeToFirstVolume;
+		TMap<const UInterchangeVolumeNode*, UInterchangeSparseVolumeTextureFactoryNode*> FirstVolumeToFactoryNode;
+
+		// Collect all the node pairings we need
+		NodeContainer.IterateNodesOfType<UInterchangeSparseVolumeTextureFactoryNode>(
+			[&SceneNodeToFirstVolume,
+			 &NodeContainer,
+			 &FirstVolumeToFactoryNode](const FString& FactoryNodeUid, UInterchangeSparseVolumeTextureFactoryNode* FactoryNode)
+			{
+				FString FirstVolumeUid = UInterchangeFactoryBaseNode::BuildTranslatedNodeUid(FactoryNodeUid);
+
+				const UInterchangeVolumeNode* FirstVolume = Cast<const UInterchangeVolumeNode>(NodeContainer.GetNode(FirstVolumeUid));
+				ensure(FirstVolume);
+
+				FirstVolumeToFactoryNode.Add(FirstVolume, FactoryNode);
+
+				// Find our scene node too
+				TArray<FString> TargetNodeUids;
+				FirstVolume->GetTargetNodeUids(TargetNodeUids);
+				for (const FString& TargetNodeUid : TargetNodeUids)
+				{
+					const UInterchangeBaseNode* BaseNode = NodeContainer.GetNode(TargetNodeUid);
+					if (const UInterchangeSceneNode* SceneNode = Cast<const UInterchangeSceneNode>(BaseNode))
+					{
+						SceneNodeToFirstVolume.Add(SceneNode, FirstVolume);
+						break;
+					}
+				}
+			}
+		);
+
+		for (TPair<const UInterchangeVolumeNode*, UInterchangeSparseVolumeTextureFactoryNode*> Pair : FirstVolumeToFactoryNode)
+		{
+			const UInterchangeVolumeNode* VolumeNode = Pair.Key;
+			UInterchangeSparseVolumeTextureFactoryNode* FactoryNode = Pair.Value;
+			MoveCustomAttributesToVolumeFactoryNode(VolumeNode, FactoryNode);
+		}
+
+		for (const TPair<const UInterchangeSceneNode*, const UInterchangeVolumeNode*>& Pair : SceneNodeToFirstVolume)
+		{
+			const UInterchangeSceneNode* SceneNode = Pair.Key;
+			ProcessHeterogeneousVolumeSceneNode(SceneNode, NodeContainer);
+		}
+	}
+
 	void RemovePseudoRootFactoryNode(UInterchangeBaseNodeContainer& NodeContainer)
 	{
 		const FString& FactoryNodeUid = GetPseudoRootFactoryNodeUid();
@@ -460,7 +750,7 @@ namespace UE::InterchangeUsdPipeline::Private
 			[NodeContainer](const FString& NodeUid, UInterchangeMeshNode* MeshNode)
 			{
 				if (FString FactoryNodeUid = UInterchangeFactoryBaseNode::BuildFactoryNodeUid(MeshNode->GetUniqueID());
-					UInterchangeMeshFactoryNode * MeshFactoryNode = Cast<UInterchangeMeshFactoryNode>(NodeContainer->GetFactoryNode(FactoryNodeUid)))
+					UInterchangeMeshFactoryNode* MeshFactoryNode = Cast<UInterchangeMeshFactoryNode>(NodeContainer->GetFactoryNode(FactoryNodeUid)))
 				{
 					if (int32 NumPrimvar; MeshNode->GetInt32Attribute(USD::Primvar::Number, NumPrimvar))
 					{
@@ -479,11 +769,13 @@ namespace UE::InterchangeUsdPipeline::Private
 							}
 
 							FString ShaderNodeTextureSampleAttribute = USD::Primvar::ShaderNodeTextureSample + FString::FromInt(Index);
-							if (FString ShaderNodeTextureSampleUID;  MeshNode->GetStringAttribute(ShaderNodeTextureSampleAttribute, ShaderNodeTextureSampleUID))
+							if (FString ShaderNodeTextureSampleUID;
+								MeshNode->GetStringAttribute(ShaderNodeTextureSampleAttribute, ShaderNodeTextureSampleUID))
 							{
-								// At this stage the USD pipeline should be the latest one in the stack, we should retrieve the MaterialExpression Factory Node associated to this ShaderNode
+								// At this stage the USD pipeline should be the latest one in the stack, we should retrieve the MaterialExpression
+								// Factory Node associated to this ShaderNode
 								FString TextureSampleFactoryNodeUID = UInterchangeFactoryBaseNode::BuildFactoryNodeUid(ShaderNodeTextureSampleUID);
-								
+
 								if (Cast<UInterchangeMaterialExpressionFactoryNode>(NodeContainer->GetFactoryNode(TextureSampleFactoryNodeUID)))
 								{
 									MeshFactoryNode->AddStringAttribute(ShaderNodeTextureSampleAttribute, TextureSampleFactoryNodeUID);
@@ -492,18 +784,24 @@ namespace UE::InterchangeUsdPipeline::Private
 						}
 					}
 				}
-			});
+			}
+		);
 	}
 
 	class FUsdPrimvarBaker
 	{
 	public:
 
-		FUsdPrimvarBaker(UStaticMesh* StaticMeshAsset, const UInterchangeBaseNodeContainer* NodeContainer, const UInterchangeFactoryBaseNode* FactoryNode)
-			: NodeContainer{ NodeContainer }
-			, FactoryNode{ FactoryNode }
-			, StaticMeshAsset{ StaticMeshAsset }
-		{}
+		FUsdPrimvarBaker(
+			UStaticMesh* StaticMeshAsset,
+			const UInterchangeBaseNodeContainer* NodeContainer,
+			const UInterchangeFactoryBaseNode* FactoryNode
+		)
+			: NodeContainer{NodeContainer}
+			, FactoryNode{FactoryNode}
+			, StaticMeshAsset{StaticMeshAsset}
+		{
+		}
 
 		void Bake()
 		{
@@ -533,21 +831,22 @@ namespace UE::InterchangeUsdPipeline::Private
 					FactoryNode->GetStringAttribute(USD::Primvar::Name + FString::FromInt(Index), PrimvarName);
 					FactoryNode->GetBooleanAttribute(USD::Primvar::TangentSpace + FString::FromInt(Index), bTangentSpace);
 
-					TVertexInstanceAttributesRef<FVector4f> Primvar = MeshDescriptionPrimvar.VertexInstanceAttributes().GetAttributesRef<FVector4f>(*PrimvarName);
-					
-					if (!(Primvar.IsValid()
-						&& Primvar.GetNumElements() == VertexColors.GetNumElements()
-						&& Primvar.GetNumElements() == Normals.GetNumElements()))
+					TVertexInstanceAttributesRef<FVector4f> Primvar = MeshDescriptionPrimvar.VertexInstanceAttributes().GetAttributesRef<FVector4f>(
+						*PrimvarName
+					);
+
+					if (!(Primvar.IsValid() && Primvar.GetNumElements() == VertexColors.GetNumElements()
+						  && Primvar.GetNumElements() == Normals.GetNumElements()))
 					{
 						continue;
 					}
 
-					// We put the primvar in the Color/Normal channel, that way during the conversion to a DynamicMesh it will also handle the non-manifold case
+					// We put the primvar in the Color/Normal channel, that way during the conversion to a DynamicMesh it will also handle the
+					// non-manifold case
 					for (const int32 VertexInstanceID : MeshDescriptionPrimvar.VertexInstances().GetElementIDs())
 					{
-						bTangentSpace
-							? Normals.Set(VertexInstanceID, Primvar.Get(VertexInstanceID))
-							: VertexColors.Set(VertexInstanceID, Primvar.Get(VertexInstanceID));
+						bTangentSpace ? Normals.Set(VertexInstanceID, Primvar.Get(VertexInstanceID))
+									  : VertexColors.Set(VertexInstanceID, Primvar.Get(VertexInstanceID));
 					}
 
 					FDynamicMesh3 DynamicMesh;
@@ -557,10 +856,7 @@ namespace UE::InterchangeUsdPipeline::Private
 					const FDynamicMeshAABBTree3 DetailSpatial(&DynamicMesh);
 					FMeshBakerDynamicMeshSampler DetailSampler(&DynamicMesh, &DetailSpatial);
 					TSharedPtr<FMeshPropertyMapEvaluator> PropertyEval = MakeShared<FMeshPropertyMapEvaluator>();
-					PropertyEval->Property =
-						bTangentSpace
-						? EMeshPropertyMapType::Normal
-						: EMeshPropertyMapType::VertexColor;
+					PropertyEval->Property = bTangentSpace ? EMeshPropertyMapType::Normal : EMeshPropertyMapType::VertexColor;
 
 					constexpr int32 TextureSize = 1024;
 					const FImageDimensions ImageDimensions(TextureSize, TextureSize);
@@ -578,11 +874,12 @@ namespace UE::InterchangeUsdPipeline::Private
 					Baker.SetGutterSize(4);
 					Baker.SetTileSize(TextureSize);
 					Baker.Bake();
-										
+
 					FTexture2DBuilder TextureBuilder;
 					const FTexture2DBuilder::ETextureType TextureType = GetTextureType(TextureSampleExpression);
 
-					TextureBuilder.InitializeAndReplaceExistingTexture(Cast<UTexture2D>(TextureSampleExpression->Texture.Get()), TextureType, ImageDimensions);
+					TextureBuilder
+						.InitializeAndReplaceExistingTexture(Cast<UTexture2D>(TextureSampleExpression->Texture.Get()), TextureType, ImageDimensions);
 					TextureBuilder.Copy(*Baker.GetBakeResults(0)[0]);
 					TextureBuilder.Commit();
 
@@ -606,11 +903,14 @@ namespace UE::InterchangeUsdPipeline::Private
 
 			FString ShaderNodeTextureSampleUID;
 			FactoryNode->GetStringAttribute(USD::Primvar::ShaderNodeTextureSample + FString::FromInt(IndexPrimvar), ShaderNodeTextureSampleUID);
-			if (UInterchangeMaterialExpressionFactoryNode* TextureSampleExpressionFactoryNode = Cast<UInterchangeMaterialExpressionFactoryNode>(NodeContainer->GetFactoryNode(ShaderNodeTextureSampleUID)))
+			if (UInterchangeMaterialExpressionFactoryNode* TextureSampleExpressionFactoryNode = Cast<UInterchangeMaterialExpressionFactoryNode>(
+					NodeContainer->GetFactoryNode(ShaderNodeTextureSampleUID)
+				))
 			{
-				if (FString ExpressionPath; TextureSampleExpressionFactoryNode->GetStringAttribute(Factory::Expression::Path.ToString(), ExpressionPath))
+				if (FString ExpressionPath;
+					TextureSampleExpressionFactoryNode->GetStringAttribute(Factory::Expression::Path.ToString(), ExpressionPath))
 				{
-					MaterialExpressionTextureSample = Cast<UMaterialExpressionTextureSample>(FSoftObjectPath{ ExpressionPath }.TryLoad());
+					MaterialExpressionTextureSample = Cast<UMaterialExpressionTextureSample>(FSoftObjectPath{ExpressionPath}.TryLoad());
 				}
 			}
 
@@ -635,7 +935,7 @@ namespace UE::InterchangeUsdPipeline::Private
 
 				for (const int32 VertexInstanceID : MeshDescription->VertexInstances().GetElementIDs())
 				{
-					VertexInstanceUVs.Set(VertexInstanceID, UVChannel, FVector2f{ UVs[VertexInstanceID] });
+					VertexInstanceUVs.Set(VertexInstanceID, UVChannel, FVector2f{UVs[VertexInstanceID]});
 				}
 			}
 
@@ -663,17 +963,17 @@ namespace UE::InterchangeUsdPipeline::Private
 
 		UE::Geometry::FTexture2DBuilder::ETextureType GetTextureType(UMaterialExpressionTextureSample* TextureSampleExpression)
 		{
-			using namespace UE::Geometry;			
+			using namespace UE::Geometry;
 			switch (TextureSampleExpression->SamplerType)
 			{
-			case EMaterialSamplerType::SAMPLERTYPE_Color:
-				return FTexture2DBuilder::ETextureType::Color;
+				case EMaterialSamplerType::SAMPLERTYPE_Color:
+					return FTexture2DBuilder::ETextureType::Color;
 
-			case EMaterialSamplerType::SAMPLERTYPE_Normal:
-				return FTexture2DBuilder::ETextureType::NormalMap;
+				case EMaterialSamplerType::SAMPLERTYPE_Normal:
+					return FTexture2DBuilder::ETextureType::NormalMap;
 
-			default:
-				return FTexture2DBuilder::ETextureType::ColorLinear;
+				default:
+					return FTexture2DBuilder::ETextureType::ColorLinear;
 			}
 		}
 
@@ -695,8 +995,7 @@ namespace UE::InterchangeUsdPipeline::Private
 
 			static FORCEINLINE bool Matches(ElementInitType A, ElementInitType B)
 			{
-				return FMath::IsNearlyEqual(A.X, B.X, Threshold) &&
-					   FMath::IsNearlyEqual(A.Y, B.Y, Threshold);
+				return FMath::IsNearlyEqual(A.X, B.X, Threshold) && FMath::IsNearlyEqual(A.Y, B.Y, Threshold);
 			}
 		};
 
@@ -729,7 +1028,8 @@ namespace UE::InterchangeUsdPipeline::Private
 		TArray<UMaterialExpressionTextureSample*> MaterialExpressionTextureSamples;
 		int32 UVChannel = 0;
 	};
-#endif // WITH_EDITORONLY_DATA
+#endif	  // WITH_EDITORONLY_DATA
+#endif	  // USE_USD_SDK
 }	 // namespace UE::InterchangeUsdPipeline::Private
 
 UInterchangeUsdPipeline::UInterchangeUsdPipeline()
@@ -738,8 +1038,14 @@ UInterchangeUsdPipeline::UInterchangeUsdPipeline()
 {
 }
 
-void UInterchangeUsdPipeline::ExecutePostFactoryPipeline(const UInterchangeBaseNodeContainer* BaseNodeContainer, const FString& NodeKey, UObject* CreatedAsset, bool bIsAReimport)
+void UInterchangeUsdPipeline::ExecutePostFactoryPipeline(
+	const UInterchangeBaseNodeContainer* BaseNodeContainer,
+	const FString& NodeKey,
+	UObject* CreatedAsset,
+	bool bIsAReimport
+)
 {
+#if USE_USD_SDK && WITH_EDITORONLY_DATA
 	if (!BaseNodeContainer || !CreatedAsset)
 	{
 		return;
@@ -751,14 +1057,13 @@ void UInterchangeUsdPipeline::ExecutePostFactoryPipeline(const UInterchangeBaseN
 		return;
 	}
 
-#if WITH_EDITORONLY_DATA
 	if (UStaticMesh* StaticMeshAsset = Cast<UStaticMesh>(CreatedAsset))
 	{
-		using namespace UE::InterchangeUsdPipeline::Private;		
+		using namespace UE::InterchangeUsdPipeline::Private;
 		FUsdPrimvarBaker Baker(StaticMeshAsset, BaseNodeContainer, FactoryNode);
 		Baker.Bake();
 	}
-#endif // WITH_EDITORONLY_DATA
+#endif	  // USE_USD_SDK && WITH_EDITORONLY_DATA
 }
 
 void UInterchangeUsdPipeline::ExecutePipeline(
@@ -771,6 +1076,7 @@ void UInterchangeUsdPipeline::ExecutePipeline(
 
 	Super::ExecutePipeline(NodeContainer, InSourceDatas, ContentBasePath);
 
+#if USE_USD_SDK
 	if (!NodeContainer)
 	{
 		return;
@@ -780,6 +1086,8 @@ void UInterchangeUsdPipeline::ExecutePipeline(
 
 	RepairMaterialInstanceFactoryNodes(*NodeContainer);
 
+	ProcessVolumeNodes(*NodeContainer);
+
 	if (!bImportPseudoRoot)
 	{
 		RemovePseudoRootFactoryNode(*NodeContainer);
@@ -787,10 +1095,13 @@ void UInterchangeUsdPipeline::ExecutePipeline(
 
 #if WITH_EDITORONLY_DATA
 	AddPrimvarAttributesToFactoryNodes(NodeContainer);
-#endif // WITH_EDITORONLY_DATA
+#endif	  // WITH_EDITORONLY_DATA
+#endif	  // USE_USD_SDK
 }
 
 bool UInterchangeUsdPipeline::CanExecuteOnAnyThread(EInterchangePipelineTask PipelineTask)
 {
-	return true;
+	// TODO: Only the volume node material handling needs to be on the game thread. Maybe that can be split into
+	// a separate USDGameThreadPipeline?
+	return false;
 }
diff --git a/Engine/Plugins/Runtime/USDCore/Source/USDUtilities/Private/USDConversionUtils.cpp b/Engine/Plugins/Runtime/USDCore/Source/USDUtilities/Private/USDConversionUtils.cpp
--- a/Engine/Plugins/Runtime/USDCore/Source/USDUtilities/Private/USDConversionUtils.cpp
+++ b/Engine/Plugins/Runtime/USDCore/Source/USDUtilities/Private/USDConversionUtils.cpp
@@ -15,6 +15,7 @@
 #include "USDMemory.h"
 #include "USDObjectUtils.h"
 #include "USDProjectSettings.h"
+#include "USDShadeConversion.h"
 #include "USDSkeletalDataConversion.h"
 #include "USDTypesConversion.h"
 #include "USDUnrealAssetInfo.h"
@@ -50,6 +51,7 @@
 #include "GroomCache.h"
 #include "InstancedFoliageActor.h"
 #include "LandscapeProxy.h"
+#include "Materials/Material.h"
 #include "Misc/PackageName.h"
 #include "PhysicsEngine/PhysicsAsset.h"
 #include "Sound/AmbientSound.h"
@@ -117,6 +119,15 @@ static FAutoConsoleVariableRef CVarCheapUniquePrimPathGeneration(
 	)
 );
 
+static bool bRemoveDuplicates = true;
+static FAutoConsoleVariableRef CVarRemoveDuplicates(
+	TEXT("USD.Volume.RemoveDuplicateAnimatedFrames"),
+	bRemoveDuplicates,
+	TEXT(
+		"If this is true (default), the contents of a .VDB file are added only once to animated Sparse Volume Textures (SVT), even if the same file shows in multiple different time samples. If this is false, every OpenVDBAsset prim filePath time sample is parsed as a new frame on the animated SVT."
+	)
+);
+
 namespace USDConversionUtilsImpl
 {
 	/** Show some warnings if the UVSet primvars show some unsupported/problematic behavior */
@@ -1691,6 +1702,395 @@ EUsdDrawMode UsdUtils::GetAppliedDrawMode(const pxr::UsdPrim& Prim)
 	}
 }
 
+TMap<FString, UsdUtils::FVolumePrimInfo> UsdUtils::GetVolumeInfoByFilePathHash(const pxr::UsdPrim& VolumePrim)
+{
+	// Collect all the .vdb files that this prim wants to parse, and the desired fields/grids from them.
+	//
+	// In VDB terminology a "grid" is essentially a 3D texture, and can have formats like float, double3, half, etc.
+	// In USD the analogous term is "field", but essentially means the same thing. Possibly the terminology is abstracted
+	// to also fit the Field3D library, which we don't support it. Field/grid will be used interchangeably here.
+	//
+	// USD is very flexible and allows the user to reference specific grids from of each .vdb file. The syntax makes it
+	// difficult to find out at once all the grids we'll need to parse from each the .vdb files, so here we need to group them
+	// up first before deferring to the SparseVolumeTextureFactory.
+	//
+	// Note that USD allows a single Volume prim to reference grids from multiple .vdb files, and to also "timeSample" the
+	// file reference to allow for volume animations. This means that in UE a "Volume" prim corresponds to a single
+	// HeterogeneousVolumeActor, but which in turn can have any number of Sparse Volume Textures (one for each .vdb file referenced).
+
+	pxr::UsdVolVolume Volume{VolumePrim};
+	if (!Volume)
+	{
+		return {};
+	}
+
+	TMap<FString, FVolumePrimInfo> FilePathHashToInfo;
+
+	FScopedUsdAllocs UsdAllocs;
+
+	pxr::UsdStageRefPtr Stage = Volume.GetPrim().GetStage();
+
+	const std::map<pxr::TfToken, pxr::SdfPath>& FieldMap = Volume.GetFieldPaths();
+	for (std::map<pxr::TfToken, pxr::SdfPath>::const_iterator Iter = FieldMap.cbegin(); Iter != FieldMap.end(); ++Iter)
+	{
+		// This field name is the name of the field for the Volume prim, which can be anything and differ from the
+		// grid name within the .vdb files
+		const pxr::TfToken& FieldName = Iter->first;
+		const pxr::SdfPath& AssetPrimPath = Iter->second;
+
+		pxr::UsdPrim OpenVDBPrim = Stage->GetPrimAtPath(AssetPrimPath);
+		pxr::UsdVolOpenVDBAsset OpenVDBPrimSchema = pxr::UsdVolOpenVDBAsset{OpenVDBPrim};
+		if (OpenVDBPrimSchema)
+		{
+			pxr::UsdAttribute FilePathAttr = OpenVDBPrimSchema.GetFilePathAttr();
+
+			FString ResolvedVDBPath = UsdUtils::GetResolvedAssetPath(FilePathAttr, pxr::UsdTimeCode::Default());
+
+			// Find timesampled paths, if any
+			TArray<double> TimeSamplePathTimeCodes;
+			TArray<int32> TimeSamplePathIndices;
+			TArray<FString> TimeSamplePaths;
+
+			TMap<FString, int32> PathToIndex;
+
+			std::vector<double> TimeSamples;
+			if (FilePathAttr.GetTimeSamples(&TimeSamples) && TimeSamples.size() > 0)
+			{
+				UE::FSdfLayerOffset CombinedOffset = UsdUtils::GetPrimToStageOffset(UE::FUsdPrim{OpenVDBPrim});
+
+				TimeSamplePathTimeCodes.Reserve(TimeSamples.size());
+				TimeSamplePaths.Reserve(TimeSamples.size());
+				TimeSamplePathIndices.Reserve(TimeSamples.size());
+				for (double TimeSample : TimeSamples)
+				{
+					// We always want to store on the AssetUserData (which is where this stuff will end up in)
+					// the time codes in the layer where the actual OpenVDBPrim is authored. If that layer is referenced
+					// by a parent layer through an offset and scale, TimeSample will contain that offset and scale here,
+					// which we need to undo
+					double LayerLocalTimeCode = (TimeSample - CombinedOffset.Offset) / CombinedOffset.Scale;
+					TimeSamplePathTimeCodes.Add(LayerLocalTimeCode);
+
+					FString ResolvedFramePath = UsdUtils::GetResolvedAssetPath(FilePathAttr, TimeSample);
+
+					// If we had no default time sample to act as the "main file", take the first frame
+					if (ResolvedVDBPath.IsEmpty())
+					{
+						ResolvedVDBPath = ResolvedFramePath;
+					}
+
+					if (bRemoveDuplicates)
+					{
+						if (int32* FoundIndex = PathToIndex.Find(ResolvedFramePath))
+						{
+							TimeSamplePathIndices.Add(*FoundIndex);
+						}
+						else
+						{
+							TimeSamplePaths.Add(ResolvedFramePath);
+
+							const int32 NewIndex = TimeSamplePaths.Num() - 1;
+							PathToIndex.Add(ResolvedFramePath, NewIndex);
+							TimeSamplePathIndices.Add(NewIndex);
+						}
+					}
+					else
+					{
+						TimeSamplePaths.Add(ResolvedFramePath);
+
+						const int32 NewIndex = TimeSamplePaths.Num() - 1;
+						TimeSamplePathIndices.Add(NewIndex);
+					}
+				}
+			}
+
+			// Hash all the relevant file paths here: The collection of file paths to parse determines the SVT, and
+			// we want one FSparseVolumeTextureInfo per SVT
+			FString FilePathHashString;
+			{
+				FSHA1 SHA1;
+				SHA1.UpdateWithString(*ResolvedVDBPath, ResolvedVDBPath.Len());
+				for (const FString& TimeSamplePath : TimeSamplePaths)
+				{
+					SHA1.UpdateWithString(*TimeSamplePath, TimeSamplePath.Len());
+				}
+				SHA1.Final();
+
+				FSHAHash FilePathHash;
+				SHA1.GetHash(&FilePathHash.Hash[0]);
+
+				FilePathHashString = FilePathHash.ToString();
+			}
+
+			if (!ResolvedVDBPath.IsEmpty())
+			{
+				FVolumePrimInfo& SparseVolumeTextureInfo = FilePathHashToInfo.FindOrAdd(FilePathHashString);
+				SparseVolumeTextureInfo.SourceOpenVDBAssetPrimPaths.AddUnique(UsdToUnreal::ConvertPath(AssetPrimPath));
+				SparseVolumeTextureInfo.SourceVDBFilePath = ResolvedVDBPath;
+				SparseVolumeTextureInfo.TimeSamplePathTimeCodes = MoveTemp(TimeSamplePathTimeCodes);
+				SparseVolumeTextureInfo.TimeSamplePathIndices = MoveTemp(TimeSamplePathIndices);
+				SparseVolumeTextureInfo.TimeSamplePaths = MoveTemp(TimeSamplePaths);
+
+				FString FieldNameStr = UsdToUnreal::ConvertToken(FieldName);
+				SparseVolumeTextureInfo.VolumeFieldNames.AddUnique(FieldNameStr);
+
+				pxr::TfToken GridName;
+				pxr::UsdAttribute Attr = OpenVDBPrimSchema.GetFieldNameAttr();
+				if (Attr && Attr.Get<pxr::TfToken>(&GridName))
+				{
+					FString GridNameStr = UsdToUnreal::ConvertToken(GridName);
+
+					// Note we want this to add an entry to SparseVolumeTexture.GridNameToChannelNames even if we won't
+					// find the schema on the prim, as we'll use these entries to make sure the generated Sparse Volume
+					// Texture contains theses desired fields
+					TMap<FString, FString>& ChannelToComponent = SparseVolumeTextureInfo.GridNameToChannelComponentMapping.FindOrAdd(GridNameStr);
+
+					if (UsdUtils::PrimHasSchema(OpenVDBPrim.GetPrim(), UnrealIdentifiers::SparseVolumeTextureAPI))
+					{
+						// Parse desired data types for AttributesA and AttributesB channels
+						TFunction<void(pxr::TfToken, TOptional<ESparseVolumeAttributesFormat>&)> HandleAttribute =
+							[&OpenVDBPrim,
+							 &AssetPrimPath,
+							 &ResolvedVDBPath](pxr::TfToken AttrName, TOptional<ESparseVolumeAttributesFormat>& AttributeFormat)
+						{
+							using FormatMapType = std::unordered_map<pxr::TfToken, ESparseVolumeAttributesFormat, pxr::TfHash>;
+							using InverseFormatMapType = std::unordered_map<ESparseVolumeAttributesFormat, FString>;
+
+							// clang-format off
+							const static FormatMapType FormatMap = {
+								{pxr::TfToken("unorm8"),  ESparseVolumeAttributesFormat::Unorm8},
+								{pxr::TfToken("float16"), ESparseVolumeAttributesFormat::Float16},
+								{pxr::TfToken("float32"), ESparseVolumeAttributesFormat::Float32},
+							};
+							const static InverseFormatMapType InverseFormatMap = {
+								{ESparseVolumeAttributesFormat::Unorm8,  TEXT("unorm8")},
+								{ESparseVolumeAttributesFormat::Float16, TEXT("float16")},
+								{ESparseVolumeAttributesFormat::Float32, TEXT("float32")},
+							};
+							// clang-format on
+
+							pxr::TfToken DataType;
+							pxr::UsdAttribute AttrA = OpenVDBPrim.GetAttribute(AttrName);
+							if (AttrA && AttrA.Get<pxr::TfToken>(&DataType))
+							{
+								FormatMapType::const_iterator Iter = FormatMap.find(DataType);
+								if (Iter != FormatMap.end())
+								{
+									ESparseVolumeAttributesFormat TargetFormat = Iter->second;
+
+									// Check in case multiple OpenVDBAsset prims want different values for the data type
+									const bool bIsSet = AttributeFormat.IsSet();
+									if (bIsSet && AttributeFormat.GetValue() != TargetFormat)
+									{
+										const FString* ExistingFormat = nullptr;
+										InverseFormatMapType::const_iterator ExistingIter = InverseFormatMap.find(AttributeFormat.GetValue());
+										if (ExistingIter != InverseFormatMap.end())
+										{
+											ExistingFormat = &ExistingIter->second;
+										}
+
+										UE_LOG(
+											LogUsd,
+											Warning,
+											TEXT(
+												"OpenVDBAsset prims disagree on the attribute channel format for the Sparse Volume Texture generated for VDB file '%s' (encountered '%s' and '%s'). If there are multiple opinions for the attribute channel formats from different OpenVDBAsset prims, they must all agree!"
+											),
+											*ResolvedVDBPath,
+											*UsdToUnreal::ConvertToken(DataType),
+											ExistingFormat ? **ExistingFormat : TEXT("unknown")
+										);
+									}
+									else if (!bIsSet)
+									{
+										AttributeFormat = TargetFormat;
+									}
+								}
+								else
+								{
+									UE_LOG(
+										LogUsd,
+										Warning,
+										TEXT(
+											"Invalid Sparse Volume Texture attribute channel format '%s'. Available formats: 'unorm8', 'float16' and 'float32'."
+										),
+										*UsdToUnreal::ConvertToken(DataType),
+										*UsdToUnreal::ConvertPath(AssetPrimPath)
+									);
+								}
+							}
+						};
+						HandleAttribute(UnrealIdentifiers::UnrealSVTAttributesADataType, SparseVolumeTextureInfo.AttributesAFormat);
+						HandleAttribute(UnrealIdentifiers::UnrealSVTAttributesBDataType, SparseVolumeTextureInfo.AttributesBFormat);
+
+						// Parse desired channel assignment
+						pxr::VtArray<pxr::TfToken> DesiredChannels;
+						pxr::VtArray<pxr::TfToken> DesiredComponents;
+						pxr::UsdAttribute ComponentsAttr = OpenVDBPrim.GetAttribute(UnrealIdentifiers::UnrealSVTMappedGridComponents);
+						pxr::UsdAttribute ChannelsAttr = OpenVDBPrim.GetAttribute(UnrealIdentifiers::UnrealSVTMappedAttributeChannels);
+						if (ChannelsAttr &&														 //
+							ComponentsAttr &&													 //
+							ChannelsAttr.Get<pxr::VtArray<pxr::TfToken>>(&DesiredChannels) &&	 //
+							ComponentsAttr.Get<pxr::VtArray<pxr::TfToken>>(&DesiredComponents))
+						{
+							// These must always match of course
+							if (DesiredChannels.size() == DesiredComponents.size())
+							{
+								// If we have more than one OpenVDBAsset prim reading from the same VDB file, the declared component to
+								// channel mappings must be compatible
+								for (uint32 Index = 0; Index < DesiredChannels.size(); ++Index)
+								{
+									FString Channel = UsdToUnreal::ConvertString(DesiredChannels[Index]);
+									FString Component = UsdToUnreal::ConvertString(DesiredComponents[Index]);
+
+									if (const FString* ExistingComponentMapping = ChannelToComponent.Find(Channel))
+									{
+										if (Component != *ExistingComponentMapping)
+										{
+											UE_LOG(
+												LogUsd,
+												Warning,
+												TEXT(
+													"Found multiple OpenVDBAsset prims (including '%s') targetting the same grid '%s', but with with conflicting grid component to Sparse Volume Texture attribute channel mapping (for example, both components '%s' and '%s' are mapped to the same channel '%s', which is not allowed)"
+												),
+												*UsdToUnreal::ConvertPath(AssetPrimPath),
+												*GridNameStr,
+												**ExistingComponentMapping,
+												*Component,
+												*Channel
+											);
+										}
+									}
+									else
+									{
+										ChannelToComponent.Add(Channel, Component);
+									}
+								}
+							}
+							else
+							{
+								UE_LOG(
+									LogUsd,
+									Warning,
+									TEXT(
+										"Failed to parse custom component to attribute mapping from OpenVDBAsset prim '%s': The '%s' and '%s' attributes should have the same number of entries, but the former has %u entries while the latter has %u"
+									),
+									*UsdToUnreal::ConvertPath(AssetPrimPath),
+									*UsdToUnreal::ConvertToken(UnrealIdentifiers::UnrealSVTMappedGridComponents),
+									*UsdToUnreal::ConvertToken(UnrealIdentifiers::UnrealSVTMappedAttributeChannels),
+									DesiredComponents.size(),
+									DesiredChannels.size()
+								);
+							}
+						}
+					}
+				}
+			}
+			else
+			{
+				UE_LOG(
+					LogUsd,
+					Warning,
+					TEXT("Failed to find the VDB file '%s' referenced by OpenVDBAsset prim at path '%s'"),
+					*ResolvedVDBPath,
+					*UsdToUnreal::ConvertPath(AssetPrimPath)
+				);
+			}
+		}
+		else
+		{
+			UE_LOG(
+				LogUsd,
+				Warning,
+				TEXT("Failed to find an OpenVDBAsset prim at path '%s' for field '%s' of prim '%s'"),
+				*UsdToUnreal::ConvertPath(AssetPrimPath),
+				*UsdToUnreal::ConvertToken(FieldName),
+				*UsdToUnreal::ConvertPath(Volume.GetPrim().GetPath())
+			);
+		}
+	}
+
+	return FilePathHashToInfo;
+}
+
+TMultiMap<FString, FString> UsdUtils::GetVolumeMaterialParameterToFieldNameMap(const pxr::UsdPrim& VolumePrim)
+{
+	if (!VolumePrim)
+	{
+		return {};
+	}
+
+	FScopedUsdAllocs Allocs;
+
+	if (!UsdUtils::PrimHasSchema(VolumePrim, UnrealIdentifiers::SparseVolumeTextureAPI))
+	{
+		return {};
+	}
+
+	pxr::UsdAttribute FieldsAttr = VolumePrim.GetAttribute(UnrealIdentifiers::UnrealSVTMappedFields);
+	pxr::UsdAttribute ParametersAttr = VolumePrim.GetAttribute(UnrealIdentifiers::UnrealSVTMappedMaterialParameters);
+	if (!FieldsAttr || !ParametersAttr)
+	{
+		return {};
+	}
+
+	pxr::VtArray<pxr::TfToken> MappedFields;
+	pxr::VtArray<pxr::TfToken> MappedParameters;
+	if (!FieldsAttr.Get<pxr::VtArray<pxr::TfToken>>(&MappedFields) || !ParametersAttr.Get<pxr::VtArray<pxr::TfToken>>(&MappedParameters))
+	{
+		return {};
+	}
+
+	if (MappedFields.size() != MappedParameters.size())
+	{
+		UE_LOG(
+			LogUsd,
+			Warning,
+			TEXT(
+				"Failed to parse custom parsed texture to material parameter mapping from volume prim '%s': The '%s' and '%s' attributes should have the same number of entries, but the former has %u entries while the latter has %u"
+			),
+			*UsdToUnreal::ConvertPath(VolumePrim.GetPrimPath()),
+			*UsdToUnreal::ConvertToken(UnrealIdentifiers::UnrealSVTMappedFields),
+			*UsdToUnreal::ConvertToken(UnrealIdentifiers::UnrealSVTMappedMaterialParameters),
+			MappedFields.size(),
+			MappedParameters.size()
+		);
+		return {};
+	}
+
+	TMultiMap<FString, FString> MaterialParameterToFieldName;
+	for (uint32 Index = 0; Index < MappedFields.size(); ++Index)
+	{
+		FString FieldName = UsdToUnreal::ConvertString(MappedFields[Index]);
+		FString MaterialParameter = UsdToUnreal::ConvertString(MappedParameters[Index]);
+		MaterialParameterToFieldName.Add(MaterialParameter, FieldName);
+	}
+
+	return MaterialParameterToFieldName;
+}
+
+TArray<FString> UsdUtils::GetSparseVolumeTextureParameterNames(const UMaterial* Material)
+{
+	if (!Material)
+	{
+		return {};
+	}
+
+	TArray<FString> Result;
+
+	TMap<FMaterialParameterInfo, FMaterialParameterMetadata> SparseVolumeTextureParameters;
+	Material->GetAllParametersOfType(EMaterialParameterType::SparseVolumeTexture, SparseVolumeTextureParameters);
+
+	Result.Reserve(SparseVolumeTextureParameters.Num());
+	for (const TPair<FMaterialParameterInfo, FMaterialParameterMetadata>& ParameterPair : SparseVolumeTextureParameters)
+	{
+		const FMaterialParameterInfo& ParameterInfo = ParameterPair.Key;
+		FString ParameterName = ParameterInfo.Name.ToString();
+
+		Result.Add(ParameterName);
+	}
+
+	return Result;
+}
+
 TArray<TUsdStore<pxr::UsdPrim>> UsdUtils::GetAllPrimsOfType(
 	const pxr::UsdPrim& StartPrim,
 	const pxr::TfType& SchemaType,
diff --git a/Engine/Plugins/Runtime/USDCore/Source/USDUtilities/Private/USDGeomMeshConversion.cpp b/Engine/Plugins/Runtime/USDCore/Source/USDUtilities/Private/USDGeomMeshConversion.cpp
--- a/Engine/Plugins/Runtime/USDCore/Source/USDUtilities/Private/USDGeomMeshConversion.cpp
+++ b/Engine/Plugins/Runtime/USDCore/Source/USDUtilities/Private/USDGeomMeshConversion.cpp
@@ -80,6 +80,7 @@
 #include "pxr/usd/usdSkel/skeletonQuery.h"
 #include "pxr/usd/usdSkel/skinningQuery.h"
 #include "pxr/usd/usdSkel/utils.h"
+#include "pxr/usd/usdVol/volume.h"
 #include "pxr/usdImaging/usdImaging/implicitSurfaceMeshUtils.h"
 #include "USDIncludesEnd.h"
 
@@ -982,7 +983,8 @@ namespace UE::UsdGeomMeshConversion::Private
 		}
 		// Check for primitive types directly (don't just check for Gprim here because if we get something else like a UsdVolVolume or
 		// UsdGeomBasisCurves ConvertGeomPrimitive will return false and we will dump our MeshDescription)
-		else if (Prim.IsA<pxr::UsdGeomCapsule>() || Prim.IsA<pxr::UsdGeomCone>() || Prim.IsA<pxr::UsdGeomCube>() || Prim.IsA<pxr::UsdGeomCylinder>() || Prim.IsA<pxr::UsdGeomPlane>() || Prim.IsA<pxr::UsdGeomSphere>())
+		else if (Prim.IsA<pxr::UsdGeomCapsule>() || Prim.IsA<pxr::UsdGeomCone>() || Prim.IsA<pxr::UsdGeomCube>() || Prim.IsA<pxr::UsdGeomCylinder>()
+				 || Prim.IsA<pxr::UsdGeomPlane>() || Prim.IsA<pxr::UsdGeomSphere>())
 		{
 			bSuccess = UsdToUnreal::ConvertGeomPrimitive(Prim, OutMeshDescription, OutMaterialAssignments, Options);
 		}
@@ -1125,7 +1127,7 @@ namespace UE::UsdGeomMeshConversion::Private
 		return UsdUtils::CombinePrimvarsIntoUVSets(AllPrimvars, PreferredPrimvars);
 	}
 
-	struct FUsdRawBoneInfluence // just weight, vertex, and Bone, sorted in CopySkinWeightsToAttribute
+	struct FUsdRawBoneInfluence	   // just weight, vertex, and Bone, sorted in CopySkinWeightsToAttribute
 	{
 		float Weight;
 		int32 VertexIndex;
@@ -1194,8 +1196,8 @@ namespace UE::UsdGeomMeshConversion::Private
 		TOptional<int32> ProvidedNumUVSets;
 		int32 MaterialIndexOffset = 0;
 
-		//Attributes for SkeletalMesh:
-		TArray <FUsdRawBoneInfluence> Influences;
+		// Attributes for SkeletalMesh:
+		TArray<FUsdRawBoneInfluence> Influences;
 	};
 
 	namespace USDSkeletalDataConversion
@@ -1233,12 +1235,14 @@ namespace UE::UsdGeomMeshConversion::Private
 			// Make a copy of the weights, sort them by vertex id and go by equal vertex-id strides.
 			// We could do an indirection but the traversal + setup cost is probably not worth it.
 			TArray<FUsdRawBoneInfluence> SortedInfluences(InInfluences);
-			SortedInfluences.Sort([](const FUsdRawBoneInfluence& A, const FUsdRawBoneInfluence& B)
+			SortedInfluences.Sort(
+				[](const FUsdRawBoneInfluence& A, const FUsdRawBoneInfluence& B)
 				{
 					return A.VertexIndex < B.VertexIndex;
-				});
+				}
+			);
 
-			// Do the base skin weights first. We do the alternate skin weights later, since they may require geometric remapping. 
+			// Do the base skin weights first. We do the alternate skin weights later, since they may require geometric remapping.
 			TArray<FBoneWeight> BoneWeights;
 			for (int32 StartStride = 0, EndStride = 0; EndStride != SortedInfluences.Num(); StartStride = EndStride)
 			{
@@ -1277,7 +1281,7 @@ namespace UE::UsdGeomMeshConversion::Private
 				OutSkinWeightsAttribute.Set(InVertexIDMap[VertexIndex], BoneWeights);
 			}
 		}
-	}
+	}	 // namespace USDSkeletalDataConversion
 
 	void CollectPrimvarData(
 		const pxr::UsdPrim& Prim,
@@ -1458,7 +1462,7 @@ namespace UE::UsdGeomMeshConversion::Private
 				InOutMeshData.PointInterpolation = GetAttrInterpolation(PointsAttr, pxr::UsdGeomTokens->vertex);
 			}
 
-			//Apply Skinning if applicable.
+			// Apply Skinning if applicable.
 			if (bProcessSkeletalAttributes)
 			{
 				SkinningQuery->ComputeSkinnedPoints(MeshToSkeletonRestPose, &InOutMeshData.Points, Options.TimeCode);
@@ -1601,7 +1605,7 @@ namespace UE::UsdGeomMeshConversion::Private
 			}
 		}
 
-		//Joint Indices/Weights
+		// Joint Indices/Weights
 		if (bProcessSkeletalAttributes)
 		{
 			const pxr::UsdPrim& SkinningPrim = SkinningQuery->GetPrim();
@@ -1609,8 +1613,8 @@ namespace UE::UsdGeomMeshConversion::Private
 			if (SkelBindingAPI)
 			{
 				// ComputeJointInfluences returns the influences per bone that applies to all the points of the mesh
-				// ComputeVaryingJointInfluences returns the joint influences for each points, expanding the influences to all points if the mesh is rigidly
-				// deformed
+				// ComputeVaryingJointInfluences returns the joint influences for each points, expanding the influences to all points if the mesh is
+				// rigidly deformed
 				pxr::VtArray<int> JointIndices;
 				pxr::VtArray<float> JointWeights;
 				size_t NumPoints = InOutMeshData.Points.size();
@@ -1635,8 +1639,8 @@ namespace UE::UsdGeomMeshConversion::Private
 						// The JointIndices/JointWeights contain the influences data for NumPoints * NumInfluencesPerComponent
 						for (uint32 InfluenceIndex = 0; InfluenceIndex < NumInfluencesPerComponent; ++InfluenceIndex, ++JointIndex)
 						{
-							// BoneWeight could be 0 if the actual number of influences were less than NumInfluencesPerComponent for a given point so just ignore
-							// it
+							// BoneWeight could be 0 if the actual number of influences were less than NumInfluencesPerComponent for a given point so
+							// just ignore it
 							float BoneWeight = JointWeights[JointIndex];
 							if (BoneWeight != 0.f)
 							{
@@ -1669,10 +1673,12 @@ namespace UE::UsdGeomMeshConversion::Private
 								}
 
 								// Use the AnimMapper to produce the indices of the Mesh's joints within the Skeleton's list of joints.
-								// Example: Imagine skeleton had { "Root", "Root/Hip", "Root/Hip/Shoulder", "Root/Hip/Shoulder/Arm", "Root/Hip/Shoulder/Arm/Elbow"
-								// }, and so BoneIndexRemapping was { 0, 1, 2, 3, 4 }. Consider a Mesh that specifies the explicit joints { "Root/Hip/Shoulder",
-								// "Root/Hip/Shoulder/Arm" }, and so uses the indices 0 and 1 to refer to Shoulder and Arm. After the Remap call
-								// SkeletonBoneIndices will hold { 2, 3 }, as those are the indices of Shoulder and Arm within the skeleton's bones
+								// Example: Imagine skeleton had { "Root", "Root/Hip", "Root/Hip/Shoulder", "Root/Hip/Shoulder/Arm",
+								// "Root/Hip/Shoulder/Arm/Elbow"
+								// }, and so BoneIndexRemapping was { 0, 1, 2, 3, 4 }. Consider a Mesh that specifies the explicit joints {
+								// "Root/Hip/Shoulder", "Root/Hip/Shoulder/Arm" }, and so uses the indices 0 and 1 to refer to Shoulder and Arm. After
+								// the Remap call SkeletonBoneIndices will hold { 2, 3 }, as those are the indices of Shoulder and Arm within the
+								// skeleton's bones
 								pxr::VtArray<int> BoneIndexRemapping;
 								if (AnimMapper->Remap(SkeletonBoneIndices, &BoneIndexRemapping))
 								{
@@ -1680,10 +1686,10 @@ namespace UE::UsdGeomMeshConversion::Private
 									{
 										FUsdRawBoneInfluence& Influence = InOutMeshData.Influences[AddedInfluenceIndex];
 
-										// We have to remove our "NumAdditionalBones" offset from the influence's bone index because that's a UE concept that
-										// the BoneIndexRemapping array doesn't really know about. After that, we have a bone index that matches the USD Skeleton
-										// joint order, then we can remap with BoneIndexRemapping and add our NumAdditionalBones back in so that it matches our
-										// USkeleton
+										// We have to remove our "NumAdditionalBones" offset from the influence's bone index because that's a UE
+										// concept that the BoneIndexRemapping array doesn't really know about. After that, we have a bone index that
+										// matches the USD Skeleton joint order, then we can remap with BoneIndexRemapping and add our
+										// NumAdditionalBones back in so that it matches our USkeleton
 										Influence.BoneIndex = NumAdditionalBones + BoneIndexRemapping[Influence.BoneIndex - NumAdditionalBones];
 									}
 								}
@@ -2093,7 +2099,8 @@ namespace UE::UsdGeomMeshConversion::Private
 			InOutMeshData.CreaseLengths,
 			InOutMeshData.CreaseSharpnesses,
 			InOutMeshData.CornerIndices,
-			InOutMeshData.CornerSharpnesses};
+			InOutMeshData.CornerSharpnesses
+		};
 
 		pxr::PxOsdMeshTopology Topology{
 			InOutMeshData.SubdivScheme,
@@ -2101,7 +2108,8 @@ namespace UE::UsdGeomMeshConversion::Private
 			InOutMeshData.FaceVertexCounts,
 			InOutMeshData.FaceVertexIndices,
 			InOutMeshData.HoleIndices,
-			SubdivTags};
+			SubdivTags
+		};
 
 		std::vector<pxr::VtArray<int>> FaceVaryingTopologies;
 		FaceVaryingTopologies.resize(FaceVaryingChannelCounter);
@@ -3805,7 +3813,8 @@ UsdUtils::FUsdPrimMaterialAssignmentInfo UsdUtils::GetPrimMaterialAssignments(
 {
 	TRACE_CPUPROFILER_EVENT_SCOPE(UsdUtils::GetPrimMaterialAssignments);
 
-	if (!UsdPrim)
+	// We handle sections when calling this on the parent Mesh prims, not when calling on the sections themselves
+	if (!UsdPrim || UsdPrim.IsA<pxr::UsdGeomSubset>())
 	{
 		return {};
 	}
@@ -3883,10 +3892,6 @@ UsdUtils::FUsdPrimMaterialAssignmentInfo UsdUtils::GetPrimMaterialAssignments(
 	{
 		pxr::VtArray<int> FaceVertexCounts = UsdGeomMeshImpl::GetFaceVertexCounts(UsdPrim, TimeCode);
 		NumFaces = FaceVertexCounts.size();
-		if (NumFaces < 1)
-		{
-			return Result;
-		}
 
 		if (bProvideMaterialIndices)
 		{
@@ -3908,7 +3913,7 @@ UsdUtils::FUsdPrimMaterialAssignmentInfo UsdUtils::GetPrimMaterialAssignments(
 		}
 	}
 
-	FString MeshPrimPath = UsdToUnreal::ConvertPath(UsdPrim.GetPath());
+	FString PrimPath = UsdToUnreal::ConvertPath(UsdPrim.GetPath());
 
 	bool bNeedsMainAssignment = true;
 
@@ -4068,7 +4073,7 @@ UsdUtils::FUsdPrimMaterialAssignmentInfo UsdUtils::GetPrimMaterialAssignments(
 					Slot.MaterialSource = UnrealMaterial.GetValue();
 					Slot.AssignmentType = UsdUtils::EPrimAssignmentType::UnrealMaterial;
 					Slot.bMeshIsDoubleSided = bIsDoubleSided;
-					Slot.PrimPaths.Add(MeshPrimPath);
+					Slot.PrimPaths.Add(PrimPath);
 
 					bHasMainAssignment = true;
 				}
@@ -4083,7 +4088,7 @@ UsdUtils::FUsdPrimMaterialAssignmentInfo UsdUtils::GetPrimMaterialAssignments(
 					Slot.MaterialSource = UnrealMaterial.GetValue();
 					Slot.AssignmentType = UsdUtils::EPrimAssignmentType::UnrealMaterial;
 					Slot.bMeshIsDoubleSided = bIsDoubleSided;
-					Slot.PrimPaths.Add(MeshPrimPath);
+					Slot.PrimPaths.Add(PrimPath);
 
 					bHasMainAssignment = true;
 				}
@@ -4099,7 +4104,7 @@ UsdUtils::FUsdPrimMaterialAssignmentInfo UsdUtils::GetPrimMaterialAssignments(
 				Slot.MaterialSource = BoundMaterial.GetValue();
 				Slot.AssignmentType = UsdUtils::EPrimAssignmentType::MaterialPrim;
 				Slot.bMeshIsDoubleSided = bIsDoubleSided;
-				Slot.PrimPaths.Add(MeshPrimPath);
+				Slot.PrimPaths.Add(PrimPath);
 
 				bHasMainAssignment = true;
 			}
@@ -4117,7 +4122,7 @@ UsdUtils::FUsdPrimMaterialAssignmentInfo UsdUtils::GetPrimMaterialAssignments(
 				Slot.MaterialSource = DisplayColor.GetValue().ToString();
 				Slot.AssignmentType = UsdUtils::EPrimAssignmentType::DisplayColor;
 				Slot.bMeshIsDoubleSided = bIsDoubleSided;
-				Slot.PrimPaths.Add(MeshPrimPath);
+				Slot.PrimPaths.Add(PrimPath);
 
 				bHasMainAssignment = true;
 			}
@@ -4168,6 +4173,7 @@ UsdUtils::FUsdPrimMaterialAssignmentInfo UsdUtils::GetPrimMaterialAssignments(
 	{
 		Result.Slots.Emplace();
 	}
+
 	return Result;
 }
 
diff --git a/Engine/Plugins/Runtime/USDCore/Source/USDUtilities/Public/USDConversionUtils.h b/Engine/Plugins/Runtime/USDCore/Source/USDUtilities/Public/USDConversionUtils.h
--- a/Engine/Plugins/Runtime/USDCore/Source/USDUtilities/Public/USDConversionUtils.h
+++ b/Engine/Plugins/Runtime/USDCore/Source/USDUtilities/Public/USDConversionUtils.h
@@ -50,6 +50,7 @@ PXR_NAMESPACE_CLOSE_SCOPE
 #endif	  // #if USE_USD_SDK
 
 class UAssetImportData;
+class UMaterial;
 class USceneComponent;
 class UUsdAssetImportData;
 class UUsdAssetUserData;
@@ -338,6 +339,47 @@ namespace UsdUtils
 	 */
 	USDUTILITIES_API EUsdDrawMode GetAppliedDrawMode(const pxr::UsdPrim& Prim);
 
+	/** Copy-pasted from OpenVDBImportOptions.h because that is editor-only, and we want to avoid spreading more #if WITH_EDITORs around */
+	enum class ESparseVolumeAttributesFormat : uint8
+	{
+		Unorm8 = 0,
+		Float16 = 1,
+		Float32 = 2,
+	};
+
+	struct FVolumePrimInfo
+	{
+		FString SourceVDBFilePath;
+
+		TArray<double> TimeSamplePathTimeCodes;
+		TArray<int32> TimeSamplePathIndices;
+		TArray<FString> TimeSamplePaths;	// First item may equal SourceVDBFilePath. May have duplicate entries
+
+		// Fields on the Volume prim that were connected to this .vdb file
+		TArray<FString> VolumeFieldNames;
+		TArray<FString> SourceOpenVDBAssetPrimPaths;
+
+		// e.g. {'velocity': {'AttributesA.R': 'X', 'AttributesA.G': 'Y', 'AttributesA.B': 'Z'}}
+		TMap<FString, TMap<FString, FString>> GridNameToChannelComponentMapping;
+		TOptional<ESparseVolumeAttributesFormat> AttributesAFormat;
+		TOptional<ESparseVolumeAttributesFormat> AttributesBFormat;
+	};
+
+	/**
+	 * Collects an FVolumePrimInfo for each SparseVolumeTexture that should be generated for VolumePrim.
+	 * The map keys are FilePathHashes, because some of these SparseVolumeTextures may consume multiple .vdb files (e.g .animations)
+	 */
+	USDUTILITIES_API TMap<FString, FVolumePrimInfo> GetVolumeInfoByFilePathHash(const pxr::UsdPrim& VolumePrim);
+
+	/**
+	 * Extracts from the volume prim a desired mapping between Unreal material parameter and the volume field name
+	 * (note the corresponding VolumeFieldNames field within FVolumePrimInfo)
+	 */
+	USDUTILITIES_API TMultiMap<FString, FString> GetVolumeMaterialParameterToFieldNameMap(const pxr::UsdPrim& VolumePrim);
+
+	/** Gets the names of all SparseVolumeTexture material parameters within the provided material */
+	USDUTILITIES_API TArray<FString> GetSparseVolumeTextureParameterNames(const UMaterial* Material);
+
 	/**
 	 * Returns all prims of type SchemaType (or a descendant type) in the subtree of prims rooted at StartPrim.
 	 * Stops going down the subtrees when it hits a schema type to exclude.
diff --git a/Engine/Source/Runtime/Interchange/Core/Private/Nodes/InterchangeFactoryBaseNode.cpp b/Engine/Source/Runtime/Interchange/Core/Private/Nodes/InterchangeFactoryBaseNode.cpp
--- a/Engine/Source/Runtime/Interchange/Core/Private/Nodes/InterchangeFactoryBaseNode.cpp
+++ b/Engine/Source/Runtime/Interchange/Core/Private/Nodes/InterchangeFactoryBaseNode.cpp
@@ -267,7 +267,28 @@ bool UInterchangeFactoryBaseNode::SetCustomReferenceObject(const FSoftObjectPath
 
 FString UInterchangeFactoryBaseNode::BuildFactoryNodeUid(const FString& TranslatedNodeUid)
 {
-	return TEXT("Factory_") + TranslatedNodeUid;
+	const static FString FactoryPrefix = TEXT("Factory_");
+
+	if (TranslatedNodeUid.StartsWith(FactoryPrefix))
+	{
+		return TranslatedNodeUid;
+	}
+
+	return FactoryPrefix + TranslatedNodeUid;
+}
+
+FString UInterchangeFactoryBaseNode::BuildTranslatedNodeUid(const FString& InFactoryNodeUid)
+{
+	const static FString FactoryPrefix = TEXT("Factory_");
+
+	FStringView View{InFactoryNodeUid};
+	if (View.StartsWith(FactoryPrefix))
+	{
+		View.RemovePrefix(FactoryPrefix.Len());
+		return FString{View};
+	}
+
+	return InFactoryNodeUid;
 }
 
 void UInterchangeFactoryBaseNode::ApplyAllCustomAttributeToObject(UObject* Object) const
@@ -314,7 +335,7 @@ void UInterchangeFactoryBaseNode::RestoreAllCustomAttributeDelegates()
 	{
 		AttributesAppliedThroughDelegatesKeySet.Add(Key);
 	}
-	
+
 	OnRestoreAllCustomAttributeDelegates();
 }
 
diff --git a/Engine/Source/Runtime/Interchange/Core/Public/Nodes/InterchangeFactoryBaseNode.h b/Engine/Source/Runtime/Interchange/Core/Public/Nodes/InterchangeFactoryBaseNode.h
--- a/Engine/Source/Runtime/Interchange/Core/Public/Nodes/InterchangeFactoryBaseNode.h
+++ b/Engine/Source/Runtime/Interchange/Core/Public/Nodes/InterchangeFactoryBaseNode.h
@@ -294,6 +294,7 @@ public:
 	INTERCHANGECORE_API bool SetCustomReferenceObject(const FSoftObjectPath& AttributeValue);
 
 	static INTERCHANGECORE_API FString BuildFactoryNodeUid(const FString& TranslatedNodeUid);
+	static INTERCHANGECORE_API FString BuildTranslatedNodeUid(const FString& InFactoryNodeUid);
 
 	/**
 	 * Adds the delegates that will read and write the attribute value to a UObject.
