Interchange: Fix broken build on Android.

diff --git a/Engine/Plugins/Interchange/Runtime/Source/Import/Private/Volume/InterchangeSparseVolumeTextureFactory.cpp b/Engine/Plugins/Interchange/Runtime/Source/Import/Private/Volume/InterchangeSparseVolumeTextureFactory.cpp
--- a/Engine/Plugins/Interchange/Runtime/Source/Import/Private/Volume/InterchangeSparseVolumeTextureFactory.cpp
+++ b/Engine/Plugins/Interchange/Runtime/Source/Import/Private/Volume/InterchangeSparseVolumeTextureFactory.cpp
@@ -163,13 +163,17 @@ namespace UE::InterchangeSparseVolumeTextureFactory::Private
 			if (VolumePayload.PayloadFrameIndices.Num() == 0)
 			{
 				// No animation index --> Not animated, so pretend it's the earliest frame possible
-				PayloadAndIndices.Emplace(&VolumePayload, TNumericLimits<int32>::Lowest());
+				FPayloadAndIndex& NewEntry = PayloadAndIndices.Emplace_GetRef();
+				NewEntry.Payload = &VolumePayload;
+				NewEntry.FrameIndex = TNumericLimits<int32>::Lowest();
 			}
 			else
 			{
 				for (const int32 FrameIndex : VolumePayload.PayloadFrameIndices)
 				{
-					PayloadAndIndices.Emplace(&VolumePayload, FrameIndex);
+					FPayloadAndIndex& NewEntry = PayloadAndIndices.Emplace_GetRef();
+					NewEntry.Payload = &VolumePayload;
+					NewEntry.FrameIndex = FrameIndex;
 				}
 			}
 		}
diff --git a/Engine/Plugins/Interchange/Runtime/Source/Pipelines/Private/InterchangeSparseVolumeTexturePipeline.cpp b/Engine/Plugins/Interchange/Runtime/Source/Pipelines/Private/InterchangeSparseVolumeTexturePipeline.cpp
--- a/Engine/Plugins/Interchange/Runtime/Source/Pipelines/Private/InterchangeSparseVolumeTexturePipeline.cpp
+++ b/Engine/Plugins/Interchange/Runtime/Source/Pipelines/Private/InterchangeSparseVolumeTexturePipeline.cpp
@@ -381,7 +381,9 @@ void UInterchangeSparseVolumeTexturePipeline::ExecutePipeline(
 
 			for (const int32 Index : AnimationIndices)
 			{
-				AnimationIDToVolumeNodes.FindOrAdd(AnimationID).Emplace(VolumeNode, Index);
+				FNodeAndAnimationIndex& NewEntry = AnimationIDToVolumeNodes.FindOrAdd(AnimationID).Emplace_GetRef();
+				NewEntry.Node = VolumeNode;
+				NewEntry.Index = Index;
 			}
 		}
 		// Static volume
