Interchange: Fix static analysis warning.

diff --git a/Engine/Plugins/Interchange/Runtime/Source/Import/Private/Volume/InterchangeSparseVolumeTextureFactory.cpp b/Engine/Plugins/Interchange/Runtime/Source/Import/Private/Volume/InterchangeSparseVolumeTextureFactory.cpp
--- a/Engine/Plugins/Interchange/Runtime/Source/Import/Private/Volume/InterchangeSparseVolumeTextureFactory.cpp
+++ b/Engine/Plugins/Interchange/Runtime/Source/Import/Private/Volume/InterchangeSparseVolumeTextureFactory.cpp
@@ -623,7 +623,6 @@ UInterchangeFactoryBase::FImportAssetResult UInterchangeSparseVolumeTextureFacto
 				FName NewName = MakeUniqueObjectName(NewOuter, UserData->GetClass());
 				ensure(UserData->Rename(*NewName.ToString(), NewOuter, REN_DontCreateRedirectors | REN_NonTransactional | REN_DoNotDirty));
 			}
-			OldAssetUserData.Reset();
 		}
 
 		SparseVolumeTexture = NewObject<USparseVolumeTexture>(Arguments.Parent, ClassToUse, *Arguments.AssetName, RF_Public | RF_Standalone);
