USD: Attempt at fixing broken Stage of QAEnterprise Win64.

diff --git a/Engine/Plugins/Interchange/Runtime/Source/Import/Private/Usd/InterchangeUsdTranslator.cpp b/Engine/Plugins/Interchange/Runtime/Source/Import/Private/Usd/InterchangeUsdTranslator.cpp
--- a/Engine/Plugins/Interchange/Runtime/Source/Import/Private/Usd/InterchangeUsdTranslator.cpp
+++ b/Engine/Plugins/Interchange/Runtime/Source/Import/Private/Usd/InterchangeUsdTranslator.cpp
@@ -38,6 +38,7 @@
 #include "InterchangeCameraNode.h"
 #include "InterchangeLightNode.h"
 #include "InterchangeManager.h"
+#include "InterchangeMaterialDefinitions.h"
 #include "InterchangeMaterialInstanceNode.h"
 #include "InterchangeMaterialReferenceNode.h"
 #include "InterchangeMeshNode.h"
