USD Interchange: Added build settings for GeometryCache import (GeometryCache and Sampling settings from Alembic Importer).

#rnx

diff --git a/Engine/Plugins/Importers/USDImporter/Source/USDSchemas/Private/USDGeometryCacheTranslator.cpp b/Engine/Plugins/Importers/USDImporter/Source/USDSchemas/Private/USDGeometryCacheTranslator.cpp
--- a/Engine/Plugins/Importers/USDImporter/Source/USDSchemas/Private/USDGeometryCacheTranslator.cpp
+++ b/Engine/Plugins/Importers/USDImporter/Source/USDSchemas/Private/USDGeometryCacheTranslator.cpp
@@ -70,6 +70,13 @@ static FAutoConsoleVariableRef CVarEnableSubdiv(
 		 "StaticMeshes")
 );
 
+static bool GStoreImportedVertexNumbers = false;
+static FAutoConsoleVariableRef CVarStoreImportedVertexNumbers(
+	TEXT("USD.GeometryCache.Import.StoreImportedVertexNumbers"),
+	GStoreImportedVertexNumbers,
+	TEXT("Whether to store the imported vertex numbers. This lets you know the vertex numbers inside the DCC.")
+);
+
 namespace UsdGeometryCacheTranslatorImpl
 {
 	bool ProcessGeometryCacheMaterials(
@@ -274,7 +281,12 @@ namespace UsdGeometryCacheTranslatorImpl
 
 			FStaticMeshOperations::ComputeTangentsAndNormals(MeshDescription, ComputeNTBsOptions);
 
-			UE::GeometryCache::Utils::GetGeometryCacheMeshDataFromMeshDescription(OutMeshData, MeshDescription, MaterialOffset, Args.FramesPerSecond);
+			UE::GeometryCache::Utils::FMeshDataConversionArguments ConversionArgs;
+			ConversionArgs.MaterialOffset = MaterialOffset;
+			ConversionArgs.FramesPerSecond = Args.FramesPerSecond;
+			ConversionArgs.bUseVelocitiesAsMotionVectors = true;
+			ConversionArgs.bStoreImportedVertexNumbers = GStoreImportedVertexNumbers;
+			UE::GeometryCache::Utils::GetGeometryCacheMeshDataFromMeshDescription(OutMeshData, MeshDescription, ConversionArgs);
 
 			return true;
 		}
diff --git a/Engine/Plugins/Interchange/Runtime/Source/FactoryNodes/Private/InterchangeGeometryCacheFactoryNode.cpp b/Engine/Plugins/Interchange/Runtime/Source/FactoryNodes/Private/InterchangeGeometryCacheFactoryNode.cpp
--- a/Engine/Plugins/Interchange/Runtime/Source/FactoryNodes/Private/InterchangeGeometryCacheFactoryNode.cpp
+++ b/Engine/Plugins/Interchange/Runtime/Source/FactoryNodes/Private/InterchangeGeometryCacheFactoryNode.cpp
@@ -53,3 +53,93 @@ bool UInterchangeGeometryCacheFactoryNode::SetNodeClassFromClassAttribute()
 	}
 	return bIsNodeClassInitialized;
 }
+
+bool UInterchangeGeometryCacheFactoryNode::GetCustomFlattenTracks(bool& AttributeValue) const
+{
+	IMPLEMENT_NODE_ATTRIBUTE_GETTER(FlattenTracks, bool);
+}
+
+bool UInterchangeGeometryCacheFactoryNode::SetCustomFlattenTracks(const bool& AttributeValue)
+{
+	IMPLEMENT_NODE_ATTRIBUTE_SETTER_NODELEGATE(FlattenTracks, bool);
+}
+
+bool UInterchangeGeometryCacheFactoryNode::GetCustomPositionPrecision(float& AttributeValue) const
+{
+	IMPLEMENT_NODE_ATTRIBUTE_GETTER(PositionPrecision, float);
+}
+
+bool UInterchangeGeometryCacheFactoryNode::SetCustomPositionPrecision(const float& AttributeValue)
+{
+	IMPLEMENT_NODE_ATTRIBUTE_SETTER_NODELEGATE(PositionPrecision, float);
+}
+
+bool UInterchangeGeometryCacheFactoryNode::GetCustomNumBitsForUVs(int32& AttributeValue) const
+{
+	IMPLEMENT_NODE_ATTRIBUTE_GETTER(NumBitsForUVs, int32);
+}
+
+bool UInterchangeGeometryCacheFactoryNode::SetCustomNumBitsForUVs(const int32& AttributeValue)
+{
+	IMPLEMENT_NODE_ATTRIBUTE_SETTER_NODELEGATE(NumBitsForUVs, int32);
+}
+
+bool UInterchangeGeometryCacheFactoryNode::GetCustomStartFrame(int32& AttributeValue) const
+{
+	IMPLEMENT_NODE_ATTRIBUTE_GETTER(StartFrame, int32);
+}
+
+bool UInterchangeGeometryCacheFactoryNode::SetCustomStartFrame(const int32& AttributeValue)
+{
+	IMPLEMENT_NODE_ATTRIBUTE_SETTER_NODELEGATE(StartFrame, int32);
+}
+
+bool UInterchangeGeometryCacheFactoryNode::GetCustomEndFrame(int32& AttributeValue) const
+{
+	IMPLEMENT_NODE_ATTRIBUTE_GETTER(EndFrame, int32);
+}
+
+bool UInterchangeGeometryCacheFactoryNode::SetCustomEndFrame(const int32& AttributeValue)
+{
+	IMPLEMENT_NODE_ATTRIBUTE_SETTER_NODELEGATE(EndFrame, int32);
+}
+
+bool UInterchangeGeometryCacheFactoryNode::GetCustomMotionVectorsImport(EInterchangeMotionVectorsHandling& AttributeValue) const
+{
+	IMPLEMENT_NODE_ATTRIBUTE_GETTER(MotionVectors, EInterchangeMotionVectorsHandling);
+}
+
+bool UInterchangeGeometryCacheFactoryNode::SetCustomMotionVectorsImport(EInterchangeMotionVectorsHandling AttributeValue)
+{
+	IMPLEMENT_NODE_ATTRIBUTE_SETTER_NODELEGATE(MotionVectors, EInterchangeMotionVectorsHandling)
+}
+
+bool UInterchangeGeometryCacheFactoryNode::GetCustomApplyConstantTopologyOptimization(bool& AttributeValue) const
+{
+	IMPLEMENT_NODE_ATTRIBUTE_GETTER(ApplyConstantTopologyOptimization, bool);
+}
+
+bool UInterchangeGeometryCacheFactoryNode::SetCustomApplyConstantTopologyOptimization(const bool& AttributeValue)
+{
+	IMPLEMENT_NODE_ATTRIBUTE_SETTER_NODELEGATE(ApplyConstantTopologyOptimization, bool);
+}
+
+bool UInterchangeGeometryCacheFactoryNode::GetCustomStoreImportedVertexNumbers(bool& AttributeValue) const
+{
+	IMPLEMENT_NODE_ATTRIBUTE_GETTER(StoreImportedVertexNumbers, bool);
+}
+
+bool UInterchangeGeometryCacheFactoryNode::SetCustomStoreImportedVertexNumbers(const bool& AttributeValue)
+{
+	IMPLEMENT_NODE_ATTRIBUTE_SETTER_NODELEGATE(StoreImportedVertexNumbers, bool);
+}
+
+bool UInterchangeGeometryCacheFactoryNode::GetCustomOptimizeIndexBuffers(bool& AttributeValue) const
+{
+	IMPLEMENT_NODE_ATTRIBUTE_GETTER(OptimizeIndexBuffers, bool);
+}
+
+bool UInterchangeGeometryCacheFactoryNode::SetCustomOptimizeIndexBuffers(const bool& AttributeValue)
+{
+	IMPLEMENT_NODE_ATTRIBUTE_SETTER_NODELEGATE(OptimizeIndexBuffers, bool);
+}
diff --git a/Engine/Plugins/Interchange/Runtime/Source/FactoryNodes/Public/InterchangeGeometryCacheFactoryNode.h b/Engine/Plugins/Interchange/Runtime/Source/FactoryNodes/Public/InterchangeGeometryCacheFactoryNode.h
--- a/Engine/Plugins/Interchange/Runtime/Source/FactoryNodes/Public/InterchangeGeometryCacheFactoryNode.h
+++ b/Engine/Plugins/Interchange/Runtime/Source/FactoryNodes/Public/InterchangeGeometryCacheFactoryNode.h
@@ -3,6 +3,7 @@
 #pragma once
 
 #include "GeometryCache.h"
+#include "InterchangeMeshDefinitions.h"
 #include "InterchangeMeshFactoryNode.h"
 #include "Nodes/InterchangeFactoryBaseNode.h"
 
@@ -25,6 +26,80 @@ public:
 	UFUNCTION(BlueprintCallable, Category = "Interchange | Node | GeometryCache")
 	void InitializeGeometryCacheNode(const FString& UniqueID, const FString& DisplayLabel, const FString& InAssetClass);
 
+	/** For the following, see GEOMETRY_CACHES_CATEGORY Properties in InterchangeGenericMeshPipeline for more details */
+
+	/** Get whether to merge all geometries into a single mesh */
+	UFUNCTION(BlueprintCallable, Category = "Interchange | Node | GeometryCache")
+	bool GetCustomFlattenTracks(bool& AttributeValue) const;
+
+	/** Set whether to merge all geometries into a single mesh */
+	UFUNCTION(BlueprintCallable, Category = "Interchange | Node | GeometryCache")
+	bool SetCustomFlattenTracks(const bool& AttributeValue);
+
+	/** Get the precision used for compressing vertex positions */
+	UFUNCTION(BlueprintCallable, Category = "Interchange | Node | GeometryCache")
+	bool GetCustomPositionPrecision(float& AttributeValue) const;
+
+	/** Set the precision used for compressing vertex positions */
+	UFUNCTION(BlueprintCallable, Category = "Interchange | Node | GeometryCache")
+	bool SetCustomPositionPrecision(const float& AttributeValue);
+
+	/** Get the number of bits for compressing the UV into */
+	UFUNCTION(BlueprintCallable, Category = "Interchange | Node | GeometryCache")
+	bool GetCustomNumBitsForUVs(int32& AttributeValue) const;
+
+	/** Set the number of bits for compressing the UV into */
+	UFUNCTION(BlueprintCallable, Category = "Interchange | Node | GeometryCache")
+	bool SetCustomNumBitsForUVs(const int32& AttributeValue);
+
+	/** Get the start frame index of the animation */
+	UFUNCTION(BlueprintCallable, Category = "Interchange | Node | GeometryCache")
+	bool GetCustomStartFrame(int32& AttributeValue) const;
+
+	/** Set the start frame index of the animation */
+	UFUNCTION(BlueprintCallable, Category = "Interchange | Node | GeometryCache")
+	bool SetCustomStartFrame(const int32& AttributeValue);
+
+	/** Get the end frame index of the animation */
+	UFUNCTION(BlueprintCallable, Category = "Interchange | Node | GeometryCache")
+	bool GetCustomEndFrame(int32& AttributeValue) const;
+
+	/** Set the end frame index of the animation */
+	UFUNCTION(BlueprintCallable, Category = "Interchange | Node | GeometryCache")
+	bool SetCustomEndFrame(const int32& AttributeValue);
+
+	/** Get how the motion vectors are managed */
+	UFUNCTION(BlueprintCallable, Category = "Interchange | Node | GeometryCache")
+	bool GetCustomMotionVectorsImport(EInterchangeMotionVectorsHandling& AttributeValue) const;
+
+	/** Set how the motion vectors are managed */
+	UFUNCTION(BlueprintCallable, Category = "Interchange | Node | GeometryCache")
+	bool SetCustomMotionVectorsImport(EInterchangeMotionVectorsHandling AttributeValue);
+
+	/** Get whether constant topology optimization is applied */
+	UFUNCTION(BlueprintCallable, Category = "Interchange | Node | GeometryCache")
+	bool GetCustomApplyConstantTopologyOptimization(bool& AttributeValue) const;
+
+	/** Get whether constant topology optimization is applied */
+	UFUNCTION(BlueprintCallable, Category = "Interchange | Node | GeometryCache")
+	bool SetCustomApplyConstantTopologyOptimization(const bool& AttributeValue);
+
+	/** Get whether vertex numbers from DCC are stored in the geometry cache */
+	UFUNCTION(BlueprintCallable, Category = "Interchange | Node | GeometryCache")
+	bool GetCustomStoreImportedVertexNumbers(bool& AttributeValue) const;
+
+	/** Set whether vertex numbers from DCC are stored in the geometry cache */
+	UFUNCTION(BlueprintCallable, Category = "Interchange | Node | GeometryCache")
+	bool SetCustomStoreImportedVertexNumbers(const bool& AttributeValue);
+
+	/** Get whether to optimize the index buffers when building the geometry cache */
+	UFUNCTION(BlueprintCallable, Category = "Interchange | Node | GeometryCache")
+	bool GetCustomOptimizeIndexBuffers(bool& AttributeValue) const;
+
+	/** Set whether to optimize the index buffers when building the geometry cache */
+	UFUNCTION(BlueprintCallable, Category = "Interchange | Node | GeometryCache")
+	bool SetCustomOptimizeIndexBuffers(const bool& AttributeValue);
+
 	/**
 	 * Return the node type name of the class. This is used when reporting errors.
 	 */
@@ -39,4 +114,16 @@ private:
 
 protected:
 	TSubclassOf<UGeometryCache> AssetClass = nullptr;
+
+private:
+	IMPLEMENT_NODE_ATTRIBUTE_KEY(FlattenTracks);
+	IMPLEMENT_NODE_ATTRIBUTE_KEY(PositionPrecision);
+	IMPLEMENT_NODE_ATTRIBUTE_KEY(NumBitsForUVs);
+	IMPLEMENT_NODE_ATTRIBUTE_KEY(OverrideTimeRange);
+	IMPLEMENT_NODE_ATTRIBUTE_KEY(StartFrame);
+	IMPLEMENT_NODE_ATTRIBUTE_KEY(EndFrame);
+	IMPLEMENT_NODE_ATTRIBUTE_KEY(MotionVectors);
+	IMPLEMENT_NODE_ATTRIBUTE_KEY(ApplyConstantTopologyOptimization);
+	IMPLEMENT_NODE_ATTRIBUTE_KEY(StoreImportedVertexNumbers);
+	IMPLEMENT_NODE_ATTRIBUTE_KEY(OptimizeIndexBuffers);
 };
diff --git a/Engine/Plugins/Interchange/Runtime/Source/Import/Private/Mesh/InterchangeGeometryCacheFactory.cpp b/Engine/Plugins/Interchange/Runtime/Source/Import/Private/Mesh/InterchangeGeometryCacheFactory.cpp
--- a/Engine/Plugins/Interchange/Runtime/Source/Import/Private/Mesh/InterchangeGeometryCacheFactory.cpp
+++ b/Engine/Plugins/Interchange/Runtime/Source/Import/Private/Mesh/InterchangeGeometryCacheFactory.cpp
@@ -73,6 +73,151 @@ UInterchangeFactoryBase::FImportAssetResult UInterchangeGeometryCacheFactory::Be
 		}
 	}
 
+	// Query the build settings
+	bool bFlattenTracks = true;
+	GeometryCacheFactoryNode->GetCustomFlattenTracks(bFlattenTracks);
+
+	float PositionPrecision = 0.01f;
+	GeometryCacheFactoryNode->GetCustomPositionPrecision(PositionPrecision);
+
+	int32 NumBitsForUVs = 10;
+	GeometryCacheFactoryNode->GetCustomNumBitsForUVs(NumBitsForUVs);
+
+	TOptional<int32> FrameStartOverride;
+	if (int32 FrameStart = 0; GeometryCacheFactoryNode->GetCustomStartFrame(FrameStart))
+	{
+		FrameStartOverride = FrameStart;
+	}
+
+	TOptional<int32> FrameEndOverride;
+	if (int32 FrameEnd = 0; GeometryCacheFactoryNode->GetCustomEndFrame(FrameEnd))
+	{
+		FrameEndOverride = FrameEnd;
+	}
+
+	EInterchangeMotionVectorsHandling MotionVectors = EInterchangeMotionVectorsHandling::NoMotionVectors;
+	GeometryCacheFactoryNode->GetCustomMotionVectorsImport(MotionVectors);
+
+	bool bApplyConstantTopologyOptimizations = false;
+	GeometryCacheFactoryNode->GetCustomApplyConstantTopologyOptimization(bApplyConstantTopologyOptimizations);
+
+	bool bStoreImportedVertexNumbers = false;
+	GeometryCacheFactoryNode->GetCustomStoreImportedVertexNumbers(bStoreImportedVertexNumbers);
+
+	bool bOptimizeIndexBuffers = false;
+	GeometryCacheFactoryNode->GetCustomOptimizeIndexBuffers(bOptimizeIndexBuffers);
+
+	FTransform GlobalOffsetTransform = FTransform::Identity;
+	if (UInterchangeCommonPipelineDataFactoryNode* CommonPipelineDataFactoryNode = UInterchangeCommonPipelineDataFactoryNode::GetUniqueInstance(Arguments.NodeContainer))
+	{
+		CommonPipelineDataFactoryNode->GetCustomGlobalOffsetTransform(GlobalOffsetTransform);
+	}
+
+	struct FNodeInfo
+	{
+		FString Uid;
+		const UInterchangeGeometryCacheNode* Node = nullptr;
+		FTransform GlobalTransform;
+		double FrameRate = 24.0;
+		bool bConstantTopology = true;
+	};
+
+	TArray<FNodeInfo> ValidatedNodes;
+	int32 FrameStart = MIN_int32;
+	int32 FrameEnd = MAX_int32;
+	bool bGloballyConstantTopology = true;
+
+	TArray<FString> MeshUids;
+	GeometryCacheFactoryNode->GetTargetNodeUids(MeshUids);
+	const FString& AssetName = Arguments.AssetName;
+	for (int32 MeshIndex = 0; MeshIndex < MeshUids.Num(); ++MeshIndex)
+	{
+		const FString& MeshUid = MeshUids[MeshIndex];
+
+		FTransform GlobalMeshTransform;
+		const UInterchangeBaseNode* Node = Arguments.NodeContainer->GetNode(MeshUid);
+		const UInterchangeGeometryCacheNode* GeometryCacheNode = Cast<const UInterchangeGeometryCacheNode>(Node);
+		if (GeometryCacheNode == nullptr)
+		{
+			// MeshUid must refer to a scene node
+			const UInterchangeSceneNode* SceneNode = Cast<const UInterchangeSceneNode>(Node);
+			if (!ensure(SceneNode))
+			{
+				UE_LOG(LogInterchangeImport, Warning, TEXT("Invalid mesh reference when importing GeometryCache asset %s."), *AssetName);
+				continue;
+			}
+
+			// Get the transform from the scene node
+			FTransform SceneNodeGlobalTransform;
+			if (SceneNode->GetCustomGlobalTransform(Arguments.NodeContainer, GlobalOffsetTransform, SceneNodeGlobalTransform))
+			{
+				GlobalMeshTransform = SceneNodeGlobalTransform;
+			}
+
+			constexpr bool bBakeMeshes = true;
+			constexpr bool bBakePivotMeshes = false;
+			UE::Interchange::Private::MeshHelper::AddSceneNodeGeometricAndPivotToGlobalTransform(GlobalMeshTransform, SceneNode, bBakeMeshes, bBakePivotMeshes);
+
+			// And get the mesh node which it references
+			FString MeshDependencyUid;
+			SceneNode->GetCustomAssetInstanceUid(MeshDependencyUid);
+			GeometryCacheNode = Cast<UInterchangeGeometryCacheNode>(Arguments.NodeContainer->GetNode(MeshDependencyUid));
+		}
+		else
+		{
+			// If we have a mesh that is not reference by a scene node, we must apply the global offset.
+			GlobalMeshTransform = GlobalOffsetTransform;
+		}
+
+		if (!ensure(GeometryCacheNode))
+		{
+			UE_LOG(LogInterchangeImport, Warning, TEXT("Invalid mesh reference when importing GeometryCache asset %s."), *AssetName);
+			continue;
+		}
+
+		// Cache the node info for the validated node
+		FNodeInfo& NodeInfo = ValidatedNodes.Emplace_GetRef();
+		NodeInfo.Uid = MeshUid;
+		NodeInfo.Node = GeometryCacheNode;
+		NodeInfo.GlobalTransform = GlobalMeshTransform;
+
+		GeometryCacheNode->GetCustomHasConstantTopology(NodeInfo.bConstantTopology);
+		bGloballyConstantTopology &= NodeInfo.bConstantTopology;
+
+		int32 MeshFrameStart = 0;
+		GeometryCacheNode->GetCustomStartFrame(MeshFrameStart);
+		FrameStart = FMath::Max(FrameStart, MeshFrameStart);
+
+		int32 MeshFrameEnd = 0;
+		GeometryCacheNode->GetCustomEndFrame(MeshFrameEnd);
+		FrameEnd = FMath::Min(FrameEnd, MeshFrameEnd);
+
+		GeometryCacheNode->GetCustomFrameRate(NodeInfo.FrameRate);
+	}
+
+	// Apply the time range overrides if enabled
+	if (FrameStartOverride.IsSet())
+	{
+		FrameStart = FrameStartOverride.GetValue();
+	}
+
+	if (FrameEndOverride.IsSet())
+	{
+		FrameEnd = FrameEndOverride.GetValue();
+	}
+
+	// Make sure there's at least one frame in the animation
+	if (FrameEnd <= FrameStart)
+	{
+		FrameEnd = FrameStart + 1;
+	}
+
+	if (ValidatedNodes.Num() == 0)
+	{
+		UE_LOG(LogInterchangeImport, Error, TEXT("Cannot import geometry cache. No valid mesh to import."));
+		return ImportAssetResult;
+	}
+
 	// Create a new geometry cache or overwrite existing asset, if possible
 	UGeometryCache* GeometryCache = nullptr;
 	if (!ExistingAsset)
@@ -93,7 +238,7 @@ UInterchangeFactoryBase::FImportAssetResult UInterchangeGeometryCacheFactory::Be
 	{
 		if (!Arguments.ReimportObject)
 		{
-			UE_LOG(LogInterchangeImport, Warning, TEXT("Could not create GeometryCache asset %s."), *Arguments.AssetName);
+			UE_LOG(LogInterchangeImport, Warning, TEXT("Could not create GeometryCache asset %s."), *AssetName);
 		}
 		return ImportAssetResult;
 	}
@@ -102,52 +247,48 @@ UInterchangeFactoryBase::FImportAssetResult UInterchangeGeometryCacheFactory::Be
 	GeometryCache->PreEditChange(nullptr);
 #endif // WITH_EDITOR
 
-	TArray<FString> MeshUids;
-	GeometryCacheFactoryNode->GetTargetNodeUids(MeshUids);
-
-	// Create a track for each mesh to be processed and add it to the GeometryCache
-	for (int32 MeshIndex = 0; MeshIndex < MeshUids.Num(); ++MeshIndex)
+	const bool bCalculateMotionVectors = MotionVectors == EInterchangeMotionVectorsHandling::CalculateMotionVectorsDuringImport;
+	auto CreateTrack = [GeometryCache, PositionPrecision, NumBitsForUVs, bApplyConstantTopologyOptimizations, bCalculateMotionVectors, bOptimizeIndexBuffers](const FString& Name, bool bConstantTopology)
 	{
-		const FString& MeshUid = MeshUids[MeshIndex];
-
-		// Extract the mesh name
-		// #ueent_todo: Make this more generic as this is currently based on the USD uids
-		TArray<FString> MeshUidTokens;
-		MeshUid.ParseIntoArray(MeshUidTokens, TEXT("/"));
-		FString MeshName = MeshUidTokens[MeshUidTokens.Num() - 1];
-
-		FName CodecName = MakeUniqueObjectName(GeometryCache, UGeometryCacheCodecV1::StaticClass(), FName(MeshName + FString(TEXT("_Codec"))));
+		FName CodecName = MakeUniqueObjectName(GeometryCache, UGeometryCacheCodecV1::StaticClass(), FName(Name + FString(TEXT("_Codec"))));
 		UGeometryCacheCodecV1* Codec = NewObject<UGeometryCacheCodecV1>(GeometryCache, CodecName, RF_Public);
 
 #if WITH_EDITORONLY_DATA
-		// #ueent_todo: Expose geometry cache settings
-		// Compression settings for good quality
-		const float VertexQuantizationPrecision = 0.0005f;
-		const int32 UVBits = 16;
-		Codec->InitializeEncoder(VertexQuantizationPrecision, UVBits);
+		Codec->InitializeEncoder(PositionPrecision, NumBitsForUVs);
 #endif
 
-		FName TrackName = MakeUniqueObjectName(GeometryCache, UGeometryCacheTrackStreamable::StaticClass(), FName(MeshName));
+		FName TrackName = MakeUniqueObjectName(GeometryCache, UGeometryCacheTrackStreamable::StaticClass(), FName(Name));
 		UGeometryCacheTrackStreamable* Track = NewObject<UGeometryCacheTrackStreamable>(GeometryCache, TrackName, RF_Public);
 
 #if WITH_EDITORONLY_DATA
-		const bool bForceSingleOptimization = false;
-		const bool bCalculateMotionVectors = false;
-		const bool bOptimizeIndexBuffers = false;
-		Track->BeginCoding(Codec, bForceSingleOptimization, bCalculateMotionVectors, bOptimizeIndexBuffers);
+		const bool bCanApplyConstantTopologyOptimizations = bApplyConstantTopologyOptimizations && bConstantTopology;
+		Track->BeginCoding(Codec, bCanApplyConstantTopologyOptimizations, bCalculateMotionVectors, bOptimizeIndexBuffers);
 		// EndCoding has to be called from the main thread once all the frame data have been added to the track
 #endif
+		return Track;
+	};
 
-		GeometryCache->AddTrack(Track);
+	// No need to flatten tracks if there's only one mesh
+	bFlattenTracks &= ValidatedNodes.Num() > 1;
+	if (bFlattenTracks)
+	{
+		Tracks.Add(CreateTrack(AssetName, bGloballyConstantTopology));
+	}
+	else
+	{
+		// Create a track for each mesh to be processed
+		for (const FNodeInfo& NodeInfo : ValidatedNodes)
+		{
+			const FString& MeshUid = NodeInfo.Uid;
 
-		TArray<FMatrix> Mats;
-		Mats.Add(FMatrix::Identity);
-		Mats.Add(FMatrix::Identity);
+			// Extract the mesh name
+			// #ueent_todo: Make this more generic as this is currently based on the USD uids
+			TArray<FString> MeshUidTokens;
+			MeshUid.ParseIntoArray(MeshUidTokens, TEXT("/"));
+			FString MeshName = MeshUidTokens[MeshUidTokens.Num() - 1];
 
-		TArray<float> MatTimes;
-		MatTimes.Add(0.0f);
-		MatTimes.Add(0.0f);
-		Track->SetMatrixSamples(Mats, MatTimes);
+			Tracks.Add(CreateTrack(MeshName, NodeInfo.bConstantTopology));
+		}
 	}
 
 	const bool bReimport = Arguments.ReimportObject && GeometryCache;
@@ -160,186 +301,232 @@ UInterchangeFactoryBase::FImportAssetResult UInterchangeGeometryCacheFactory::Be
 		GeometryCacheFactorySetupAssetMaterialArray(GeometryCache->Materials, GeometryCache->MaterialSlotNames, SlotMaterialDependencies, Arguments.NodeContainer, bReimport);
 	}
 
-	FTransform GlobalOffsetTransform = FTransform::Identity;
-	if (UInterchangeCommonPipelineDataFactoryNode* CommonPipelineDataFactoryNode = UInterchangeCommonPipelineDataFactoryNode::GetUniqueInstance(Arguments.NodeContainer))
-	{
-		CommonPipelineDataFactoryNode->GetCustomGlobalOffsetTransform(GlobalOffsetTransform);
-	}
+	int32 NumFrames = FrameEnd - FrameStart;
+	FString Title = FString::Format(TEXT("Importing frames for {0} ({1}/{2})"), { AssetName, 0,  NumFrames });
+	FScopedSlowTask SlowTask(static_cast<float>(FrameEnd - FrameStart), FText::FromString(Title));
+	SlowTask.MakeDialog(true);
 
 	const int32 MaxWorkerThreads = FTaskGraphInterface::Get().GetNumWorkerThreads();
 	const int32 NumFrameThreads = FMath::Clamp(MaxWorkerThreads, 1, GGeometryCacheParallelFrameReads);
 
-	// For each mesh, retrieve their payloads (frames) and process them into their respective tracks.
 	// The GeometryCache data is processed in the game thread since the payloads are queried in
 	// parallel but need to be processed in order. This allows to minimize memory usage since
 	// only NumFrameThreads amount of data will be in memory at a time.
-
-	int32 StartFrame = MIN_int32;
-	int32 EndFrame = MAX_int32;
-	int32 MaterialOffset = 0;
-	int32 NumMeshes = MeshUids.Num();
-	for (int32 MeshIndex = 0; MeshIndex < NumMeshes; ++MeshIndex)
+	// 
+	// ProcessFrames takes a delegate to handle the conversion to FGeometryCacheMeshData, which
+	// can be merged into a single mesh (flattened) or kept separate (one per track)
+	auto ProcessFrames = [MeshTranslatorPayloadInterface, NumFrameThreads, FrameStart, FrameEnd, &SlowTask, &Tracks = Tracks, &ValidatedNodes, &AssetName, NumFrames](auto ConvertMeshes)
 	{
-		const FString& MeshUid = MeshUids[MeshIndex];
-
-		FTransform GlobalMeshTransform;
-		const UInterchangeBaseNode* Node = Arguments.NodeContainer->GetNode(MeshUid);
-		const UInterchangeGeometryCacheNode* GeometryCacheNode = Cast<const UInterchangeGeometryCacheNode>(Node);
-		if (GeometryCacheNode == nullptr)
+		// Fetch the payloads for all meshes at time FrameIndex and return their processed MeshDescriptions
+		auto ReadFrame = [&ValidatedNodes, &AssetName, MeshTranslatorPayloadInterface](int32 FrameIndex)
 		{
-			// MeshUid must refer to a scene node
-			const UInterchangeSceneNode* SceneNode = Cast<const UInterchangeSceneNode>(Node);
-			if (!ensure(SceneNode))
+			TArray<FMeshDescription> MeshDescriptions;
+			for (const FNodeInfo& NodeInfo : ValidatedNodes)
 			{
-				UE_LOG(LogInterchangeImport, Warning, TEXT("Invalid mesh reference when importing GeometryCache asset %s."), *Arguments.AssetName);
-				continue;
-			}
+				const UInterchangeGeometryCacheNode* GeometryCacheNode = NodeInfo.Node;
 
-			// Get the transform from the scene node
-			FTransform SceneNodeGlobalTransform;
-			if (SceneNode->GetCustomGlobalTransform(Arguments.NodeContainer, GlobalOffsetTransform, SceneNodeGlobalTransform))
-			{
-				GlobalMeshTransform = SceneNodeGlobalTransform;
-			}
-
-			constexpr bool bBakeMeshes = true;
-			constexpr bool bBakePivotMeshes = false;
-			UE::Interchange::Private::MeshHelper::AddSceneNodeGeometricAndPivotToGlobalTransform(GlobalMeshTransform, SceneNode, bBakeMeshes, bBakePivotMeshes);
-
-			// And get the mesh node which it references
-			FString MeshDependencyUid;
-			SceneNode->GetCustomAssetInstanceUid(MeshDependencyUid);
-			GeometryCacheNode = Cast<UInterchangeGeometryCacheNode>(Arguments.NodeContainer->GetNode(MeshDependencyUid));
-		}
-		else
-		{
-			// If we have a mesh that is not reference by a scene node, we must apply the global offset.
-			GlobalMeshTransform = GlobalOffsetTransform;
-		}
-
-		if (!ensure(GeometryCacheNode))
-		{
-			UE_LOG(LogInterchangeImport, Warning, TEXT("Invalid mesh reference when importing GeometryCache asset %s."), *Arguments.AssetName);
-			continue;
-		}
+				TOptional<FInterchangeMeshPayLoadKey> OptionalPayLoadKey = GeometryCacheNode->GetPayLoadKey();
+				if (!ensure(OptionalPayLoadKey.IsSet()))
+				{
+					MeshDescriptions.Add({});
+					UE_LOG(LogInterchangeImport, Warning, TEXT("Empty mesh reference payload when importing GeometryCache asset %s."), *AssetName);
+					continue;
+				}
 
-		TOptional<FInterchangeMeshPayLoadKey> OptionalPayLoadKey = GeometryCacheNode->GetPayLoadKey();
-		if (!ensure(OptionalPayLoadKey.IsSet()))
-		{
-			UE_LOG(LogInterchangeImport, Warning, TEXT("Empty mesh reference payload when importing GeometryCache asset %s."), *Arguments.AssetName);
-			continue;
-		}
+				FInterchangeMeshPayLoadKey PayLoadKey(OptionalPayLoadKey->UniqueId, FrameIndex);
+				TOptional<UE::Interchange::FMeshPayloadData> PayloadData = MeshTranslatorPayloadInterface->GetMeshPayloadData(PayLoadKey, NodeInfo.GlobalTransform);
 
-		int32 MeshStartFrame = 0;
-		GeometryCacheNode->GetCustomStartFrame(MeshStartFrame);
-		StartFrame = FMath::Max(StartFrame, MeshStartFrame);
+				if (PayloadData.IsSet() && !PayloadData->MeshDescription.IsEmpty())
+				{
+					// Compute the normals and tangents for the mesh
+					const float ComparisonThreshold = THRESH_POINTS_ARE_SAME;
 
-		int32 MeshEndFrame = 0;
-		GeometryCacheNode->GetCustomEndFrame(MeshEndFrame);
-		EndFrame = FMath::Min(EndFrame, MeshEndFrame);
+					// This function make sure the Polygon Normals Tangents Binormals are computed and also remove degenerated triangle from the render mesh
+					// description.
+					FStaticMeshOperations::ComputeTriangleTangentsAndNormals(PayloadData->MeshDescription, ComparisonThreshold);
 
-		double FrameRate = 24.0;
-		GeometryCacheNode->GetCustomFrameRate(FrameRate);
+					// Compute any missing normals or tangents.
+					// Static meshes always blend normals of overlapping corners.
+					EComputeNTBsFlags ComputeNTBsOptions = EComputeNTBsFlags::BlendOverlappingNormals;
+					ComputeNTBsOptions |= EComputeNTBsFlags::IgnoreDegenerateTriangles;
+					ComputeNTBsOptions |= EComputeNTBsFlags::UseMikkTSpace;
 
-		bool bConstantTopology = true;
-		GeometryCacheNode->GetCustomHasConstantTopology(bConstantTopology);
+					FStaticMeshOperations::ComputeTangentsAndNormals(PayloadData->MeshDescription, ComputeNTBsOptions);
 
-		FString MeshName = GeometryCacheNode->GetDisplayLabel();
-		FString Title = FString::Format(TEXT("Importing frames for {0} (Mesh {1}/{2})"), { MeshName, MeshIndex + 1,  NumMeshes });
-		FScopedSlowTask SlowTask(static_cast<float>((MeshEndFrame + 1) - MeshStartFrame), FText::FromString(Title));
-		SlowTask.MakeDialog(true);
+					MeshDescriptions.Emplace(PayloadData->MeshDescription);
+				}
+				else
+				{
+					MeshDescriptions.Add({});
+				}
+			}
+			return MeshDescriptions;
+		};
 
-		FEvent* FrameWrittenEvent = FPlatformProcess::GetSynchEventFromPool();
-		std::atomic<int32> WriteFrameIndex = MeshStartFrame;
+		// Frame data can be read concurrently but will be processed sequentially.
+		std::atomic<int32> WriteFrameIndex = FrameStart;
 		FCriticalSection Mutex;
+		FEvent* FrameWrittenEvent = FPlatformProcess::GetSynchEventFromPool();
+		int32 TotalFrames = 0;
 		float CompletedFrames = 0.0f;
+		bool bIsCancelled = false;
 
-		// Parallel frame read: frame data can be read concurrently but have to be processed in order for AddMeshSample
-		ParallelFor(
-			NumFrameThreads,
-			[MeshStartFrame,
-			MeshEndFrame,
-			FrameRate,
-			&OptionalPayLoadKey,
-			&MeshTranslatorPayloadInterface,
-			&GlobalMeshTransform,
-			MaterialOffset,
+		ParallelFor(NumFrameThreads, 
+			[FrameStart, 
+			FrameEnd,
+			&bIsCancelled,
+			&SlowTask,
+			ReadFrame,
+			ConvertMeshes,
 			&WriteFrameIndex,
-			FrameWrittenEvent,
 			&Mutex,
-			StreamableTrack = Cast<UGeometryCacheTrackStreamable>(GeometryCache->Tracks[MeshIndex]),
-			bConstantTopology,
-			&SlowTask,
-			&CompletedFrames,
-			NumFrameThreads](int32 FrameThreadIndex)
-			{
-				int32 FrameIndex = MeshStartFrame + FrameThreadIndex;
+			&FrameWrittenEvent,
+			&Tracks, 
+			&ValidatedNodes, 
+			&AssetName,
+			&TotalFrames, 
+			&CompletedFrames, 
+			NumFrames,
+			NumFrameThreads](int32 ThreadIndex)
+		{
+			int32 FrameIndex = FrameStart + ThreadIndex;
 
-				while (FrameIndex < MeshEndFrame)
+			while (FrameIndex < FrameEnd && !bIsCancelled)
+			{
+				if (IsInGameThread() && SlowTask.ShouldCancel())
 				{
-					FInterchangeMeshPayLoadKey PayLoadKey(OptionalPayLoadKey->UniqueId, FrameIndex);
-					TOptional<UE::Interchange::FMeshPayloadData> PayloadData = MeshTranslatorPayloadInterface->GetMeshPayloadData(PayLoadKey, GlobalMeshTransform);
-
-					if (PayloadData.IsSet() && !PayloadData->MeshDescription.IsEmpty())
-					{
-						// Compute the normals and tangents for the mesh
-						const float ComparisonThreshold = THRESH_POINTS_ARE_SAME;
+					bIsCancelled = true;
+					break;
+				}
 
-						// This function make sure the Polygon Normals Tangents Binormals are computed and also remove degenerated triangle from the render mesh
-						// description.
-						FStaticMeshOperations::ComputeTriangleTangentsAndNormals(PayloadData->MeshDescription, ComparisonThreshold);
+				// Read frame data into memory
+				TArray<FMeshDescription> MeshDescriptions = ReadFrame(FrameIndex);
 
-						// Compute any missing normals or tangents.
-						// Static meshes always blend normals of overlapping corners.
-						EComputeNTBsFlags ComputeNTBsOptions = EComputeNTBsFlags::BlendOverlappingNormals;
-						ComputeNTBsOptions |= EComputeNTBsFlags::IgnoreDegenerateTriangles;
-						ComputeNTBsOptions |= EComputeNTBsFlags::UseMikkTSpace;
+				// And convert it to flattened mesh or separate meshes
+				TArray<FGeometryCacheMeshData> MeshData = ConvertMeshes(FrameIndex, MeshDescriptions);
 
-						FStaticMeshOperations::ComputeTangentsAndNormals(PayloadData->MeshDescription, ComputeNTBsOptions);
+				// Wait until it's our turn to process this frame.
+				while (WriteFrameIndex < FrameIndex)
+				{
+					const uint32 WaitTimeInMs = 10;
+					FrameWrittenEvent->Wait(WaitTimeInMs);
+				}
 
-						// Convert the MeshDescription to MeshData
-						FGeometryCacheMeshData MeshData;
-						UE::GeometryCache::Utils::GetGeometryCacheMeshDataFromMeshDescription(MeshData, PayloadData->MeshDescription, MaterialOffset, FrameRate);
+				{
+					FScopeLock WriteLock(&Mutex);
 
-						// Wait until it's our turn to process this frame.
-						while (WriteFrameIndex < FrameIndex)
+#if WITH_EDITORONLY_DATA
+					// Add the mesh data to the tracks
+					for (int32 Index = 0; Index < MeshData.Num(); ++Index)
+					{
+						if (UGeometryCacheTrackStreamable* Track = Cast<UGeometryCacheTrackStreamable>(Tracks[Index].Get()))
 						{
-							const uint32 WaitTimeInMs = 10;
-							FrameWrittenEvent->Wait(WaitTimeInMs);
+							const FNodeInfo& NodeInfo = ValidatedNodes[Index];
+							Track->AddMeshSample(MeshData[Index], (FrameIndex - FrameStart) / NodeInfo.FrameRate, NodeInfo.bConstantTopology);
 						}
-
-						{
-							FScopeLock WriteLock(&Mutex);
-
-#if WITH_EDITORONLY_DATA
-							// Add it to the track
-							StreamableTrack->AddMeshSample(MeshData, (FrameIndex - MeshStartFrame) / FrameRate, bConstantTopology);
+					}
 #endif
 
-							// Mark the next frame index as ready for processing.
-							++WriteFrameIndex;
+					// Mark the next frame index as ready for processing.
+					++WriteFrameIndex;
 
-							++CompletedFrames;
-							if (IsInGameThread())
-							{
-								SlowTask.EnterProgressFrame(CompletedFrames);
-								CompletedFrames = 0.0f;
-							}
+					++CompletedFrames;
+					if (IsInGameThread())
+					{
+						TotalFrames += int32(CompletedFrames);
+						FString Title = FString::Format(TEXT("Importing frames for {0} ({1}/{2})"), { AssetName, TotalFrames,  NumFrames });
 
-							FrameWrittenEvent->Trigger();
-						}
+						SlowTask.EnterProgressFrame(CompletedFrames, FText::FromString(Title));
+						CompletedFrames = 0.0f;
 					}
-					// Get new frame index to read for next run cycle
-					FrameIndex += NumFrameThreads;
+
+					FrameWrittenEvent->Trigger();
 				}
-			});
+
+				// Get new frame index to read for next run cycle
+				FrameIndex += NumFrameThreads;
+			}
+		});
 
 		FPlatformProcess::ReturnSynchEventToPool(FrameWrittenEvent);
 
-		++MaterialOffset;
+		return !bIsCancelled;
+	};
+
+	bool bSuccess = false;
+	const bool bUseVelocitiesAsMotionVectors = MotionVectors == EInterchangeMotionVectorsHandling::ImportVelocitiesAsMotionVectors;
+	if (bFlattenTracks)
+	{
+		auto MergeMeshes = [FrameRate = ValidatedNodes[0].FrameRate, bUseVelocitiesAsMotionVectors, bStoreImportedVertexNumbers](int32 FrameIndex, TArray<FMeshDescription>& MeshDescriptions)
+		{
+			// Take the first MeshDescription as the base on which to append the others
+			FMeshDescription MergedMesh = MoveTemp(MeshDescriptions[0]);
+
+			TArray<const FMeshDescription*> MeshDescriptionPtrs;
+			MeshDescriptionPtrs.Reserve(MeshDescriptions.Num());
+			for (int32 Index = 1; Index < MeshDescriptions.Num(); ++Index)
+			{
+				MeshDescriptionPtrs.Add(&MeshDescriptions[Index]);
+			}
+
+			FStaticMeshOperations::FAppendSettings AppendSettings;
+			FStaticMeshOperations::AppendMeshDescriptions(MeshDescriptionPtrs, MergedMesh, AppendSettings);
+
+			UE::GeometryCache::Utils::FMeshDataConversionArguments ConversionArgs;
+			ConversionArgs.MaterialOffset = 0;
+			ConversionArgs.FramesPerSecond = FrameRate;
+			ConversionArgs.bUseVelocitiesAsMotionVectors = bUseVelocitiesAsMotionVectors;
+			ConversionArgs.bStoreImportedVertexNumbers = bStoreImportedVertexNumbers;
+
+			FGeometryCacheMeshData MeshData;
+			UE::GeometryCache::Utils::GetGeometryCacheMeshDataFromMeshDescription(MeshData, MergedMesh, ConversionArgs);
+
+			return TArray<FGeometryCacheMeshData>{MeshData};
+		};
+
+		bSuccess = ProcessFrames(MergeMeshes);
+	}
+	else
+	{
+		auto OneMeshPerTrack = [&ValidatedNodes, bUseVelocitiesAsMotionVectors, bStoreImportedVertexNumbers](int32 FrameIndex, TArray<FMeshDescription>& MeshDescriptions)
+		{
+			TArray<FGeometryCacheMeshData> AllMeshData;
+			AllMeshData.Reserve(ValidatedNodes.Num());
+			int32 MaterialOffset = 0;
+			for (int32 Index = 0; Index < ValidatedNodes.Num(); ++Index)
+			{
+				if (MeshDescriptions[Index].IsEmpty())
+				{
+					AllMeshData.Add({});
+					continue;
+				}
+
+				// Convert the MeshDescription to MeshData
+				UE::GeometryCache::Utils::FMeshDataConversionArguments ConversionArgs;
+				ConversionArgs.MaterialOffset = MaterialOffset;
+				ConversionArgs.FramesPerSecond = ValidatedNodes[Index].FrameRate;
+				ConversionArgs.bUseVelocitiesAsMotionVectors = bUseVelocitiesAsMotionVectors;
+				ConversionArgs.bStoreImportedVertexNumbers = bStoreImportedVertexNumbers;
+
+				FGeometryCacheMeshData MeshData;
+				UE::GeometryCache::Utils::GetGeometryCacheMeshDataFromMeshDescription(MeshData, MeshDescriptions[Index], ConversionArgs);
+				AllMeshData.Emplace(MeshData);
+
+				++MaterialOffset;
+			}
+			return AllMeshData;
+		};
+
+		bSuccess = ProcessFrames(OneMeshPerTrack);
 	}
 
-	GeometryCache->SetFrameStartEnd(StartFrame, EndFrame);
+	if (!bSuccess)
+	{
+		GeometryCache->MarkAsGarbage();
+		return ImportAssetResult;
+	}
+
+	GeometryCache->SetFrameStartEnd(FrameStart, FrameEnd);
 
 	ImportAssetResult.ImportedObject = GeometryCache;
 	return ImportAssetResult;
@@ -412,11 +599,23 @@ UInterchangeFactoryBase::FImportAssetResult UInterchangeGeometryCacheFactory::En
 
 #if WITH_EDITORONLY_DATA
 	// Finalize the coding for all tracks
-	for (int32 Index = 0; Index < GeometryCache->Tracks.Num(); ++Index)
+	for (UGeometryCacheTrack* Track : Tracks)
 	{
-		if (UGeometryCacheTrackStreamable* StreamableTrack = Cast<UGeometryCacheTrackStreamable>(GeometryCache->Tracks[Index]))
+		if (UGeometryCacheTrackStreamable* StreamableTrack = Cast<UGeometryCacheTrackStreamable>(Track))
 		{
-			StreamableTrack->EndCoding();
+			if (StreamableTrack->EndCoding())
+			{
+				TArray<FMatrix> Mats;
+				Mats.Add(FMatrix::Identity);
+				Mats.Add(FMatrix::Identity);
+
+				TArray<float> MatTimes;
+				MatTimes.Add(0.0f);
+				MatTimes.Add(0.0f);
+				Track->SetMatrixSamples(Mats, MatTimes);
+
+				GeometryCache->AddTrack(StreamableTrack);
+			}
 		}
 	}
 #endif
diff --git a/Engine/Plugins/Interchange/Runtime/Source/Import/Public/Mesh/InterchangeGeometryCacheFactory.h b/Engine/Plugins/Interchange/Runtime/Source/Import/Public/Mesh/InterchangeGeometryCacheFactory.h
--- a/Engine/Plugins/Interchange/Runtime/Source/Import/Public/Mesh/InterchangeGeometryCacheFactory.h
+++ b/Engine/Plugins/Interchange/Runtime/Source/Import/Public/Mesh/InterchangeGeometryCacheFactory.h
@@ -6,6 +6,8 @@
 
 #include "InterchangeGeometryCacheFactory.generated.h"
 
+class UGeometryCacheTrack;
+
 UCLASS(BlueprintType)
 class INTERCHANGEIMPORT_API UInterchangeGeometryCacheFactory : public UInterchangeFactoryBase
 {
@@ -29,4 +31,7 @@ public:
 
 	// Interchange factory base interface end
 	//////////////////////////////////////////////////////////////////////////
+
+private:
+	TArray<TObjectPtr<UGeometryCacheTrack>> Tracks;
 };
diff --git a/Engine/Plugins/Interchange/Runtime/Source/Nodes/Public/InterchangeMeshDefinitions.h b/Engine/Plugins/Interchange/Runtime/Source/Nodes/Public/InterchangeMeshDefinitions.h
--- a/Engine/Plugins/Interchange/Runtime/Source/Nodes/Public/InterchangeMeshDefinitions.h
+++ b/Engine/Plugins/Interchange/Runtime/Source/Nodes/Public/InterchangeMeshDefinitions.h
@@ -38,6 +38,17 @@ enum class EInterchangeMeshCollision : uint8
 	None = 0xFF
 };
 
+UENUM(BlueprintType)
+enum class EInterchangeMotionVectorsHandling : uint8
+{
+	/** No motion vectors will be present in the geometry cache. */
+	NoMotionVectors,
+	/** Imports the velocities and converts them to motion vectors. This will increase file size as the motion vectors will be stored on disc. */
+	ImportVelocitiesAsMotionVectors,
+	/** Force calculation of motion vectors during import. This will increase file size as the motion vectors will be stored on disc. */
+	CalculateMotionVectorsDuringImport
+};
+
 namespace UE
 {
 	namespace Interchange
diff --git a/Engine/Plugins/Interchange/Runtime/Source/Pipelines/Private/InterchangeGenericAssetsPipeline.cpp b/Engine/Plugins/Interchange/Runtime/Source/Pipelines/Private/InterchangeGenericAssetsPipeline.cpp
--- a/Engine/Plugins/Interchange/Runtime/Source/Pipelines/Private/InterchangeGenericAssetsPipeline.cpp
+++ b/Engine/Plugins/Interchange/Runtime/Source/Pipelines/Private/InterchangeGenericAssetsPipeline.cpp
@@ -9,6 +9,7 @@
 #include "CoreMinimal.h"
 #include "Engine/SkinnedAssetCommon.h"
 #include "Framework/Docking/TabManager.h"
+#include "GeometryCache.h"
 #include "HAL/PlatformApplicationMisc.h"
 #include "InterchangeAnimationTrackSetNode.h"
 #include "InterchangeAnimSequenceFactoryNode.h"
@@ -209,6 +210,7 @@ void UInterchangeGenericAssetsPipeline::FilterPropertiesFromTranslatedData(UInte
 
 		TArray<FString> SkeletalMeshes;
 		TArray<FString> StaticMeshes;
+		TArray<FString> GeometryCaches;
 		int32 RawStaticMesh = 0;
 		int32 RawSkeletalMesh = 0;
 		int32 RawMorphTargetShape = 0;
@@ -226,6 +228,12 @@ void UInterchangeGenericAssetsPipeline::FilterPropertiesFromTranslatedData(UInte
 				PipelineMeshesUtilities->GetAllStaticMeshGeometry(StaticMeshes);
 			}
 
+			PipelineMeshesUtilities->GetAllGeometryCacheInstance(GeometryCaches);
+			if (GeometryCaches.Num() == 0)
+			{
+				PipelineMeshesUtilities->GetAllGeometryCacheGeometry(GeometryCaches);
+			}
+
 			InBaseNodeContainer->IterateNodesOfType<UInterchangeMeshNode>([&RawStaticMesh, &RawSkeletalMesh, &RawMorphTargetShape](const FString& NodeUid, UInterchangeMeshNode* MeshNode)
 				{
 					if (MeshNode->IsMorphTarget())
@@ -275,6 +283,7 @@ void UInterchangeGenericAssetsPipeline::FilterPropertiesFromTranslatedData(UInte
 		//Found which categories to hide
 		bool bHideStaticMeshes = false;
 		bool bHideSkeletalMeshes = false;
+		bool bHideGeometryCaches = GeometryCaches.Num() == 0;
 		bool bHideCommonMeshes = false;
 		bool bHideCommonSkeletalMeshesAndAnimations = false;
 		bool bHideCommonSkeletalMeshesAndAnimations_StaticMesh = false;
@@ -352,6 +361,10 @@ void UInterchangeGenericAssetsPipeline::FilterPropertiesFromTranslatedData(UInte
 		{
 			HideFullCategory(UInterchangeGenericMeshPipeline::GetPipelineCategory(USkeletalMesh::StaticClass()), MeshPipeline);
 		}
+		if (bHideGeometryCaches)
+		{
+			HideFullCategory(UInterchangeGenericMeshPipeline::GetPipelineCategory(UGeometryCache::StaticClass()), MeshPipeline);
+		}
 		if(bHideCommonMeshes)
 		{
 			HideFullCategory(UInterchangeGenericCommonMeshesProperties::GetPipelineCategory(nullptr), CommonMeshesProperties);
diff --git a/Engine/Plugins/Interchange/Runtime/Source/Pipelines/Private/InterchangeGenericGeometryCachePipeline.cpp b/Engine/Plugins/Interchange/Runtime/Source/Pipelines/Private/InterchangeGenericGeometryCachePipeline.cpp
--- a/Engine/Plugins/Interchange/Runtime/Source/Pipelines/Private/InterchangeGenericGeometryCachePipeline.cpp
+++ b/Engine/Plugins/Interchange/Runtime/Source/Pipelines/Private/InterchangeGenericGeometryCachePipeline.cpp
@@ -111,6 +111,21 @@ UInterchangeGeometryCacheFactoryNode* UInterchangeGenericMeshPipeline::CreateGeo
 	GeometryCacheFactoryNode->SetCustomUseBackwardsCompatibleF16TruncUVs(CommonMeshesProperties->bUseBackwardsCompatibleF16TruncUVs);
 	GeometryCacheFactoryNode->SetCustomRemoveDegenerates(CommonMeshesProperties->bRemoveDegenerates);
 
+	GeometryCacheFactoryNode->SetCustomFlattenTracks(bFlattenTracks);
+	GeometryCacheFactoryNode->SetCustomPositionPrecision(CompressedPositionPrecision);
+	GeometryCacheFactoryNode->SetCustomNumBitsForUVs(CompressedTextureCoordinatesNumberOfBits);
+
+	if (bOverrideTimeRange)
+	{
+		GeometryCacheFactoryNode->SetCustomStartFrame(FrameStart);
+		GeometryCacheFactoryNode->SetCustomEndFrame(FrameEnd);
+	}
+
+	GeometryCacheFactoryNode->SetCustomMotionVectorsImport(MotionVectors);
+	GeometryCacheFactoryNode->SetCustomApplyConstantTopologyOptimization(bApplyConstantTopologyOptimizations);
+	GeometryCacheFactoryNode->SetCustomStoreImportedVertexNumbers(bStoreImportedVertexNumbers);
+	GeometryCacheFactoryNode->SetCustomOptimizeIndexBuffers(bOptimizeIndexBuffers);
+
 	return GeometryCacheFactoryNode;
 }
 
diff --git a/Engine/Plugins/Interchange/Runtime/Source/Pipelines/Private/InterchangeGenericMeshPipeline.cpp b/Engine/Plugins/Interchange/Runtime/Source/Pipelines/Private/InterchangeGenericMeshPipeline.cpp
--- a/Engine/Plugins/Interchange/Runtime/Source/Pipelines/Private/InterchangeGenericMeshPipeline.cpp
+++ b/Engine/Plugins/Interchange/Runtime/Source/Pipelines/Private/InterchangeGenericMeshPipeline.cpp
@@ -5,6 +5,7 @@
 #include "Animation/Skeleton.h"
 #include "Engine/SkeletalMesh.h"
 #include "Engine/StaticMesh.h"
+#include "GeometryCache.h"
 #include "InterchangeAnimationTrackSetNode.h"
 #include "InterchangeGenericAssetsPipeline.h"
 #include "InterchangeMeshNode.h"
@@ -34,6 +35,10 @@ FString UInterchangeGenericMeshPipeline::GetPipelineCategory(UClass* AssetClass)
 		{
 			return TEXT("Skeletal Meshes");
 		}
+		else if (AssetClass->IsChildOf(UGeometryCache::StaticClass()))
+		{
+			return TEXT("Geometry Caches");
+		}
 	}
 	return TEXT("Static Meshes");
 }
@@ -54,9 +59,10 @@ void UInterchangeGenericMeshPipeline::AdjustSettingsForContext(const FInterchang
 	bool bAutoDetectConvertStaticMeshToSkeletalMesh = false;
 	bool bContainStaticMesh = false;
 	bool bContainSkeletalMesh = false;
+	bool bContainGeometryCache = false;
 	bool bContainStaticMeshAnimationNode = false;
 	bool bIgnoreStaticMesh = false;
-	GetMeshesInformationFromTranslatedData(ContextParams.BaseNodeContainer, bAutoDetectConvertStaticMeshToSkeletalMesh, bContainStaticMesh, bContainSkeletalMesh, bContainStaticMeshAnimationNode, bIgnoreStaticMesh);
+	GetMeshesInformationFromTranslatedData(ContextParams.BaseNodeContainer, bAutoDetectConvertStaticMeshToSkeletalMesh, bContainStaticMesh, bContainSkeletalMesh, bContainGeometryCache, bContainStaticMeshAnimationNode, bIgnoreStaticMesh);
 
 	//Avoid creating physics asset when importing a LOD or the alternate skinning
 	if (ContextParams.ContextType == EInterchangePipelineContext::AssetCustomLODImport
@@ -153,6 +159,7 @@ void UInterchangeGenericMeshPipeline::AdjustSettingsForContext(const FInterchang
 	const FString CommonMeshesCategory = UInterchangeGenericCommonMeshesProperties::GetPipelineCategory(nullptr);
 	const FString StaticMeshesCategory = UInterchangeGenericMeshPipeline::GetPipelineCategory(UStaticMesh::StaticClass());
 	const FString SkeletalMeshesCategory = UInterchangeGenericMeshPipeline::GetPipelineCategory(USkeletalMesh::StaticClass());
+	const FString GeometryCachesCategory = UInterchangeGenericMeshPipeline::GetPipelineCategory(UGeometryCache::StaticClass());
 	const FString CommonSkeletalMeshesAndAnimationCategory = UInterchangeGenericCommonSkeletalMeshesAndAnimationsProperties::GetPipelineCategory(nullptr);
 
 	TArray<FString> HideCategories;
@@ -173,6 +180,7 @@ void UInterchangeGenericMeshPipeline::AdjustSettingsForContext(const FInterchang
 			}
 			bImportStaticMeshes = false;
 			HideCategories.Add(StaticMeshesCategory);
+			HideCategories.Add(GeometryCachesCategory);
 			if(!bContainSkeletalMesh
 				|| SkeletalMeshImportContentType == EInterchangeSkeletalMeshContentType::Geometry
 				|| CommonMeshesProperties->ForceAllMeshAsType == EInterchangeForceMeshType::IFMT_StaticMesh)
@@ -183,6 +191,7 @@ void UInterchangeGenericMeshPipeline::AdjustSettingsForContext(const FInterchang
 		else if (UStaticMesh* StaticMesh = Cast<UStaticMesh>(ContextParams.ReimportAsset))
 		{
 			HideCategories.Add(SkeletalMeshesCategory);
+			HideCategories.Add(GeometryCachesCategory);
 			HideCategories.Add(CommonSkeletalMeshesAndAnimationCategory);
 			bImportSkeletalMeshes = false;
 			if (!bContainStaticMesh
@@ -192,15 +201,24 @@ void UInterchangeGenericMeshPipeline::AdjustSettingsForContext(const FInterchang
 			}
 		}
 		else if (UAnimSequence* AnimSequence = Cast<UAnimSequence>(ContextParams.ReimportAsset))
+		{
+			HideCategories.Add(StaticMeshesCategory);
+			HideCategories.Add(SkeletalMeshesCategory);
+			HideCategories.Add(GeometryCachesCategory);
+			HideCategories.Add(CommonMeshesCategory);
+		}
+		else if (UGeometryCache* GeometryCache = Cast<UGeometryCache>(ContextParams.ReimportAsset))
 		{
 			HideCategories.Add(StaticMeshesCategory);
 			HideCategories.Add(SkeletalMeshesCategory);
 			HideCategories.Add(CommonMeshesCategory);
+			HideCategories.Add(CommonSkeletalMeshesAndAnimationCategory);
 		}
 		else if (ContextParams.ReimportAsset)
 		{
 			HideCategories.Add(StaticMeshesCategory);
 			HideCategories.Add(SkeletalMeshesCategory);
+			HideCategories.Add(GeometryCachesCategory);
 			HideCategories.Add(CommonMeshesCategory);
 			HideCategories.Add(CommonSkeletalMeshesAndAnimationCategory);
 		}
@@ -208,6 +226,11 @@ void UInterchangeGenericMeshPipeline::AdjustSettingsForContext(const FInterchang
 
 	if (UInterchangePipelineBase* OuterMostPipeline = GetMostPipelineOuter())
 	{
+		if (bContainGeometryCache)
+		{
+			HideProperty(OuterMostPipeline, CommonMeshesProperties.Get(), GET_MEMBER_NAME_CHECKED(UInterchangeGenericCommonMeshesProperties, ForceAllMeshAsType));
+		}
+
 		constexpr bool bDoTransientSubPipeline = true;
 		if (UInterchangeGenericAssetsPipeline* ParentPipeline = Cast<UInterchangeGenericAssetsPipeline>(OuterMostPipeline))
 		{
@@ -299,6 +322,7 @@ void UInterchangeGenericMeshPipeline::GetMeshesInformationFromTranslatedData(con
 	, bool& bAutoDetectConvertStaticMeshToSkeletalMesh
 	, bool& bContainStaticMesh
 	, bool& bContainSkeletalMesh
+	, bool& bContainGeometryCache
 	, bool& bContainStaticMeshAnimationNode
 	, bool& bIgnoreStaticMeshes) const
 {
@@ -310,15 +334,23 @@ void UInterchangeGenericMeshPipeline::GetMeshesInformationFromTranslatedData(con
 	bAutoDetectConvertStaticMeshToSkeletalMesh = false;
 	bContainStaticMesh = false;
 	bContainSkeletalMesh = false;
+	bContainGeometryCache = false;
 	bContainStaticMeshAnimationNode = false;
 	bIgnoreStaticMeshes = false;
 	{
 		TArray<FString> StaticMeshNodeUids;
-		InBaseNodeContainer->IterateNodesOfType<UInterchangeMeshNode>([&bContainSkeletalMesh, &StaticMeshNodeUids](const FString& NodeUid, UInterchangeMeshNode* MeshNode)
+		InBaseNodeContainer->IterateNodesOfType<UInterchangeMeshNode>([&bContainSkeletalMesh, &bContainGeometryCache, &StaticMeshNodeUids](const FString& NodeUid, UInterchangeMeshNode* MeshNode)
 			{
 				if (!MeshNode->IsMorphTarget())
 				{
-					MeshNode->IsSkinnedMesh() ? bContainSkeletalMesh = true : StaticMeshNodeUids.Add(NodeUid);
+					if (Cast<UInterchangeGeometryCacheNode>(MeshNode))
+					{
+						bContainGeometryCache = true;
+					}
+					else
+					{
+						MeshNode->IsSkinnedMesh() ? bContainSkeletalMesh = true : StaticMeshNodeUids.Add(NodeUid);
+					}
 				}
 			});
 		bContainStaticMesh = !StaticMeshNodeUids.IsEmpty();
@@ -372,9 +404,10 @@ UInterchangePipelineMeshesUtilities* UInterchangeGenericMeshPipeline::CreateMesh
 	bool bAutoDetectConvertStaticMeshToSkeletalMesh = false;
 	bool bContainStaticMesh = false;
 	bool bContainSkeletalMesh = false;
+	bool bContainGeometryCache = false;
 	bool bContainStaticMeshAnimationNode = false;
 	bool bIgnoreStaticMeshes = false;
-	Pipeline->GetMeshesInformationFromTranslatedData(InBaseNodeContainer, bAutoDetectConvertStaticMeshToSkeletalMesh, bContainStaticMesh, bContainSkeletalMesh, bContainStaticMeshAnimationNode, bIgnoreStaticMeshes);
+	Pipeline->GetMeshesInformationFromTranslatedData(InBaseNodeContainer, bAutoDetectConvertStaticMeshToSkeletalMesh, bContainStaticMesh, bContainSkeletalMesh, bContainGeometryCache, bContainStaticMeshAnimationNode, bIgnoreStaticMeshes);
 
 	//Set the context option to use when querying the pipeline mesh utilities
 	FInterchangePipelineMeshesUtilitiesContext DataContext;
diff --git a/Engine/Plugins/Interchange/Runtime/Source/Pipelines/Public/InterchangeGenericMeshPipeline.h b/Engine/Plugins/Interchange/Runtime/Source/Pipelines/Public/InterchangeGenericMeshPipeline.h
--- a/Engine/Plugins/Interchange/Runtime/Source/Pipelines/Public/InterchangeGenericMeshPipeline.h
+++ b/Engine/Plugins/Interchange/Runtime/Source/Pipelines/Public/InterchangeGenericMeshPipeline.h
@@ -248,6 +248,49 @@ public:
 	UPROPERTY(EditAnywhere, BlueprintReadWrite, Category = "Geometry Caches")
 	bool bImportGeometryCaches = true;
 
+	/** Whether or not to merge all vertex animation into one track */
+	UPROPERTY(EditAnywhere, BlueprintReadWrite, Category = "Geometry Caches")
+	bool bFlattenTracks = true;
+
+	/** Precision used for compressing vertex positions (lower = better result but less compression, higher = more lossy compression but smaller size) */
+	UPROPERTY(EditAnywhere, BlueprintReadWrite, Category = "Geometry Caches", meta = (ClampMin = "0.000001", ClampMax = "1000", UIMin = "0.0001", UIMax = "10"))
+	float CompressedPositionPrecision = 0.01f;
+
+	/** Bit-precision used for compressing texture coordinates (hight = better result but less compression, lower = more lossy compression but smaller size) */
+	UPROPERTY(EditAnywhere, BlueprintReadWrite, Category = "Geometry Caches", meta = (ClampMin = "1", ClampMax = "31", UIMin = "4", UIMax = "16"))
+	int32 CompressedTextureCoordinatesNumberOfBits = 10;
+
+	/** If enabled, override the imported animation range. Otherwise, the imported range is automatically set to the range of non-empty animated frames */
+	UPROPERTY(EditAnywhere, BlueprintReadWrite, Category = "Geometry Caches", meta = (SubCategory = "Sampling"))
+	bool bOverrideTimeRange = false;
+
+	UPROPERTY(EditAnywhere, BlueprintReadWrite, Category = "Geometry Caches", meta = (SubCategory = "Sampling", editcondition = "bOverrideTimeRange == true"))
+	int32 FrameStart = 0;
+
+	UPROPERTY(EditAnywhere, BlueprintReadWrite, Category = "Geometry Caches", meta = (SubCategory = "Sampling", editcondition = "bOverrideTimeRange == true"))
+	int32 FrameEnd = 1;
+
+	UPROPERTY(EditAnywhere, BlueprintReadWrite, Category = "Geometry Caches", AdvancedDisplay)
+	EInterchangeMotionVectorsHandling MotionVectors = EInterchangeMotionVectorsHandling::NoMotionVectors;
+
+	/**
+	 * Force the preprocessor to only do optimization once instead of when the preprocessor decides. This may lead to some problems with certain meshes but makes sure motion
+	 * blur always works if the topology is constant. 
+	 */
+	UPROPERTY(EditAnywhere, BlueprintReadWrite, Category = "Geometry Caches", AdvancedDisplay)
+	bool bApplyConstantTopologyOptimizations = false;
+
+	/**
+	 * Store the imported vertex numbers. This lets you know the vertex numbers inside the DCC.
+	 * The values of each vertex number will range from 0 to 7 for a cube. Even if the number of positions might be 24.
+	 */
+	UPROPERTY(EditAnywhere, BlueprintReadWrite, Category = "Geometry Caches", AdvancedDisplay)
+	bool bStoreImportedVertexNumbers = false;
+
+	/** Optimizes index buffers for each unique frame, to allow better cache coherency on the GPU. Very costly and time-consuming process, recommended to OFF. */
+	UPROPERTY(EditAnywhere, BlueprintReadWrite, Category = "Geometry Caches", AdvancedDisplay)
+	bool bOptimizeIndexBuffers = false;
+
 	virtual void AdjustSettingsForContext(const FInterchangePipelineContextParams& ContextParams) override;
 
 	virtual void PreDialogCleanup(const FName PipelineStackName) override;
@@ -385,6 +428,7 @@ private:
 		, bool& bAutoDetectConvertStaticMeshToSkeletalMesh
 		, bool& bContainStaticMesh
 		, bool& bContainSkeletalMesh
+		, bool& bContainGeometryCache
 		, bool& bContainStaticMeshAnimationNode
 		, bool& bIgnoreStaticMeshes) const;
 
diff --git a/Engine/Plugins/Runtime/GeometryCache/Source/GeometryCache/Private/GeometryCacheHelpers.cpp b/Engine/Plugins/Runtime/GeometryCache/Source/GeometryCache/Private/GeometryCacheHelpers.cpp
--- a/Engine/Plugins/Runtime/GeometryCache/Source/GeometryCache/Private/GeometryCacheHelpers.cpp
+++ b/Engine/Plugins/Runtime/GeometryCache/Source/GeometryCache/Private/GeometryCacheHelpers.cpp
@@ -12,8 +12,7 @@ namespace UE::GeometryCache::Utils
 	void GetGeometryCacheMeshDataFromMeshDescription(
 		FGeometryCacheMeshData& OutMeshData,
 		FMeshDescription& MeshDescription,
-		int32 MaterialIndex,
-		float FramesPerSecond
+		const FMeshDataConversionArguments& Args
 	)
 	{
 		TRACE_CPUPROFILER_EVENT_SCOPE(GetGeometryCacheMeshDataFromMeshDescription);
@@ -47,7 +46,7 @@ namespace UE::GeometryCache::Utils
 			VertexInstanceVelocities = MeshDescription.VertexInstanceAttributes().GetAttributesRef<FVector3f>(MeshAttribute::VertexInstance::Velocity
 			);
 
-		const bool bHasVelocities = VertexInstanceVelocities.IsValid();
+		const bool bHasVelocities = Args.bUseVelocitiesAsMotionVectors && VertexInstanceVelocities.IsValid();
 		OutMeshData.VertexInfo.bHasMotionVectors = bHasVelocities;
 
 		const int32 NumVertices = MeshDescription.Vertices().Num();
@@ -64,8 +63,15 @@ namespace UE::GeometryCache::Utils
 			OutMeshData.MotionVectors.Reserve(NumMeshDataVertices);
 		}
 
+		const bool bHasImportedVertexNumbers = (NumVertices > 0) && Args.bStoreImportedVertexNumbers;
+		if (bHasImportedVertexNumbers)
+		{
+			OutMeshData.ImportedVertexNumbers.Reserve(NumMeshDataVertices);
+		}
+
 		FBox BoundingBox(EForceInit::ForceInitToZero);
 		int32 VertexIndex = 0;
+		int32 MaterialIndex = Args.MaterialOffset;
 		for (FPolygonGroupID PolygonGroupID : MeshDescription.PolygonGroups().GetElementIDs())
 		{
 			// Skip empty polygon groups
@@ -85,10 +91,16 @@ namespace UE::GeometryCache::Utils
 				{
 					for (FVertexInstanceID VertexInstanceID : MeshDescription.GetTriangleVertexInstances(TriangleID))
 					{
-						const FVector3f& Position = VertexPositions[MeshDescription.GetVertexInstanceVertex(VertexInstanceID)];
+						FVertexID VertexID = MeshDescription.GetVertexInstanceVertex(VertexInstanceID);
+						const FVector3f& Position = VertexPositions[VertexID];
 						OutMeshData.Positions.Add(Position);
 						BoundingBox += (FVector)Position;
 
+						if (bHasImportedVertexNumbers)
+						{
+							OutMeshData.ImportedVertexNumbers.Add(VertexID);
+						}
+
 						OutMeshData.Indices.Add(VertexIndex++);
 
 						FPackedNormal Normal = VertexInstanceNormals[VertexInstanceID];
@@ -106,7 +118,7 @@ namespace UE::GeometryCache::Utils
 						if (bHasVelocities)
 						{
 							FVector3f MotionVector = VertexInstanceVelocities[VertexInstanceID];
-							MotionVector *= (-1.f / FramesPerSecond);	 // Velocity is per seconds but we need per frame for motion vectors
+							MotionVector *= (-1.f / Args.FramesPerSecond);	 // Velocity is per seconds but we need per frame for motion vectors
 
 							OutMeshData.MotionVectors.Add(MotionVector);
 						}
diff --git a/Engine/Plugins/Runtime/GeometryCache/Source/GeometryCache/Public/GeometryCacheHelpers.h b/Engine/Plugins/Runtime/GeometryCache/Source/GeometryCache/Public/GeometryCacheHelpers.h
--- a/Engine/Plugins/Runtime/GeometryCache/Source/GeometryCache/Public/GeometryCacheHelpers.h
+++ b/Engine/Plugins/Runtime/GeometryCache/Source/GeometryCache/Public/GeometryCacheHelpers.h
@@ -25,5 +25,13 @@ struct FMeshDescription;
 
 namespace UE::GeometryCache::Utils
 {
-	GEOMETRYCACHE_API void GetGeometryCacheMeshDataFromMeshDescription(FGeometryCacheMeshData& OutMeshData, FMeshDescription& MeshDescription, int32 MaterialIndex, float FramesPerSecond);
+	struct FMeshDataConversionArguments
+	{	
+		int32 MaterialOffset = 0;
+		float FramesPerSecond = 24.0f;
+		bool bUseVelocitiesAsMotionVectors = true;
+		bool bStoreImportedVertexNumbers = false;
+	};
+
+	GEOMETRYCACHE_API void GetGeometryCacheMeshDataFromMeshDescription(FGeometryCacheMeshData& OutMeshData, FMeshDescription& MeshDescription, const FMeshDataConversionArguments& Args);
 }
\ No newline at end of file
