Interchange USD:
- Fix compile errors with MaterialX

diff --git a/Engine/Plugins/Interchange/Runtime/Source/Import/Private/Usd/InterchangeUsdTranslator.cpp b/Engine/Plugins/Interchange/Runtime/Source/Import/Private/Usd/InterchangeUsdTranslator.cpp
--- a/Engine/Plugins/Interchange/Runtime/Source/Import/Private/Usd/InterchangeUsdTranslator.cpp
+++ b/Engine/Plugins/Interchange/Runtime/Source/Import/Private/Usd/InterchangeUsdTranslator.cpp
@@ -946,6 +946,7 @@ namespace UE::InterchangeUsdTranslator::Private
 			);
 		}
 
+#if WITH_EDITOR
 		// Let's iterate over the Shader Nodes, the TextureSample ones more specifically,
 		// to see if have to retrieve any attributes related to the conversion of geompropvalues
 		// we'll store the UID of the shader node in order to retrieve it during the baking phase (in the post factory import)
@@ -970,6 +971,7 @@ namespace UE::InterchangeUsdTranslator::Private
 				}
 			}
 		);
+#endif // WITH_EDITOR
 
 		// Now we need to check if we have to set the slot of the mesh nodes here
 		if (TArray<FMaterialSlotMesh>* SlotMeshes = PrimPathToSlotMeshNodes.Find(MaterialUid))
@@ -978,6 +980,7 @@ namespace UE::InterchangeUsdTranslator::Private
 			{
 				for (const FMaterialSlotMesh& MaterialSlotMesh : *SlotMeshes)
 				{
+#if WITH_EDITOR
 					if (!GeomPropValues.IsEmpty())
 					{
 						MaterialSlotMesh.MeshNode->AddInt32Attribute(USD::Primvar::Number, GeomPropValues.Num());
@@ -995,7 +998,7 @@ namespace UE::InterchangeUsdTranslator::Private
 							);
 						}
 					}
-
+#endif // WITH_EDITOR
 					if (!MaterialSlotMesh.MeshNode->GetSlotMaterialDependencyUid(MaterialSlotMesh.MaterialSlotName, *NewMaterialUID))
 					{
 						MaterialSlotMesh.MeshNode->SetSlotMaterialDependencyUid(MaterialSlotMesh.MaterialSlotName, *NewMaterialUID);
