[Interchange] Add import skeletalmesh socket feature.
#rnx

diff --git a/Engine/Plugins/Interchange/Runtime/Source/FactoryNodes/Private/InterchangeMeshFactoryNode.cpp b/Engine/Plugins/Interchange/Runtime/Source/FactoryNodes/Private/InterchangeMeshFactoryNode.cpp
--- a/Engine/Plugins/Interchange/Runtime/Source/FactoryNodes/Private/InterchangeMeshFactoryNode.cpp
+++ b/Engine/Plugins/Interchange/Runtime/Source/FactoryNodes/Private/InterchangeMeshFactoryNode.cpp
@@ -172,6 +172,12 @@ bool UInterchangeMeshFactoryNode::ShouldHideAttribute(const UE::Interchange::FAt
 
 #endif //WITH_EDITOR
 
+const FString& UInterchangeMeshFactoryNode::GetMeshSocketPrefix()
+{
+	static FString StaticMeshSocketPrefix = TEXT("SOCKET_");
+	return StaticMeshSocketPrefix;
+}
+
 int32 UInterchangeMeshFactoryNode::GetLodDataCount() const
 {
 	return LodDependencies.GetCount();
@@ -232,6 +238,15 @@ bool UInterchangeMeshFactoryNode::SetCustomKeepSectionsSeparate(const bool& Attr
 	IMPLEMENT_NODE_ATTRIBUTE_SETTER_NODELEGATE(KeepSectionsSeparate, bool)
 }
 
+bool UInterchangeMeshFactoryNode::GetCustomImportSockets(bool& AttributeValue) const
+{
+	IMPLEMENT_NODE_ATTRIBUTE_GETTER(ImportSockets, bool)
+}
+bool UInterchangeMeshFactoryNode::SetCustomImportSockets(const bool& AttributeValue)
+{
+	IMPLEMENT_NODE_ATTRIBUTE_SETTER_NODELEGATE(ImportSockets, bool)
+}
+
 void UInterchangeMeshFactoryNode::GetSlotMaterialDependencies(TMap<FString, FString>& OutMaterialDependencies) const
 {
 	bool bKeepSectionsSeparate = false;
diff --git a/Engine/Plugins/Interchange/Runtime/Source/FactoryNodes/Public/InterchangeMeshFactoryNode.h b/Engine/Plugins/Interchange/Runtime/Source/FactoryNodes/Public/InterchangeMeshFactoryNode.h
--- a/Engine/Plugins/Interchange/Runtime/Source/FactoryNodes/Public/InterchangeMeshFactoryNode.h
+++ b/Engine/Plugins/Interchange/Runtime/Source/FactoryNodes/Public/InterchangeMeshFactoryNode.h
@@ -50,6 +50,9 @@ public:
 #endif //WITH_EDITOR
 
 public:
+
+	static const FString& GetMeshSocketPrefix();
+
 	/** Return the number of LODs this static mesh has.*/
 	UFUNCTION(BlueprintCallable, Category = "Interchange | FactoryNode | Mesh")
 	int32 GetLodDataCount() const;
@@ -95,6 +98,14 @@ public:
 	UFUNCTION(BlueprintCallable, Category = "Interchange | FactoryNode | Mesh")
 	bool SetCustomKeepSectionsSeparate(const bool& AttributeValue);
 
+	/** Query whether the mesh factory should create sockets.*/
+	UFUNCTION(BlueprintCallable, Category = "Interchange | Node | Mesh")
+	bool GetCustomImportSockets(bool& AttributeValue) const;
+
+	/** Set whether the mesh factory should create sockets.*/
+	UFUNCTION(BlueprintCallable, Category = "Interchange | Node | Mesh")
+	bool SetCustomImportSockets(const bool& AttributeValue);
+
 	/** Retrieve the correspondence table between slot names and assigned materials for this object. */
 	UFUNCTION(BlueprintCallable, Category = "Interchange | FactoryNode | Mesh")
 	void GetSlotMaterialDependencies(TMap<FString, FString>& OutMaterialDependencies) const;
@@ -191,19 +202,20 @@ public:
 
 private:
 
-	const UE::Interchange::FAttributeKey Macro_CustomVertexColorReplaceKey = UE::Interchange::FAttributeKey(TEXT("VertexColorReplace"));
-	const UE::Interchange::FAttributeKey Macro_CustomVertexColorIgnoreKey = UE::Interchange::FAttributeKey(TEXT("VertexColorIgnore"));
-	const UE::Interchange::FAttributeKey Macro_CustomVertexColorOverrideKey = UE::Interchange::FAttributeKey(TEXT("VertexColorOverride"));
-	const UE::Interchange::FAttributeKey Macro_CustomKeepSectionsSeparateKey = UE::Interchange::FAttributeKey(TEXT("KeepSectionsSeparate"));
-	const UE::Interchange::FAttributeKey Macro_CustomLODGroupKey = UE::Interchange::FAttributeKey(TEXT("LODGroup"));
-	const UE::Interchange::FAttributeKey Macro_CustomRecomputeNormalsKey = UE::Interchange::FAttributeKey(TEXT("RecomputeNormals"));
-	const UE::Interchange::FAttributeKey Macro_CustomRecomputeTangentsKey = UE::Interchange::FAttributeKey(TEXT("RecomputeTangents"));
-	const UE::Interchange::FAttributeKey Macro_CustomUseMikkTSpaceKey = UE::Interchange::FAttributeKey(TEXT("UseMikkTSpace"));
-	const UE::Interchange::FAttributeKey Macro_CustomComputeWeightedNormalsKey = UE::Interchange::FAttributeKey(TEXT("ComputeWeightedNormals"));
-	const UE::Interchange::FAttributeKey Macro_CustomUseHighPrecisionTangentBasisKey = UE::Interchange::FAttributeKey(TEXT("UseHighPrecisionTangentBasis"));
-	const UE::Interchange::FAttributeKey Macro_CustomUseFullPrecisionUVsKey = UE::Interchange::FAttributeKey(TEXT("UseFullPrecisionUVs"));
-	const UE::Interchange::FAttributeKey Macro_CustomUseBackwardsCompatibleF16TruncUVsKey = UE::Interchange::FAttributeKey(TEXT("UseBackwardsCompatibleF16TruncUVs"));
-	const UE::Interchange::FAttributeKey Macro_CustomRemoveDegeneratesKey = UE::Interchange::FAttributeKey(TEXT("RemoveDegenerates"));
+	IMPLEMENT_NODE_ATTRIBUTE_KEY(VertexColorReplace)
+	IMPLEMENT_NODE_ATTRIBUTE_KEY(VertexColorIgnore)
+	IMPLEMENT_NODE_ATTRIBUTE_KEY(VertexColorOverride)
+	IMPLEMENT_NODE_ATTRIBUTE_KEY(KeepSectionsSeparate)
+	IMPLEMENT_NODE_ATTRIBUTE_KEY(LODGroup)
+	IMPLEMENT_NODE_ATTRIBUTE_KEY(RecomputeNormals)
+	IMPLEMENT_NODE_ATTRIBUTE_KEY(RecomputeTangents)
+	IMPLEMENT_NODE_ATTRIBUTE_KEY(UseMikkTSpace)
+	IMPLEMENT_NODE_ATTRIBUTE_KEY(ComputeWeightedNormals)
+	IMPLEMENT_NODE_ATTRIBUTE_KEY(UseHighPrecisionTangentBasis)
+	IMPLEMENT_NODE_ATTRIBUTE_KEY(UseFullPrecisionUVs)
+	IMPLEMENT_NODE_ATTRIBUTE_KEY(UseBackwardsCompatibleF16TruncUVs)
+	IMPLEMENT_NODE_ATTRIBUTE_KEY(RemoveDegenerates)
+	IMPLEMENT_NODE_ATTRIBUTE_KEY(ImportSockets)
 
 	UE::Interchange::TArrayAttributeHelper<FString> LodDependencies;
 	UE::Interchange::TMapAttributeHelper<FString, FString> SlotMaterialDependencies;
diff --git a/Engine/Plugins/Interchange/Runtime/Source/FactoryNodes/Public/InterchangeSkeletalMeshFactoryNode.h b/Engine/Plugins/Interchange/Runtime/Source/FactoryNodes/Public/InterchangeSkeletalMeshFactoryNode.h
--- a/Engine/Plugins/Interchange/Runtime/Source/FactoryNodes/Public/InterchangeSkeletalMeshFactoryNode.h
+++ b/Engine/Plugins/Interchange/Runtime/Source/FactoryNodes/Public/InterchangeSkeletalMeshFactoryNode.h
@@ -195,20 +195,21 @@ private:
 	virtual bool SetNodeClassFromClassAttribute() override;
 
 	const UE::Interchange::FAttributeKey ClassNameAttributeKey = UE::Interchange::FBaseNodeStaticData::ClassTypeAttributeKey();
-	const UE::Interchange::FAttributeKey Macro_CustomImportMorphTargetKey = UE::Interchange::FAttributeKey(TEXT("ImportMorphTarget"));
-	const UE::Interchange::FAttributeKey Macro_CustomAddCurveMetadataToSkeletonKey = UE::Interchange::FAttributeKey(TEXT("AddCurveMetadataToSkeleton"));
-	const UE::Interchange::FAttributeKey Macro_CustomImportVertexAttributesKey = UE::Interchange::FAttributeKey(TEXT("ImportVertexAttributes"));
-	const UE::Interchange::FAttributeKey Macro_CustomSkeletonSoftObjectPathKey = UE::Interchange::FAttributeKey(TEXT("SkeletonSoftObjectPath"));
-	const UE::Interchange::FAttributeKey Macro_CustomCreatePhysicsAssetKey = UE::Interchange::FAttributeKey(TEXT("CreatePhysicsAsset"));
-	const UE::Interchange::FAttributeKey Macro_CustomPhysicAssetSoftObjectPathKey = UE::Interchange::FAttributeKey(TEXT("PhysicAssetSoftObjectPath"));
-	const UE::Interchange::FAttributeKey Macro_CustomImportContentTypeKey = UE::Interchange::FAttributeKey(TEXT("ImportContentType"));
-	const UE::Interchange::FAttributeKey Macro_CustomUseHighPrecisionSkinWeightsKey = UE::Interchange::FAttributeKey(TEXT("UseHighPrecisionSkinWeights"));
-	const UE::Interchange::FAttributeKey Macro_CustomThresholdPositionKey = UE::Interchange::FAttributeKey(TEXT("ThresholdPosition"));
-	const UE::Interchange::FAttributeKey Macro_CustomThresholdTangentNormalKey = UE::Interchange::FAttributeKey(TEXT("ThresholdTangentNormal"));
-	const UE::Interchange::FAttributeKey Macro_CustomThresholdUVKey = UE::Interchange::FAttributeKey(TEXT("ThresholdUV"));
-	const UE::Interchange::FAttributeKey Macro_CustomMorphThresholdPositionKey = UE::Interchange::FAttributeKey(TEXT("MorphThresholdPosition"));
-	const UE::Interchange::FAttributeKey Macro_CustomBoneInfluenceLimitKey = UE::Interchange::FAttributeKey(TEXT("BoneInfluenceLimit"));
-	const UE::Interchange::FAttributeKey Macro_CustomMergeMorphTargetShapeWithSameNameKey = UE::Interchange::FAttributeKey(TEXT("MergeMorphTargetShapeWithSameName"));
+	IMPLEMENT_NODE_ATTRIBUTE_KEY(ImportMorphTarget)
+	IMPLEMENT_NODE_ATTRIBUTE_KEY(AddCurveMetadataToSkeleton)
+	IMPLEMENT_NODE_ATTRIBUTE_KEY(ImportVertexAttributes)
+	IMPLEMENT_NODE_ATTRIBUTE_KEY(SkeletonSoftObjectPath)
+	IMPLEMENT_NODE_ATTRIBUTE_KEY(CreatePhysicsAsset)
+	IMPLEMENT_NODE_ATTRIBUTE_KEY(PhysicAssetSoftObjectPath)
+	IMPLEMENT_NODE_ATTRIBUTE_KEY(ImportContentType)
+	IMPLEMENT_NODE_ATTRIBUTE_KEY(UseHighPrecisionSkinWeights)
+	IMPLEMENT_NODE_ATTRIBUTE_KEY(ThresholdPosition)
+	IMPLEMENT_NODE_ATTRIBUTE_KEY(ThresholdTangentNormal)
+	IMPLEMENT_NODE_ATTRIBUTE_KEY(ThresholdUV)
+	IMPLEMENT_NODE_ATTRIBUTE_KEY(MorphThresholdPosition)
+	IMPLEMENT_NODE_ATTRIBUTE_KEY(BoneInfluenceLimit)
+	IMPLEMENT_NODE_ATTRIBUTE_KEY(MergeMorphTargetShapeWithSameName)
+	
 
 	bool ApplyCustomUseHighPrecisionSkinWeightsToAsset(UObject* Asset) const;
 	bool FillCustomUseHighPrecisionSkinWeightsFromAsset(UObject* Asset);
diff --git a/Engine/Plugins/Interchange/Runtime/Source/Import/Private/Mesh/InterchangeSkeletalMeshFactory.cpp b/Engine/Plugins/Interchange/Runtime/Source/Import/Private/Mesh/InterchangeSkeletalMeshFactory.cpp
--- a/Engine/Plugins/Interchange/Runtime/Source/Import/Private/Mesh/InterchangeSkeletalMeshFactory.cpp
+++ b/Engine/Plugins/Interchange/Runtime/Source/Import/Private/Mesh/InterchangeSkeletalMeshFactory.cpp
@@ -96,6 +96,27 @@ namespace UE
 
 		namespace Private
 		{
+			void FindAllSockets(const UInterchangeBaseNodeContainer* NodeContainer, const FString& JointNodeId, const FString JointDisplayLabel, TMap<FString, TArray<FString>>& OutSocketsPerBoneMap)
+			{
+				const TArray<FString> ChildrenIds = NodeContainer->GetNodeChildrenUids(JointNodeId);
+				for (int32 ChildIndex = 0; ChildIndex < ChildrenIds.Num(); ++ChildIndex)
+				{
+					const FString& ChildUid = ChildrenIds[ChildIndex];
+					if (const UInterchangeSceneNode* ChildJointNode = Cast<UInterchangeSceneNode>(NodeContainer->GetNode(ChildUid)))
+					{
+						if (FSkeletonHelper::IsValidSocket(NodeContainer, ChildJointNode))
+						{
+							//We found a socket add it to the map
+							TArray<FString>& Sockets = OutSocketsPerBoneMap.FindOrAdd(JointDisplayLabel);
+							Sockets.AddUnique(ChildUid);
+							//Socket cannot have children
+							continue;
+						}
+						FindAllSockets(NodeContainer, ChildrenIds[ChildIndex], ChildJointNode->GetDisplayLabel(), OutSocketsPerBoneMap);
+					}
+				}
+			}
+			
 			void FillMorphTargetMeshDescriptionsPerMorphTargetName(const FMeshNodeContext& MeshNodeContext
 																 , TMap<FString, TOptional<UE::Interchange::FMeshPayloadData>>& MorphTargetMeshDescriptionsPerMorphTargetName
 																 , UInterchangeSkeletalMeshFactory::FLodPayloads& LodPayloads
@@ -1053,7 +1074,8 @@ namespace UE
 				, UInterchangeSkeletalMeshFactory::FImportAssetObjectLODData& ImportAssetObjectLODDataRef
 				, const FTransform& GlobalOffsetTransform
 				, const UInterchangeSkeletalMeshLodDataNode* LodDataNode
-				, const bool bBakeMeshes)
+				, const bool bBakeMeshes
+				, bool bImportSockets)
 			{
 				if (!LodDataNode)
 				{
@@ -1093,7 +1115,8 @@ namespace UE
 					, ImportAssetObjectLODDataRef.RefBonesBinary
 					, ImportAssetObjectLODDataRef.bUseTimeZeroAsBindPose
 					, ImportAssetObjectLODDataRef.bDiffPose
-					, BoneNotBindNames);
+					, BoneNotBindNames
+					, bImportSockets);
 
 				FTransform RootJointNodeGlobalTransform;
 				ensure(RootJointNode->GetCustomGlobalTransform(Arguments.NodeContainer, GlobalOffsetTransform, RootJointNodeGlobalTransform));
@@ -1243,6 +1266,9 @@ void UInterchangeSkeletalMeshFactory::CreatePayloadTasks(const FImportAssetObjec
 	bool bImportMorphTarget = true;
 	SkeletalMeshFactoryNode->GetCustomImportMorphTarget(bImportMorphTarget);
 
+	bool bImportSockets = false;
+	SkeletalMeshFactoryNode->GetCustomImportSockets(bImportSockets);
+
 	int32 LodCount = SkeletalMeshFactoryNode->GetLodDataCount();
 	TArray<FString> LodDataUniqueIds;
 	SkeletalMeshFactoryNode->GetLodDataUniqueIds(LodDataUniqueIds);
@@ -1295,6 +1321,7 @@ void UInterchangeSkeletalMeshFactory::CreatePayloadTasks(const FImportAssetObjec
 			, GlobalOffsetTransform
 			, LodDataNode
 			, bBakeMeshes
+			, bImportSockets
 		);
 
 		struct FInternalInstanceData
@@ -1394,6 +1421,9 @@ UInterchangeFactoryBase::FImportAssetResult UInterchangeSkeletalMeshFactory::Beg
 		return ImportAssetResult;
 	}
 
+	bool bImportSockets = false;
+	SkeletalMeshFactoryNode->GetCustomImportSockets(bImportSockets);
+
 	UObject* ExistingAsset = Arguments.ReimportObject;
 	if (!ExistingAsset)
 	{
@@ -1594,7 +1624,8 @@ UInterchangeFactoryBase::FImportAssetResult UInterchangeSkeletalMeshFactory::Beg
 			, RootJointNodeId
 			, ImportAssetObjectLODData.RefBonesBinary
 			, ImportAssetObjectLODData.bUseTimeZeroAsBindPose
-			, ImportAssetObjectLODData.bDiffPose);
+			, ImportAssetObjectLODData.bDiffPose
+			, bImportSockets);
 
 		if (bSpecifiedSkeleton && !SkeletonReference->IsCompatibleMesh(SkeletalMesh))
 		{
@@ -2177,6 +2208,88 @@ UInterchangeFactoryBase::FImportAssetResult UInterchangeSkeletalMeshFactory::End
 				Sockets.RemoveAt(SocketIndex, EAllowShrinking::No);
 			}
 		}
+
+		//Add imported sockets
+		bool bImportSockets = false;
+		SkeletalMeshFactoryNode->GetCustomImportSockets(bImportSockets);
+		UE::Interchange::Private::FContentInfo ContentInfo = UE::Interchange::Private::GetContentInfo(SkeletalMeshFactoryNode, ImportAssetObjectData.bIsReImport);
+		if (bImportSockets && !ContentInfo.bApplySkinningOnly)
+		{
+			FString RootSkeletonNodeUid;
+			//Find the Root skeleton uid
+			{
+				TArray<FString> LodDataUniqueIds;
+				SkeletalMeshFactoryNode->GetLodDataUniqueIds(LodDataUniqueIds);
+				//We import socket only for the base lod
+				const int32 LodIndex = 0;
+				if(LodDataUniqueIds.IsValidIndex(LodIndex))
+				{
+					FString LodUniqueId = LodDataUniqueIds[LodIndex];
+					const UInterchangeSkeletalMeshLodDataNode* LodDataNode = Cast<UInterchangeSkeletalMeshLodDataNode>(Arguments.NodeContainer->GetNode(LodUniqueId));
+					if (LodDataNode)
+					{
+						FString SkeletonNodeUid;
+						if (LodDataNode->GetCustomSkeletonUid(SkeletonNodeUid))
+						{
+							const UInterchangeSkeletonFactoryNode* SkeletonFactoryNode = Cast<UInterchangeSkeletonFactoryNode>(Arguments.NodeContainer->GetNode(SkeletonNodeUid));
+							if (SkeletonFactoryNode)
+							{
+								SkeletonFactoryNode->GetCustomRootJointUid(RootSkeletonNodeUid);
+							}
+						}
+					}
+				}
+			}
+			//If we find the root skeleton uid we can create the sockets
+			if(!RootSkeletonNodeUid.IsEmpty())
+			{
+				const FReferenceSkeleton& RefSkeleton = SkeletonReference->GetReferenceSkeleton();
+				TMap<FString, TArray<FString>> SocketsPerBoneMap;
+				UE::Interchange::Private::FindAllSockets(Arguments.NodeContainer, RootSkeletonNodeUid, RootSkeletonNodeUid, SocketsPerBoneMap);
+				for (const TPair<FString, TArray<FString>>& BoneAndSockets : SocketsPerBoneMap)
+				{
+					//Find the bone in the skeleton
+					const FString& BoneName = BoneAndSockets.Key;
+					int32 BoneIndex = RefSkeleton.FindBoneIndex(*BoneName);
+					if (RefSkeleton.IsValidIndex(BoneIndex))
+					{
+						//Add or update sockets
+						for (const FString& SocketNodeUid : BoneAndSockets.Value)
+						{
+							if (const UInterchangeSceneNode* SocketSceneNode = Cast<UInterchangeSceneNode>(Arguments.NodeContainer->GetNode(SocketNodeUid)))
+							{
+								const FString SocketNameString = SocketSceneNode->GetDisplayLabel();
+								
+								//Remove the SOCKET_ prefix from the socket node name
+								const int32 SocketPrefixLength = UInterchangeMeshFactoryNode::GetMeshSocketPrefix().Len();
+								FName SocketName = (SocketNameString.Len() > SocketPrefixLength) ? *FString(SocketNameString.RightChop(SocketPrefixLength)) : *SocketNameString;
+								USkeletalMeshSocket* Socket = SkeletalMesh->FindSocket(SocketName);
+								if (!Socket)
+								{
+									//Create the socket
+									Socket = NewObject<USkeletalMeshSocket>(SkeletalMesh);
+									Socket->SocketName = SocketName;
+									Sockets.Add(Socket);
+								}
+
+								if(ensure(Socket))
+								{
+									//Update the socket bone here in case the socket exist but the bone change
+									Socket->BoneName = *BoneName;
+
+									FTransform LocalSocketTransform;
+									SocketSceneNode->GetCustomLocalTransform(LocalSocketTransform);
+									Socket->RelativeLocation = LocalSocketTransform.GetLocation();
+									Socket->RelativeRotation = LocalSocketTransform.GetRotation().Rotator();
+									Socket->RelativeScale = LocalSocketTransform.GetScale3D();
+								}
+							}
+						}
+					}
+				}
+			}
+		}
+
 		Sockets.Shrink();
 	}
 	else
diff --git a/Engine/Plugins/Interchange/Runtime/Source/Import/Private/Mesh/InterchangeSkeletonHelper.cpp b/Engine/Plugins/Interchange/Runtime/Source/Import/Private/Mesh/InterchangeSkeletonHelper.cpp
--- a/Engine/Plugins/Interchange/Runtime/Source/Import/Private/Mesh/InterchangeSkeletonHelper.cpp
+++ b/Engine/Plugins/Interchange/Runtime/Source/Import/Private/Mesh/InterchangeSkeletonHelper.cpp
@@ -5,6 +5,7 @@
 #include "Animation/Skeleton.h"
 #include "InterchangeCommonPipelineDataFactoryNode.h"
 #include "InterchangeImportLog.h"
+#include "InterchangeMeshFactoryNode.h"
 #include "InterchangeResultsContainer.h"
 #include "InterchangeSceneNode.h"
 #include "ReferenceSkeleton.h"
@@ -15,6 +16,16 @@
 
 namespace UE::Interchange::Private
 {
+	bool FSkeletonHelper::IsValidSocket(const UInterchangeBaseNodeContainer* NodeContainer, const UInterchangeSceneNode* JointNode)
+	{
+		
+		if (JointNode && JointNode->GetDisplayLabel().StartsWith(UInterchangeMeshFactoryNode::GetMeshSocketPrefix(), ESearchCase::CaseSensitive))
+		{
+			const FString& JointNodeId = JointNode->GetUniqueID();
+			return (NodeContainer->GetNodeChildrenCount(JointNodeId) == 0);
+		}
+		return false;
+	}
 
 	bool FSkeletonHelper::ProcessImportMeshSkeleton(TObjectPtr<UInterchangeResultsContainer> Results
 		, const USkeleton* SkeletonAsset
@@ -23,7 +34,8 @@ namespace UE::Interchange::Private
 		, const FString& RootJointNodeId
 		, TArray<SkeletalMeshImportData::FBone>& RefBonesBinary
 		, const bool bUseTimeZeroAsBindPose
-		, bool& bOutDiffPose)
+		, bool& bOutDiffPose
+		, bool bImportSockets)
 	{
 		RefBonesBinary.Empty();
 		// Setup skeletal hierarchy + names structure.
@@ -32,7 +44,7 @@ namespace UE::Interchange::Private
 		FReferenceSkeletonModifier RefSkelModifier(RefSkeleton, SkeletonAsset);
 		TArray <FJointInfo> JointInfos;
 		TArray<FString> BoneNotBindNames;
-		RecursiveAddBones(NodeContainer, RootJointNodeId, JointInfos, INDEX_NONE, RefBonesBinary, bUseTimeZeroAsBindPose, bOutDiffPose, BoneNotBindNames);
+		RecursiveAddBones(NodeContainer, RootJointNodeId, JointInfos, INDEX_NONE, RefBonesBinary, bUseTimeZeroAsBindPose, bOutDiffPose, BoneNotBindNames, bImportSockets);
 		if (bOutDiffPose)
 		{
 			//bOutDiffPose can only be true if the user ask to bind with time zero transform.
@@ -73,7 +85,12 @@ namespace UE::Interchange::Private
 		return true;
 	}
 
-	bool FSkeletonHelper::IsCompatibleSkeleton(const USkeleton* Skeleton, const FString RootJoinUid, const UInterchangeBaseNodeContainer* BaseNodeContainer, bool bConvertStaticToSkeletalActive, bool bCheckForIdenticalSkeleton)
+	bool FSkeletonHelper::IsCompatibleSkeleton(const USkeleton* Skeleton
+		, const FString RootJoinUid
+		, const UInterchangeBaseNodeContainer* BaseNodeContainer
+		, bool bConvertStaticToSkeletalActive
+		, bool bCheckForIdenticalSkeleton
+		, bool bImportSockets)
 	{
 		if (!Skeleton)
 		{
@@ -88,7 +105,7 @@ namespace UE::Interchange::Private
 
 		TArray<FMeshBoneInfo> SkeletalLodRawInfos;
 		SkeletalLodRawInfos.Reserve(SkeletonBoneCount);
-		RecursiveBuildSkeletalSkeleton(RootJoinUid, INDEX_NONE, BaseNodeContainer, SkeletalLodRawInfos, bConvertStaticToSkeletalActive);
+		RecursiveBuildSkeletalSkeleton(RootJoinUid, INDEX_NONE, BaseNodeContainer, SkeletalLodRawInfos, bConvertStaticToSkeletalActive, bImportSockets);
 		const int32 SkeletalLodBoneCount = SkeletalLodRawInfos.Num();
 
 		// first ensure the parent exists for each bone
@@ -224,7 +241,8 @@ namespace UE::Interchange::Private
 		, TArray<SkeletalMeshImportData::FBone>& RefBonesBinary
 		, const bool bUseTimeZeroAsBindPose
 		, bool& bOutDiffPose
-		, TArray<FString>& OutBoneNotBindNames)
+		, TArray<FString>& OutBoneNotBindNames
+		, bool bImportSockets)
 	{
 		const UInterchangeSceneNode* JointNode = Cast<UInterchangeSceneNode>(NodeContainer->GetNode(JointNodeId));
 		if (!JointNode)
@@ -301,7 +319,17 @@ namespace UE::Interchange::Private
 		Bone.NumChildren = ChildrenIds.Num();
 		for (int32 ChildIndex = 0; ChildIndex < ChildrenIds.Num(); ++ChildIndex)
 		{
-			RecursiveAddBones(NodeContainer, ChildrenIds[ChildIndex], JointInfos, JointInfoIndex, RefBonesBinary, bUseTimeZeroAsBindPose, bOutDiffPose, OutBoneNotBindNames);
+			const FString& ChildUid = ChildrenIds[ChildIndex];
+			if(const UInterchangeSceneNode* ChildJointNode = Cast<UInterchangeSceneNode>(NodeContainer->GetNode(ChildUid)))
+			{
+				if (bImportSockets && IsValidSocket(NodeContainer, ChildJointNode))
+				{
+					//Sockets will be add by the skeletalmesh factory
+					//We dont want a socket to be a skeleton joint
+					continue;
+				}
+				RecursiveAddBones(NodeContainer, ChildrenIds[ChildIndex], JointInfos, JointInfoIndex, RefBonesBinary, bUseTimeZeroAsBindPose, bOutDiffPose, OutBoneNotBindNames, bImportSockets);
+			}
 		}
 	}
 
@@ -378,7 +406,12 @@ namespace UE::Interchange::Private
 		return true;
 	}
 
-	void FSkeletonHelper::RecursiveBuildSkeletalSkeleton(const FString JoinToAddUid, const int32 ParentIndex, const UInterchangeBaseNodeContainer* BaseNodeContainer, TArray<FMeshBoneInfo>& SkeletalLodRawInfos, bool bConvertStaticToSkeletalActive)
+	void FSkeletonHelper::RecursiveBuildSkeletalSkeleton(const FString JoinToAddUid
+		, const int32 ParentIndex
+		, const UInterchangeBaseNodeContainer* BaseNodeContainer
+		, TArray<FMeshBoneInfo>& SkeletalLodRawInfos
+		, bool bConvertStaticToSkeletalActive
+		, bool bImportSockets)
 	{
 		const UInterchangeSceneNode* SceneNode = Cast<UInterchangeSceneNode>(BaseNodeContainer->GetNode(JoinToAddUid));
 		if (!SceneNode)
@@ -390,6 +423,10 @@ namespace UE::Interchange::Private
 		{
 			return;
 		}
+		if (bImportSockets && IsValidSocket(BaseNodeContainer, SceneNode))
+		{
+			return;
+		}
 
 		int32 JoinIndex = SkeletalLodRawInfos.Num();
 		FMeshBoneInfo& Info = SkeletalLodRawInfos.AddZeroed_GetRef();
@@ -402,7 +439,7 @@ namespace UE::Interchange::Private
 		const TArray<FString> ChildrenIds = BaseNodeContainer->GetNodeChildrenUids(JoinToAddUid);
 		for (int32 ChildIndex = 0; ChildIndex < ChildrenIds.Num(); ++ChildIndex)
 		{
-			RecursiveBuildSkeletalSkeleton(ChildrenIds[ChildIndex], JoinIndex, BaseNodeContainer, SkeletalLodRawInfos);
+			RecursiveBuildSkeletalSkeleton(ChildrenIds[ChildIndex], JoinIndex, BaseNodeContainer, SkeletalLodRawInfos, bConvertStaticToSkeletalActive, bImportSockets);
 		}
 	}
 } //namespace UE::Interchange::Private
diff --git a/Engine/Plugins/Interchange/Runtime/Source/Import/Private/Mesh/InterchangeStaticMeshFactory.cpp b/Engine/Plugins/Interchange/Runtime/Source/Import/Private/Mesh/InterchangeStaticMeshFactory.cpp
--- a/Engine/Plugins/Interchange/Runtime/Source/Import/Private/Mesh/InterchangeStaticMeshFactory.cpp
+++ b/Engine/Plugins/Interchange/Runtime/Source/Import/Private/Mesh/InterchangeStaticMeshFactory.cpp
@@ -1833,6 +1833,14 @@ bool UInterchangeStaticMeshFactory::ImportConvexCollision(const FImportAssetObje
 
 bool UInterchangeStaticMeshFactory::ImportSockets(const FImportAssetObjectParams& Arguments, UStaticMesh* StaticMesh, const UInterchangeStaticMeshFactoryNode* FactoryNode)
 {
+	bool bImportSockets = false;
+	FactoryNode->GetCustomImportSockets(bImportSockets);
+	if (!bImportSockets)
+	{
+		//Skip Import sockets
+		return true;
+	}
+
 	TArray<FString> SocketUids;
 	FactoryNode->GetSocketUids(SocketUids);
 
@@ -1857,9 +1865,9 @@ bool UInterchangeStaticMeshFactory::ImportSockets(const FImportAssetObjectParams
 		if (const UInterchangeSceneNode* SceneNode = Cast<UInterchangeSceneNode>(Arguments.NodeContainer->GetNode(SocketUid)))
 		{
 			FString NodeDisplayName = SceneNode->GetDisplayLabel();
-			if (NodeDisplayName.StartsWith(TEXT("SOCKET_")))
+			if (NodeDisplayName.StartsWith(UInterchangeMeshFactoryNode::GetMeshSocketPrefix()))
 			{
-				NodeDisplayName.RightChopInline(sizeof("SOCKET_") - 1, EAllowShrinking::No);
+				NodeDisplayName.RightChopInline(UInterchangeMeshFactoryNode::GetMeshSocketPrefix().Len(), EAllowShrinking::No);
 			}
 			FName SocketName = FName(NodeDisplayName);
 			ImportedSocketNames.Add(SocketName);
diff --git a/Engine/Plugins/Interchange/Runtime/Source/Import/Public/InterchangeSkeletonHelper.h b/Engine/Plugins/Interchange/Runtime/Source/Import/Public/InterchangeSkeletonHelper.h
--- a/Engine/Plugins/Interchange/Runtime/Source/Import/Public/InterchangeSkeletonHelper.h
+++ b/Engine/Plugins/Interchange/Runtime/Source/Import/Public/InterchangeSkeletonHelper.h
@@ -12,6 +12,7 @@ struct FReferenceSkeleton;
 class UInterchangeBaseNode;
 class UInterchangeBaseNodeContainer;
 class UInterchangeResultsContainer;
+class UInterchangeSceneNode;
 class USkeleton;
 struct FMeshBoneInfo;
 
@@ -28,6 +29,8 @@ namespace UE::Interchange::Private
 	struct INTERCHANGEIMPORT_API FSkeletonHelper
 	{
 	public:
+		static bool IsValidSocket(const UInterchangeBaseNodeContainer* NodeContainer, const UInterchangeSceneNode* Node);
+
 		static bool ProcessImportMeshSkeleton(TObjectPtr<UInterchangeResultsContainer> Results
 			, const USkeleton* SkeletonAsset
 			, FReferenceSkeleton& RefSkeleton
@@ -35,9 +38,15 @@ namespace UE::Interchange::Private
 			, const FString& RootJointNodeId
 			, TArray<SkeletalMeshImportData::FBone>& RefBonesBinary
 			, const bool bUseTimeZeroAsBindPose
-			, bool& bOutDiffPose);
+			, bool& bOutDiffPose
+			, bool bImportSockets);
 
-		static bool IsCompatibleSkeleton(const USkeleton* Skeleton, const FString RootJoinUid, const UInterchangeBaseNodeContainer* BaseNodeContainer, bool bConvertStaticToSkeletalActive, bool bCheckForIdenticalSkeleton = false);
+		static bool IsCompatibleSkeleton(const USkeleton* Skeleton
+			, const FString RootJoinUid
+			, const UInterchangeBaseNodeContainer* BaseNodeContainer
+			, bool bConvertStaticToSkeletalActive
+			, bool bCheckForIdenticalSkeleton
+			, bool bImportSockets);
 		static void RecursiveAddSkeletonMetaDataValues(UInterchangeBaseNodeContainer* NodeContainer, UInterchangeBaseNode* DestinationNode, const FString& JointUid);
 
 		static void RecursiveBoneHasBindPose(const UInterchangeBaseNodeContainer* NodeContainer, const FString& JointNodeId, bool& bHasBoneWithoutBindPose);
@@ -49,7 +58,8 @@ namespace UE::Interchange::Private
 			, TArray<SkeletalMeshImportData::FBone>& RefBonesBinary
 			, const bool bUseTimeZeroAsBindPose
 			, bool& bOutDiffPose
-			, TArray<FString>& OutBoneNotBindNames);
+			, TArray<FString>& OutBoneNotBindNames
+			, bool bImportSockets);
 
 	private:
 		static FName SkeletalLodGetBoneName(const TArray<FMeshBoneInfo>& SkeletalLodRawInfos, int32 BoneIndex);
@@ -61,7 +71,8 @@ namespace UE::Interchange::Private
 			, const int32 ParentIndex
 			, const UInterchangeBaseNodeContainer* BaseNodeContainer
 			, TArray<FMeshBoneInfo>& SkeletalLodRawInfos
-			, bool bConvertStaticToSkeletalActive = false);
+			, bool bConvertStaticToSkeletalActive
+			, bool bImportSockets);
 	};
 } //namespace UE::Interchange::Private
 
diff --git a/Engine/Plugins/Interchange/Runtime/Source/Pipelines/Private/InterchangeGenericAnimationPipeline.cpp b/Engine/Plugins/Interchange/Runtime/Source/Pipelines/Private/InterchangeGenericAnimationPipeline.cpp
--- a/Engine/Plugins/Interchange/Runtime/Source/Pipelines/Private/InterchangeGenericAnimationPipeline.cpp
+++ b/Engine/Plugins/Interchange/Runtime/Source/Pipelines/Private/InterchangeGenericAnimationPipeline.cpp
@@ -905,7 +905,12 @@ void UInterchangeGenericAnimationPipeline::CreateAnimSequenceFactoryNode(UInterc
 
 		//TODO: support skeleton helper in runtime
 #if WITH_EDITOR
-		bSkeletonCompatible = UE::Interchange::Private::FSkeletonHelper::IsCompatibleSkeleton(CommonSkeletalMeshesAndAnimationsProperties->Skeleton.Get(), RootJointUid, BaseNodeContainer, CommonSkeletalMeshesAndAnimationsProperties->bConvertStaticsWithMorphTargetsToSkeletals || CommonMeshesProperties->ForceAllMeshAsType == EInterchangeForceMeshType::IFMT_SkeletalMesh);
+		bSkeletonCompatible = UE::Interchange::Private::FSkeletonHelper::IsCompatibleSkeleton(CommonSkeletalMeshesAndAnimationsProperties->Skeleton.Get()
+			, RootJointUid
+			, BaseNodeContainer
+			, CommonSkeletalMeshesAndAnimationsProperties->bConvertStaticsWithMorphTargetsToSkeletals || CommonMeshesProperties->ForceAllMeshAsType == EInterchangeForceMeshType::IFMT_SkeletalMesh
+			, false /*bCheckForIdenticalSkeleton*/
+			, CommonMeshesProperties->bImportSockets);
 #endif
 		if (bSkeletonCompatible)
 		{
diff --git a/Engine/Plugins/Interchange/Runtime/Source/Pipelines/Private/InterchangeGenericAssetsPipeline.cpp b/Engine/Plugins/Interchange/Runtime/Source/Pipelines/Private/InterchangeGenericAssetsPipeline.cpp
--- a/Engine/Plugins/Interchange/Runtime/Source/Pipelines/Private/InterchangeGenericAssetsPipeline.cpp
+++ b/Engine/Plugins/Interchange/Runtime/Source/Pipelines/Private/InterchangeGenericAssetsPipeline.cpp
@@ -934,7 +934,12 @@ void UInterchangeGenericAssetsPipeline::CreateSkeletonConflict(USkeleton* Specif
 		const bool bConvertStaticToSkeletalActive = CommonSkeletalMeshesAndAnimationsProperties->bConvertStaticsWithMorphTargetsToSkeletals || CommonMeshesProperties->ForceAllMeshAsType == EInterchangeForceMeshType::IFMT_SkeletalMesh;
 		//If we have a compatible skeleton we do not need to create a conflict, for stricter comparisons we check for identical skeletons to report any deviations.
 		constexpr bool bCheckForIdenticalSkeleton = true;
-		if (UE::Interchange::Private::FSkeletonHelper::IsCompatibleSkeleton(SpecifiedSkeleton ? SpecifiedSkeleton : SkeletalMesh->GetSkeleton(), RootJointNodeId, TransientBaseNodeContainer, bConvertStaticToSkeletalActive, bCheckForIdenticalSkeleton))
+		if (UE::Interchange::Private::FSkeletonHelper::IsCompatibleSkeleton(SpecifiedSkeleton ? SpecifiedSkeleton : SkeletalMesh->GetSkeleton()
+			, RootJointNodeId
+			, TransientBaseNodeContainer
+			, bConvertStaticToSkeletalActive
+			, bCheckForIdenticalSkeleton
+			, CommonMeshesProperties->bImportSockets))
 		{
 			return;
 		}
diff --git a/Engine/Plugins/Interchange/Runtime/Source/Pipelines/Private/InterchangeGenericSkeletalMeshPipeline.cpp b/Engine/Plugins/Interchange/Runtime/Source/Pipelines/Private/InterchangeGenericSkeletalMeshPipeline.cpp
--- a/Engine/Plugins/Interchange/Runtime/Source/Pipelines/Private/InterchangeGenericSkeletalMeshPipeline.cpp
+++ b/Engine/Plugins/Interchange/Runtime/Source/Pipelines/Private/InterchangeGenericSkeletalMeshPipeline.cpp
@@ -449,6 +449,7 @@ UInterchangeSkeletalMeshFactoryNode* UInterchangeGenericMeshPipeline::CreateSkel
 	SkeletalMeshFactoryNode->SetCustomImportVertexAttributes(bImportVertexAttributes);
 
 	SkeletalMeshFactoryNode->SetCustomImportContentType(SkeletalMeshImportContentType);
+	SkeletalMeshFactoryNode->SetCustomImportSockets(CommonMeshesProperties->bImportSockets);
 
 	SkeletalMeshFactoryNode->SetCustomAddCurveMetadataToSkeleton(CommonSkeletalMeshesAndAnimationsProperties->bAddCurveMetadataToSkeleton);
 	//If we have a specified skeleton
@@ -458,7 +459,12 @@ UInterchangeSkeletalMeshFactoryNode* UInterchangeGenericMeshPipeline::CreateSkel
 
 		//TODO: support skeleton helper in runtime
 #if WITH_EDITOR
-		bSkeletonCompatible = UE::Interchange::Private::FSkeletonHelper::IsCompatibleSkeleton(CommonSkeletalMeshesAndAnimationsProperties->Skeleton.Get(), RootJointNode->GetUniqueID(), BaseNodeContainer, CommonSkeletalMeshesAndAnimationsProperties->bConvertStaticsWithMorphTargetsToSkeletals || CommonMeshesProperties->ForceAllMeshAsType == EInterchangeForceMeshType::IFMT_SkeletalMesh);
+		bSkeletonCompatible = UE::Interchange::Private::FSkeletonHelper::IsCompatibleSkeleton(CommonSkeletalMeshesAndAnimationsProperties->Skeleton.Get()
+			, RootJointNode->GetUniqueID()
+			, BaseNodeContainer
+			, CommonSkeletalMeshesAndAnimationsProperties->bConvertStaticsWithMorphTargetsToSkeletals || CommonMeshesProperties->ForceAllMeshAsType == EInterchangeForceMeshType::IFMT_SkeletalMesh
+			, false /*bCheckForIdenticalSkeleton*/
+			, CommonMeshesProperties->bImportSockets);
 #endif
 		if (bSkeletonCompatible)
 		{
diff --git a/Engine/Plugins/Interchange/Runtime/Source/Pipelines/Private/InterchangeGenericStaticMeshPipeline.cpp b/Engine/Plugins/Interchange/Runtime/Source/Pipelines/Private/InterchangeGenericStaticMeshPipeline.cpp
--- a/Engine/Plugins/Interchange/Runtime/Source/Pipelines/Private/InterchangeGenericStaticMeshPipeline.cpp
+++ b/Engine/Plugins/Interchange/Runtime/Source/Pipelines/Private/InterchangeGenericStaticMeshPipeline.cpp
@@ -427,6 +427,9 @@ UInterchangeStaticMeshFactoryNode* UInterchangeGenericMeshPipeline::CreateStatic
 
 	StaticMeshFactoryNode->InitializeStaticMeshNode(StaticMeshUid, DisplayLabel, UStaticMesh::StaticClass()->GetName());
 	BaseNodeContainer->AddNode(StaticMeshFactoryNode);
+	
+	//Set the pipeline import sockets property
+	StaticMeshFactoryNode->SetCustomImportSockets(CommonMeshesProperties->bImportSockets);
 
 	if (CommonMeshesProperties->bKeepSectionsSeparate)
 	{
diff --git a/Engine/Plugins/Interchange/Runtime/Source/Pipelines/Private/InterchangePipelineMeshesUtilities.cpp b/Engine/Plugins/Interchange/Runtime/Source/Pipelines/Private/InterchangePipelineMeshesUtilities.cpp
--- a/Engine/Plugins/Interchange/Runtime/Source/Pipelines/Private/InterchangePipelineMeshesUtilities.cpp
+++ b/Engine/Plugins/Interchange/Runtime/Source/Pipelines/Private/InterchangePipelineMeshesUtilities.cpp
@@ -4,6 +4,7 @@
 
 #include "CoreMinimal.h"
 #include "Nodes/InterchangeBaseNode.h"
+#include "InterchangeMeshFactoryNode.h"
 #include "InterchangeMeshNode.h"
 #include "InterchangeSceneNode.h"
 #include "Nodes/InterchangeBaseNodeContainer.h"
@@ -176,7 +177,7 @@ static bool IsSceneNodeASocket(const UInterchangeSceneNode* SceneNode)
 	// Generic pipeline determines where a scene node should be considered as a socket from its naming convention.
 	// This is no longer decided by the translator.
 	FString NodeDisplayName = SceneNode->GetDisplayLabel();
-	return NodeDisplayName.StartsWith(TEXT("SOCKET_"));
+	return NodeDisplayName.StartsWith(UInterchangeMeshFactoryNode::GetMeshSocketPrefix());
 }
 
 bool FInterchangePipelineMeshesUtilitiesContext::IsStaticMeshInstance(const FInterchangeMeshInstance& MeshInstance, UInterchangeBaseNodeContainer* BaseNodeContainer)
diff --git a/Engine/Plugins/Interchange/Runtime/Source/Pipelines/Public/InterchangeGenericAssetsPipelineSharedSettings.h b/Engine/Plugins/Interchange/Runtime/Source/Pipelines/Public/InterchangeGenericAssetsPipelineSharedSettings.h
--- a/Engine/Plugins/Interchange/Runtime/Source/Pipelines/Public/InterchangeGenericAssetsPipelineSharedSettings.h
+++ b/Engine/Plugins/Interchange/Runtime/Source/Pipelines/Public/InterchangeGenericAssetsPipelineSharedSettings.h
@@ -98,6 +98,14 @@ public:
 	UPROPERTY(EditAnywhere, BlueprintReadWrite, Category = "Common Meshes")
 	FColor VertexOverrideColor;
 
+	/**
+	 * If checked, import sockets
+	 * StaticMesh, naming convention, SOCKET_MeshName_xx where "MeshName" should match the mesh you want to add socket to. The xx part is to add a unique id if many socket on the same mesh exist.
+	 * SkeletalMesh, Naming convention, any leaf scene node under the skeleton root with a name starting with "SOCKET_" prefix.
+	 */
+	UPROPERTY(EditAnywhere, BlueprintReadWrite, Category = "Common Meshes")
+	bool bImportSockets = true;
+
 	/** If enabled, normals in the imported mesh are ignored and recomputed. */
 	UPROPERTY(EditAnywhere, BlueprintReadWrite, Category = "Common Meshes", meta = (SubCategory = "Build"))
 	bool bRecomputeNormals = true;
