USD: Hide unused/deprecated option from the translator settings dialog.

diff --git a/Engine/Plugins/Interchange/Runtime/Source/Import/Public/Usd/InterchangeUsdTranslator.h b/Engine/Plugins/Interchange/Runtime/Source/Import/Public/Usd/InterchangeUsdTranslator.h
--- a/Engine/Plugins/Interchange/Runtime/Source/Import/Public/Usd/InterchangeUsdTranslator.h
+++ b/Engine/Plugins/Interchange/Runtime/Source/Import/Public/Usd/InterchangeUsdTranslator.h
@@ -31,7 +31,7 @@ public:
 
 	/** Only import geometry prims with these specific purposes from the USD file */
 	UE_DEPRECATED(5.6, "This option is no longer a translator setting and has moved to the InterchangeUSDPipeline")
-	UPROPERTY(EditAnywhere, Category = "USD Translator", meta = (Bitmask, BitmaskEnum = "/Script/UnrealUSDWrapper.EUsdPurpose"))
+	UPROPERTY()
 	int32 GeometryPurpose;
 
 	/** Specifies which set of shaders to use when parsing USD materials, in addition to the universal render context. */
