USD: Swap 'Import' for 'Support' on error message for consistency with the actual plugin name.

diff --git a/Engine/Plugins/Interchange/Runtime/Source/Import/Private/Usd/InterchangeUsdTranslator.cpp b/Engine/Plugins/Interchange/Runtime/Source/Import/Private/Usd/InterchangeUsdTranslator.cpp
--- a/Engine/Plugins/Interchange/Runtime/Source/Import/Private/Usd/InterchangeUsdTranslator.cpp
+++ b/Engine/Plugins/Interchange/Runtime/Source/Import/Private/Usd/InterchangeUsdTranslator.cpp
@@ -2344,7 +2344,7 @@ namespace UE::InterchangeUsdTranslator::Private
 							Error,
 							TEXT("Failed to find a compatible translator for file '%s'.%s"),
 							*FilePath,
-							bIsOpenVDB ? TEXT(" Is the 'Interchange OpenVDB Import' plugin enabled?") : TEXT("")
+							bIsOpenVDB ? TEXT(" Is the 'Interchange OpenVDB Support' plugin enabled?") : TEXT("")
 						);
 						continue;
 					}
