Interchange: Fixed missing GeometryCache plugin missing from Interchange plugin as a dependency.

#rnx

diff --git a/Engine/Plugins/Interchange/Runtime/Interchange.uplugin b/Engine/Plugins/Interchange/Runtime/Interchange.uplugin
--- a/Engine/Plugins/Interchange/Runtime/Interchange.uplugin
+++ b/Engine/Plugins/Interchange/Runtime/Interchange.uplugin
@@ -87,6 +87,10 @@
 		}
 	],
 	"Plugins": [
+		{
+			"Name": "GeometryCache",
+			"Enabled": true
+		},
 		{
 			"Name": "USDCore",
 			"Enabled": true
