Interchange: Return early if remove redundant keys has nothing to do

diff --git a/Engine/Plugins/Interchange/Runtime/Source/Parsers/CommonParser/Private/InterchangeCommonAnimationPayload.cpp b/Engine/Plugins/Interchange/Runtime/Source/Parsers/CommonParser/Private/InterchangeCommonAnimationPayload.cpp
--- a/Engine/Plugins/Interchange/Runtime/Source/Parsers/CommonParser/Private/InterchangeCommonAnimationPayload.cpp
+++ b/Engine/Plugins/Interchange/Runtime/Source/Parsers/CommonParser/Private/InterchangeCommonAnimationPayload.cpp
@@ -394,6 +394,11 @@ template<typename ValueType>
 void FInterchangeStepCurve::InternalRemoveRedundantKey(TArray<ValueType>& Values, TFunctionRef<bool(const ValueType& ValueA, const ValueType& ValueB)> CompareFunction)
 {
 	const int32 KeyCount = Values.Num();
+	if (KeyCount == 0)
+	{
+		return;
+	}
+
 	TArray<float> NewKeyTimes;
 	NewKeyTimes.Reserve(KeyCount);
 	TArray<ValueType> NewValues;
