Interchange USD:
- Duplicate ShaderGraphNode for USD Meshes that are TwoSided

diff --git a/Engine/Plugins/Interchange/Runtime/Source/Import/Private/Usd/InterchangeUsdTranslator.cpp b/Engine/Plugins/Interchange/Runtime/Source/Import/Private/Usd/InterchangeUsdTranslator.cpp
--- a/Engine/Plugins/Interchange/Runtime/Source/Import/Private/Usd/InterchangeUsdTranslator.cpp
+++ b/Engine/Plugins/Interchange/Runtime/Source/Import/Private/Usd/InterchangeUsdTranslator.cpp
@@ -112,6 +112,7 @@ namespace UE::InterchangeUsdTranslator::Private
 	const static FString PrimitiveShapePrefix = TEXT("\\PrimitiveShape\\");
 
 	const static FString LODContainerSuffix = TEXT("LODContainer");
+	const static FString TwoSidedSuffix = TEXT("_TwoSided");
 
 	const static FString LODString = TEXT("LOD");
 
@@ -270,7 +271,7 @@ namespace UE::InterchangeUsdTranslator::Private
 	public:
 		/** Add a material instance to the node container, otherwise it will add a material if it comes from a Translator (for example coming from
 		 * MaterialX which cannot handle material instances) */
-		FString AddMaterialNode(
+		TArray<FString> AddMaterialNode(
 			const UE::FUsdPrim& Prim,
 			UInterchangeUsdTranslatorSettings* TranslatorSettings,
 			UInterchangeBaseNodeContainer& NodeContainer
@@ -861,11 +862,18 @@ namespace UE::InterchangeUsdTranslator::Private
 		NodeContainer.AddNode(Node);
 	}
 
-	FString AddDisplayColorMaterialInstanceNodeIfNeeded(UInterchangeBaseNodeContainer& NodeContainer, const FString& DisplayColorDesc)
+	FString AddDisplayColorMaterialInstanceNodeIfNeeded(UInterchangeBaseNodeContainer& NodeContainer, const UsdUtils::FUsdPrimMaterialSlot& Slot)
 	{
 		using namespace UsdUnreal::MaterialUtils;
 
-		FString NodeUid = MaterialPrefix + DisplayColorDesc;
+		FString DisplayColorDesc = Slot.MaterialSource;
+		FString NodeUid = MaterialPrefix + Slot.MaterialSource;
+
+		if (Slot.bMeshIsDoubleSided)
+		{
+			NodeUid += TwoSidedSuffix;
+			DisplayColorDesc += TwoSidedSuffix;
+		}
 
 		// We'll treat the DisplayColorDesc (something like "!DisplayColor_1_0") as the material instance UID here
 		const UInterchangeShaderGraphNode* Node = Cast<UInterchangeShaderGraphNode>(NodeContainer.GetNode(NodeUid));
@@ -887,6 +895,11 @@ namespace UE::InterchangeUsdTranslator::Private
 		NewNode->AddBooleanAttribute(UE::Interchange::USD::ParseMaterialIdentifier, true);
 		NewNode->AddStringAttribute(*UE::Interchange::USD::DisplayColorDescIdentifier, DisplayColorDesc);
 
+		if (Slot.bMeshIsDoubleSided)
+		{
+			NewNode->SetCustomTwoSided(true);
+			NewNode->SetCustomTwoSidedTransmission(true); // weirdly enough we also have to set this in case of a translucent material
+		}
 		NodeContainer.AddNode(NewNode);
 
 		return NodeUid;
@@ -946,9 +959,11 @@ namespace UE::InterchangeUsdTranslator::Private
 			);
 		}
 
-#if WITH_EDITOR
 		// Let's iterate over the Shader Nodes, the TextureSample ones more specifically,
 		// to see if have to retrieve any attributes related to the conversion of geompropvalues
+#if WITH_EDITOR
+		// Let's iterate over the Shader Nodes, the TextureSample ones, more specifically,
+		// to see if we have to retrieve any attributes related to the conversion of geompropvalues
 		// we'll store the UID of the shader node in order to retrieve it during the baking phase (in the post factory import)
 		TArray<FString> ShaderNodesTextureSampleUIDs;
 		NodeContainer.IterateNodesOfType<UInterchangeShaderNode>(
@@ -1234,7 +1249,7 @@ namespace UE::InterchangeUsdTranslator::Private
 		const FString* BaseParameterName = nullptr;
 	};
 
-	FString UInterchangeUSDTranslatorImpl::AddMaterialNode(
+	TArray<FString> UInterchangeUSDTranslatorImpl::AddMaterialNode(
 		const UE::FUsdPrim& Prim,
 		UInterchangeUsdTranslatorSettings* TranslatorSettings,
 		UInterchangeBaseNodeContainer& NodeContainer
@@ -1253,18 +1268,34 @@ namespace UE::InterchangeUsdTranslator::Private
 			TOptional<FString> UnrealContentPath = UsdUtils::GetUnrealSurfaceOutput(Prim);
 			if (UnrealContentPath.IsSet())
 			{
-				return UE::InterchangeUsdTranslator::Private::AddUnrealMaterialReferenceNodeIfNeeded(NodeContainer, UnrealContentPath.GetValue());
+				return { UE::InterchangeUsdTranslator::Private::AddUnrealMaterialReferenceNodeIfNeeded(NodeContainer, UnrealContentPath.GetValue()) };
 			}
 		}
 
+		TArray<FString> MaterialUids;
+
 		FString PrimPath = Prim.GetPrimPath().GetString();
 		FString MaterialUid = MaterialPrefix + PrimPath;
+		FString MaterialTwoSidedUid = MaterialUid + TwoSidedSuffix;
 		FString MaterialPrimName(Prim.GetName().ToString());
+		bool bDuplicateMaterialToBeTwoSided = PrimPathToSlotMeshNodes.Contains(MaterialTwoSidedUid);
+
+		MaterialUids.Emplace(MaterialUid);
+		if (bDuplicateMaterialToBeTwoSided)
+		{
+			MaterialUids.Emplace(MaterialTwoSidedUid);
+		}
+
+		// Check if Node already exist as a TwoSided material
+		if (const UInterchangeShaderGraphNode* Node = Cast<UInterchangeShaderGraphNode>(NodeContainer.GetNode(MaterialTwoSidedUid)))
+		{
+			return MaterialUids;
+		}
 
 		// Check if Node already exist with this ID
-		if (const UInterchangeShaderGraphNode* Node = Cast<UInterchangeShaderGraphNode>(NodeContainer.GetNode(MaterialUid)))
+		if (const UInterchangeShaderGraphNode* ShaderGraphNode = Cast<UInterchangeShaderGraphNode>(NodeContainer.GetNode(MaterialUid)))
 		{
-			return MaterialUid;
+			return MaterialUids;
 		}
 
 		// we only create material instances if we didn't find any MaterialX instances (including a shader graph)
@@ -1272,7 +1303,7 @@ namespace UE::InterchangeUsdTranslator::Private
 			AddMaterialXShaderGraph(Prim, TranslatorSettings, NodeContainer, GeomPropValues))
 		{
 			SetMaterialSlotDependencies(MaterialPrimName, MaterialUid, NodeContainer, GeomPropValues, true);
-			return MaterialUid;
+			return MaterialUids;
 		}
 
 		UInterchangeShaderGraphNode* MaterialNode = NewObject<UInterchangeShaderGraphNode>(&NodeContainer);
@@ -1326,7 +1357,21 @@ namespace UE::InterchangeUsdTranslator::Private
 			MaterialNode->AddBooleanAttribute(UE::Interchange::USD::NeedsVTParentIdentifier, true);
 		}
 
-		return MaterialUid;
+		// Check if we have to duplicate the material to be TwoSided
+		if (bDuplicateMaterialToBeTwoSided)
+		{
+			FString MaterialAssetName = MaterialPrimName + TwoSidedSuffix;
+			UInterchangeShaderGraphNode* MaterialTwoSidedNode = NewObject<UInterchangeShaderGraphNode>(&NodeContainer);
+			UInterchangeBaseNode::CopyStorage(MaterialNode, MaterialTwoSidedNode);
+			MaterialTwoSidedNode->InitializeNode(MaterialTwoSidedUid, MaterialAssetName, EInterchangeNodeContainerType::TranslatedAsset);
+			MaterialTwoSidedNode->SetCustomTwoSided(true);
+			MaterialTwoSidedNode->SetCustomTwoSidedTransmission(true);
+			NodeContainer.AddNode(MaterialTwoSidedNode);
+
+			MaterialUidToActualNodeUid.Add(MaterialTwoSidedUid, MaterialTwoSidedUid);
+			SetMaterialSlotDependencies(MaterialAssetName, MaterialTwoSidedUid, NodeContainer);
+		}
+		return MaterialUids;
 	}
 
 	FString UInterchangeUSDTranslatorImpl::GetTextureSourcePath(const FString& TexturePathOnDisk)
@@ -1954,12 +1999,16 @@ namespace UE::InterchangeUsdTranslator::Private
 						bAddDirectlyAsSlotDependency = true;
 
 						// MaterialSource here is e.g. "!DisplayColor_0_1"
-						MaterialInstanceUid = AddDisplayColorMaterialInstanceNodeIfNeeded(NodeContainer, Slot.MaterialSource);
+						MaterialInstanceUid = AddDisplayColorMaterialInstanceNodeIfNeeded(NodeContainer, Slot);
 						break;
 					}
 					case UsdUtils::EPrimAssignmentType::MaterialPrim:
 					{
 						MaterialInstanceUid = MaterialPrefix + Slot.MaterialSource;	   // This is the prim path
+						if (Slot.bMeshIsDoubleSided)
+						{
+							MaterialInstanceUid += TwoSidedSuffix;
+						}
 						break;
 					}
 					case UsdUtils::EPrimAssignmentType::UnrealMaterial:
@@ -1992,7 +2041,7 @@ namespace UE::InterchangeUsdTranslator::Private
 				else
 				{
 					// one material can be attached to several meshes
-					PrimPathToSlotMeshNodes.FindOrAdd(MaterialInstanceUid).Add({SlotName, MeshNode});
+					PrimPathToSlotMeshNodes.FindOrAdd(MaterialInstanceUid).Emplace(SlotName, MeshNode);
 				}
 			}
 		}
@@ -3176,7 +3225,7 @@ namespace UE::InterchangeUsdTranslator::Private
 		TArray<FString> AssetNodeUids;
 		if (Prim.IsA(TEXT("Material")))
 		{
-			AssetNodeUids = {TranslatorImpl.AddMaterialNode(Prim, TranslatorSettings, NodeContainer)};
+			AssetNodeUids = TranslatorImpl.AddMaterialNode(Prim, TranslatorSettings, NodeContainer);
 		}
 		else if (Prim.IsA(TEXT("Mesh")))
 		{
