Fix compile error on mac, fallback from Emplace to Add.

diff --git a/Engine/Plugins/Interchange/Runtime/Source/Import/Private/Usd/InterchangeUsdTranslator.cpp b/Engine/Plugins/Interchange/Runtime/Source/Import/Private/Usd/InterchangeUsdTranslator.cpp
--- a/Engine/Plugins/Interchange/Runtime/Source/Import/Private/Usd/InterchangeUsdTranslator.cpp
+++ b/Engine/Plugins/Interchange/Runtime/Source/Import/Private/Usd/InterchangeUsdTranslator.cpp
@@ -2041,7 +2041,7 @@ namespace UE::InterchangeUsdTranslator::Private
 				else
 				{
 					// one material can be attached to several meshes
-					PrimPathToSlotMeshNodes.FindOrAdd(MaterialInstanceUid).Emplace(SlotName, MeshNode);
+					PrimPathToSlotMeshNodes.FindOrAdd(MaterialInstanceUid).Add({SlotName, MeshNode});
 				}
 			}
 		}
