USD: Add support for making primvar-compatible material instances to USD Interchange.

diff --git a/Engine/Plugins/Importers/USDImporter/Source/USDSchemas/Private/MeshTranslationImpl.cpp b/Engine/Plugins/Importers/USDImporter/Source/USDSchemas/Private/MeshTranslationImpl.cpp
--- a/Engine/Plugins/Importers/USDImporter/Source/USDSchemas/Private/MeshTranslationImpl.cpp
+++ b/Engine/Plugins/Importers/USDImporter/Source/USDSchemas/Private/MeshTranslationImpl.cpp
@@ -4,6 +4,7 @@
 
 #include "Objects/USDInfoCache.h"
 #include "Objects/USDPrimLinkCache.h"
+#include "UnrealUSDWrapper.h"
 #include "USDAssetCache3.h"
 #include "USDAssetUserData.h"
 #include "USDClassesModule.h"
@@ -38,8 +39,6 @@
 #include "pxr/usd/usdGeom/mesh.h"
 #include "USDIncludesEnd.h"
 
-#define UNUSED_UV_INDEX USD_PREVIEW_SURFACE_MAX_UV_SETS
-
 static_assert(
 	USD_PREVIEW_SURFACE_MAX_UV_SETS <= MAX_MESH_TEXTURE_COORDS_MD,
 	"UsdPreviewSurface materials can only have up to as many UV sets as MeshDescription supports!"
diff --git a/Engine/Plugins/Interchange/Runtime/Source/Common/Public/Usd/InterchangeUsdDefinitions.h b/Engine/Plugins/Interchange/Runtime/Source/Common/Public/Usd/InterchangeUsdDefinitions.h
--- a/Engine/Plugins/Interchange/Runtime/Source/Common/Public/Usd/InterchangeUsdDefinitions.h
+++ b/Engine/Plugins/Interchange/Runtime/Source/Common/Public/Usd/InterchangeUsdDefinitions.h
@@ -14,6 +14,20 @@ namespace UE::Interchange::USD
 	// Prefix for custom attributes that describe material instance parameters
 	const FString AttributeIdentifierPrefix = TEXT("USD_MI_Attribute_");
 
+	// Prefixes we use to stash some primvar mapping information as custom attributes on mesh / material nodes, so that
+	// the USD Pipeline can produce primvar-compatible materials
+	const FString PrimvarUVIndexAttributePrefix = TEXT("USD_PrimvarUVIndex_");
+	const FString ParameterToPrimvarAttributePrefix = TEXT("USD_ParameterPrimvar_");
+
+	// Additional suffix we add to the UID of all primvar-compatible materials
+	const FString CompatibleMaterialUidSuffix = TEXT("_USD_CompatibleMaterial_");
+
+	// Some tokens we add to the material parameter name for USD material nodes. Put here because we need to use the
+	// same tokens on the translator and the USD Pipeline, when computing primvar-compatible materials
+	const FString UseTextureParameterPrefix = TEXT("Use");
+	const FString UseTextureParameterSuffix = TEXT("Texture");
+	const FString UVIndexParameterSuffix = TEXT("UVIndex");
+
 	// Flag indicating whether we should parse a UInterchangeShaderGraphNode on the InterchangeUSDPipeline
 	const FString ParseMaterialIdentifier = TEXT("USD_MI_ParseMaterial");
 
diff --git a/Engine/Plugins/Interchange/Runtime/Source/Import/Private/Usd/InterchangeUsdTranslator.cpp b/Engine/Plugins/Interchange/Runtime/Source/Import/Private/Usd/InterchangeUsdTranslator.cpp
--- a/Engine/Plugins/Interchange/Runtime/Source/Import/Private/Usd/InterchangeUsdTranslator.cpp
+++ b/Engine/Plugins/Interchange/Runtime/Source/Import/Private/Usd/InterchangeUsdTranslator.cpp
@@ -1124,8 +1124,10 @@ namespace UE::InterchangeUsdTranslator::Private
 
 		void EnableTextureForChannel(bool bEnable) const
 		{
-			MaterialNode.AddFloatAttribute(	   //
-				MakeInstanceParameterName(FString::Printf(TEXT("Use%sTexture"), **BaseParameterName)),
+			using namespace UE::Interchange::USD;
+
+			MaterialNode.AddFloatAttribute(
+				MakeInstanceParameterName(UseTextureParameterPrefix + *BaseParameterName + UseTextureParameterSuffix),
 				bEnable ? 1.0f : 0.0f
 			);
 		}
@@ -1177,7 +1179,7 @@ namespace UE::InterchangeUsdTranslator::Private
 			if (const int32* FoundIndex = PrimvarToUVIndex.Find(Value.Primvar))
 			{
 				MaterialNode.AddFloatAttribute(	   //
-					MakeInstanceParameterName(FString::Printf(TEXT("%sUVIndex"), **BaseParameterName)),
+					MakeInstanceParameterName(*BaseParameterName + UE::Interchange::USD::UVIndexParameterSuffix),
 					*FoundIndex
 				);
 			}
@@ -1340,6 +1342,31 @@ namespace UE::InterchangeUsdTranslator::Private
 			}
 		}
 
+		// Also set our parameter to uv index mapping as-is also as custom attributes, so that the USD Pipeline can make
+		// primvar-compatible materials
+		{
+			// Parameter to primvar
+			for (const TPair<FString, UsdToUnreal::FParameterValue>& Parameter : MaterialData.Parameters)
+			{
+				if (const UsdToUnreal::FTextureParameterValue* TextureParameterValue = Parameter.Value.TryGet<UsdToUnreal::FTextureParameterValue>())
+				{
+					const FString& MaterialParameter = Parameter.Key;
+					const FString& Primvar = TextureParameterValue->Primvar;
+
+					MaterialNode->AddStringAttribute(UE::Interchange::USD::ParameterToPrimvarAttributePrefix + MaterialParameter, Primvar);
+				}
+			}
+
+			// Primvar to UVIndex
+			for (const TPair<FString, int32>& Pair : MaterialData.PrimvarToUVIndex)
+			{
+				const FString& Primvar = Pair.Key;
+				int32 UVIndex = Pair.Value;
+
+				MaterialNode->AddInt32Attribute(UE::Interchange::USD::PrimvarUVIndexAttributePrefix + Primvar, UVIndex);
+			}
+		}
+
 		// Find and set the relevant info so that we can pick the reference material on the pipeline
 		//
 		// TODO: Proper VT texture support (we'd need to know the texture resolution at this point, and we haven't parsed them yet...).
@@ -1981,6 +2008,22 @@ namespace UE::InterchangeUsdTranslator::Private
 				CachedMaterialAssignments.Add(PrimPath, Assignments);
 			}
 
+			// Move these into the asset node because the USD Pipeline will compare these with the assigned material's
+			// parameter-to-primvar mapping in order to make sure the mesh is using a primvar-compatible material.
+			//
+			// Note that ideally we'd cache this mapping and reuse it on the payload retrieval step. Instead, we will
+			// just end up calling the same function again during payload retrieval and hoping that it produces the same
+			// primvar-to-UVIndex mapping. It should though, as the mesh conversion options are the same. We can't cache
+			// the mapping because we run into USD allocator issues, given that all the strings contained in the
+			// FUsdPrimMaterialAssignmentInfo object are allocated inside an USD allocator scope
+			for (const TPair<FString, int32>& PrimvarPair : UsdUtils::GetPrimvarToUVIndexMap(Prim))
+			{
+				const FString& PrimvarName = PrimvarPair.Key;
+				int32 UVIndex = PrimvarPair.Value;
+
+				MeshNode->AddInt32Attribute(UE::Interchange::USD::PrimvarUVIndexAttributePrefix + PrimvarName, UVIndex);
+			}
+
 			for (const UsdUtils::FUsdPrimMaterialSlot& Slot : Assignments.Slots)
 			{
 				// We do this because Interchange will, in some scenarios, merge material slots with identical slot names.
diff --git a/Engine/Plugins/Interchange/Runtime/Source/Pipelines/Private/InterchangeUSDPipeline.cpp b/Engine/Plugins/Interchange/Runtime/Source/Pipelines/Private/InterchangeUSDPipeline.cpp
--- a/Engine/Plugins/Interchange/Runtime/Source/Pipelines/Private/InterchangeUSDPipeline.cpp
+++ b/Engine/Plugins/Interchange/Runtime/Source/Pipelines/Private/InterchangeUSDPipeline.cpp
@@ -33,7 +33,9 @@
 #include "StaticMeshOperations.h"
 
 #include "UnrealUSDWrapper.h"
+#include "Usd/InterchangeUsdDefinitions.h"
 #include "USDConversionUtils.h"
+#include "USDLog.h"
 #include "USDMaterialUtils.h"
 #include "USDProjectSettings.h"
 #include "UsdWrappers/SdfPath.h"
@@ -283,10 +285,13 @@ namespace UE::InterchangeUsdPipeline::Private
 		const UInterchangeShaderGraphNode& ShaderGraphNode
 	)
 	{
-		FString MaterialFactoryNodeUid = UInterchangeFactoryBaseNode::BuildFactoryNodeUid(ShaderGraphNode.GetUniqueID());
+		const FString ShaderGraphNodeUid = ShaderGraphNode.GetUniqueID();
+		FString MaterialFactoryNodeUid = UInterchangeFactoryBaseNode::BuildFactoryNodeUid(ShaderGraphNodeUid);
 		FString MaterialFactoryNodeName = ShaderGraphNode.GetDisplayLabel();
+
 		UInterchangeMaterialInstanceFactoryNode* NewFactoryNode = NewObject<UInterchangeMaterialInstanceFactoryNode>(&NodeContainer);
 		NewFactoryNode->InitializeNode(MaterialFactoryNodeUid, MaterialFactoryNodeName, EInterchangeNodeContainerType::FactoryData);
+		NewFactoryNode->AddTargetNodeUid(ShaderGraphNodeUid);
 
 		const bool bAddSourceNodeName = false;
 		UInterchangeUserDefinedAttributesAPI::DuplicateAllUserDefinedAttribute(&ShaderGraphNode, NewFactoryNode, bAddSourceNodeName);
@@ -301,71 +306,74 @@ namespace UE::InterchangeUsdPipeline::Private
 			NewFactoryNode->SetCustomInstanceClassName(MaterialClass->GetPathName());
 		}
 
-		// Handle material attributes
+		// Move material attributes to the factory nodes
 		TArray<UE::Interchange::FAttributeKey> AllAttributeKeys;
 		ShaderGraphNode.GetAttributeKeys(AllAttributeKeys);
 		for (const UE::Interchange::FAttributeKey& AttributeKey : AllAttributeKeys)
 		{
 			const FString AttributeKeyString = AttributeKey.ToString();
 
-			if (!AttributeKeyString.Contains(UE::Interchange::USD::AttributeIdentifierPrefix))
-			{
-				// It's not one of our attributes
-				continue;
-			}
-
-			const FString CleanAttributeKey = AttributeKeyString.Replace(	 //
-				*UE::Interchange::USD::AttributeIdentifierPrefix,
-				TEXT(""),
-				ESearchCase::CaseSensitive
-			);
-			const FString StandardParameterName = UInterchangeShaderPortsAPI::MakeInputValueKey(CleanAttributeKey);
-
-			const UE::Interchange::EAttributeTypes AttributeType = ShaderGraphNode.GetAttributeType(AttributeKey);
-			switch (AttributeType)
+			// It's a regular material attribute
+			if (AttributeKeyString.Contains(UE::Interchange::USD::AttributeIdentifierPrefix))
 			{
-				case UE::Interchange::EAttributeTypes::Bool:
+				const FString CleanAttributeKey = AttributeKeyString.Replace(	 //
+					*UE::Interchange::USD::AttributeIdentifierPrefix,
+					TEXT(""),
+					ESearchCase::CaseSensitive
+				);
+				const FString StandardParameterName = UInterchangeShaderPortsAPI::MakeInputValueKey(CleanAttributeKey);
+
+				const UE::Interchange::EAttributeTypes AttributeType = ShaderGraphNode.GetAttributeType(AttributeKey);
+				switch (AttributeType)
 				{
-					bool Value;
-					if (ShaderGraphNode.GetBooleanAttribute(AttributeKey.Key, Value))
-					{
-						NewFactoryNode->AddBooleanAttribute(StandardParameterName, Value);
-					}
-				}
-				break;
-				case UE::Interchange::EAttributeTypes::Float:
-				{
-					float Value;
-					if (ShaderGraphNode.GetFloatAttribute(AttributeKey.Key, Value))
-					{
-						NewFactoryNode->AddFloatAttribute(StandardParameterName, Value);
-					}
+					case UE::Interchange::EAttributeTypes::Bool:
+						if (bool Value; ShaderGraphNode.GetBooleanAttribute(AttributeKey.Key, Value))
+						{
+							NewFactoryNode->AddBooleanAttribute(StandardParameterName, Value);
+						}
+						break;
+					case UE::Interchange::EAttributeTypes::Float:
+						if (float Value; ShaderGraphNode.GetFloatAttribute(AttributeKey.Key, Value))
+						{
+							NewFactoryNode->AddFloatAttribute(StandardParameterName, Value);
+						}
+						break;
+					case UE::Interchange::EAttributeTypes::LinearColor:
+						if (FLinearColor Value; ShaderGraphNode.GetLinearColorAttribute(AttributeKey.Key, Value))
+						{
+							NewFactoryNode->AddLinearColorAttribute(StandardParameterName, Value);
+						}
+						break;
+					case UE::Interchange::EAttributeTypes::String:
+						if (FString TextureTranslatedNodeUid; ShaderGraphNode.GetStringAttribute(AttributeKey.Key, TextureTranslatedNodeUid))
+						{
+							// We only use string attributes for textures for now
+							FString FactoryTextureUid = UInterchangeFactoryBaseNode::BuildFactoryNodeUid(TextureTranslatedNodeUid);
+
+							NewFactoryNode->AddStringAttribute(StandardParameterName, FactoryTextureUid);
+							NewFactoryNode->AddFactoryDependencyUid(FactoryTextureUid);
+						}
+						break;
+					default:
+						break;
 				}
-				break;
-				case UE::Interchange::EAttributeTypes::LinearColor:
+			}
+			// It's another attribute related to primvar-compatible materials, just move it over to the factory node as-is
+			else if (AttributeKeyString.StartsWith(UE::Interchange::USD::ParameterToPrimvarAttributePrefix))
+			{
+				FString AttributeValue;
+				if (ShaderGraphNode.GetStringAttribute(AttributeKey.Key, AttributeValue))
 				{
-					FLinearColor Value;
-					if (ShaderGraphNode.GetLinearColorAttribute(AttributeKey.Key, Value))
-					{
-						NewFactoryNode->AddLinearColorAttribute(StandardParameterName, Value);
-					}
+					NewFactoryNode->AddStringAttribute(AttributeKeyString, AttributeValue);
 				}
-				break;
-				case UE::Interchange::EAttributeTypes::String:
+			}
+			else if (AttributeKeyString.StartsWith(UE::Interchange::USD::PrimvarUVIndexAttributePrefix))
+			{
+				int32 AttributeValue;
+				if (ShaderGraphNode.GetInt32Attribute(AttributeKey.Key, AttributeValue))
 				{
-					FString TextureTranslatedNodeUid;
-					if (ShaderGraphNode.GetStringAttribute(AttributeKey.Key, TextureTranslatedNodeUid))
-					{
-						// We only use string attributes for textures for now
-						FString FactoryTextureUid = UInterchangeFactoryBaseNode::BuildFactoryNodeUid(TextureTranslatedNodeUid);
-
-						NewFactoryNode->AddStringAttribute(StandardParameterName, FactoryTextureUid);
-						NewFactoryNode->AddFactoryDependencyUid(FactoryTextureUid);
-					}
+					NewFactoryNode->AddInt32Attribute(AttributeKeyString, AttributeValue);
 				}
-				break;
-				default:
-					break;
 			}
 		}
 
@@ -742,52 +750,432 @@ namespace UE::InterchangeUsdPipeline::Private
 	/**
 	 * Add the primvar attributes from the mesh nodes to factory nodes
 	 */
-	void AddPrimvarAttributesToFactoryNodes(UInterchangeBaseNodeContainer* NodeContainer)
+	void AddPrimvarAttributesToFactoryNodes(
+		const UInterchangeMeshNode* MeshNode,
+		UInterchangeMeshFactoryNode* MeshFactoryNode,
+		UInterchangeBaseNodeContainer* NodeContainer
+	)
 	{
 		using namespace UE::Interchange;
 
-		NodeContainer->IterateNodesOfType<UInterchangeMeshNode>(
-			[NodeContainer](const FString& NodeUid, UInterchangeMeshNode* MeshNode)
+		if (int32 NumPrimvar; MeshNode->GetInt32Attribute(USD::Primvar::Number, NumPrimvar))
+		{
+			MeshFactoryNode->AddInt32Attribute(USD::Primvar::Number, NumPrimvar);
+			for (int32 Index = 0; Index < NumPrimvar; ++Index)
 			{
-				if (FString FactoryNodeUid = UInterchangeFactoryBaseNode::BuildFactoryNodeUid(MeshNode->GetUniqueID());
-					UInterchangeMeshFactoryNode* MeshFactoryNode = Cast<UInterchangeMeshFactoryNode>(NodeContainer->GetFactoryNode(FactoryNodeUid)))
+				FString PrimvarAttribute = USD::Primvar::Name + FString::FromInt(Index);
+				if (FString PrimvarName; MeshNode->GetStringAttribute(PrimvarAttribute, PrimvarName))
 				{
-					if (int32 NumPrimvar; MeshNode->GetInt32Attribute(USD::Primvar::Number, NumPrimvar))
+					MeshFactoryNode->AddStringAttribute(PrimvarAttribute, PrimvarName);
+				}
+
+				if (bool bTangentSpace; MeshNode->GetBooleanAttribute(PrimvarAttribute, bTangentSpace))
+				{
+					MeshFactoryNode->AddBooleanAttribute(PrimvarAttribute, bTangentSpace);
+				}
+
+				FString ShaderNodeTextureSampleAttribute = USD::Primvar::ShaderNodeTextureSample + FString::FromInt(Index);
+				if (FString ShaderNodeTextureSampleUID; MeshNode->GetStringAttribute(ShaderNodeTextureSampleAttribute, ShaderNodeTextureSampleUID))
+				{
+					// At this stage the USD pipeline should be the latest one in the stack, we should retrieve the MaterialExpression
+					// Factory Node associated to this ShaderNode
+					FString TextureSampleFactoryNodeUID = UInterchangeFactoryBaseNode::BuildFactoryNodeUid(ShaderNodeTextureSampleUID);
+
+					if (Cast<UInterchangeMaterialExpressionFactoryNode>(NodeContainer->GetFactoryNode(TextureSampleFactoryNodeUID)))
 					{
-						MeshFactoryNode->AddInt32Attribute(USD::Primvar::Number, NumPrimvar);
-						for (int32 Index = 0; Index < NumPrimvar; ++Index)
-						{
-							FString PrimvarAttribute = USD::Primvar::Name + FString::FromInt(Index);
-							if (FString PrimvarName; MeshNode->GetStringAttribute(PrimvarAttribute, PrimvarName))
-							{
-								MeshFactoryNode->AddStringAttribute(PrimvarAttribute, PrimvarName);
-							}
+						MeshFactoryNode->AddStringAttribute(ShaderNodeTextureSampleAttribute, TextureSampleFactoryNodeUID);
+					}
+				}
+			}
+		}
+	}
+#endif	  // WITH_EDITORONLY_DATA
 
-							if (bool bTangentSpace; MeshNode->GetBooleanAttribute(PrimvarAttribute, bTangentSpace))
-							{
-								MeshFactoryNode->AddBooleanAttribute(PrimvarAttribute, bTangentSpace);
-							}
+	struct FMaterialPrimvarInfo
+	{
+		TMap<FString, FString> ParameterToPrimvar;
+		TMap<FString, int32> PrimvarToUVIndex;
+	};
 
-							FString ShaderNodeTextureSampleAttribute = USD::Primvar::ShaderNodeTextureSample + FString::FromInt(Index);
-							if (FString ShaderNodeTextureSampleUID;
-								MeshNode->GetStringAttribute(ShaderNodeTextureSampleAttribute, ShaderNodeTextureSampleUID))
-							{
-								// At this stage the USD pipeline should be the latest one in the stack, we should retrieve the MaterialExpression
-								// Factory Node associated to this ShaderNode
-								FString TextureSampleFactoryNodeUID = UInterchangeFactoryBaseNode::BuildFactoryNodeUid(ShaderNodeTextureSampleUID);
-
-								if (Cast<UInterchangeMaterialExpressionFactoryNode>(NodeContainer->GetFactoryNode(TextureSampleFactoryNodeUID)))
-								{
-									MeshFactoryNode->AddStringAttribute(ShaderNodeTextureSampleAttribute, TextureSampleFactoryNodeUID);
-								}
-							}
-						}
+	// Gets or creates a filled out FMaterialPrimvarInfo struct for a particular factory node.
+	// We use this because we may iterate over the same material many times during ProcessMeshNodes,
+	// and we don't want to recompute this info every time
+	FMaterialPrimvarInfo* GetOrCreateMaterialPrimvarInfo(
+		UInterchangeMaterialInstanceFactoryNode* InstanceFactoryNode,
+		TMap<FString, FMaterialPrimvarInfo>& InOutFactoryNodeUidToInfo
+	)
+	{
+		using namespace UE::Interchange;
+
+		if (!InstanceFactoryNode)
+		{
+			return nullptr;
+		}
+
+		const FString& MaterialUid = InstanceFactoryNode->GetUniqueID();
+
+		FMaterialPrimvarInfo* FoundInfo = InOutFactoryNodeUidToInfo.Find(MaterialUid);
+		if (!FoundInfo)
+		{
+			FMaterialPrimvarInfo& NewInfo = InOutFactoryNodeUidToInfo.Emplace(MaterialUid);
+
+			TArray<UE::Interchange::FAttributeKey> MaterialAttributeKeys;
+			InstanceFactoryNode->GetAttributeKeys(MaterialAttributeKeys);
+
+			NewInfo.ParameterToPrimvar.Reserve(MaterialAttributeKeys.Num());
+			NewInfo.PrimvarToUVIndex.Reserve(MaterialAttributeKeys.Num());
+			for (const UE::Interchange::FAttributeKey& AttributeKey : MaterialAttributeKeys)
+			{
+				const FString& AttributeKeyString = AttributeKey.ToString();
+
+				if (AttributeKeyString.StartsWith(USD::ParameterToPrimvarAttributePrefix))
+				{
+					FString ParameterName = AttributeKeyString;
+					ensure(ParameterName.RemoveFromStart(USD::ParameterToPrimvarAttributePrefix));
+
+					FString Primvar;
+					if (InstanceFactoryNode->GetStringAttribute(AttributeKeyString, Primvar))
+					{
+						NewInfo.ParameterToPrimvar.Add(ParameterName, Primvar);
+					}
+				}
+				else if (AttributeKeyString.StartsWith(USD::PrimvarUVIndexAttributePrefix))
+				{
+					FString ParameterName = AttributeKeyString;
+					ensure(ParameterName.RemoveFromStart(USD::PrimvarUVIndexAttributePrefix));
+
+					int32 UVIndex = -1;
+					if (InstanceFactoryNode->GetInt32Attribute(AttributeKeyString, UVIndex))
+					{
+						NewInfo.PrimvarToUVIndex.Add(ParameterName, UVIndex);
+					}
+				}
+			}
+
+			FoundInfo = &NewInfo;
+		}
+
+		return FoundInfo;
+	}
+
+	// Returns a primvar-compatible version of OriginalFactoryNode.
+	// i.e. creates (if needed) a new MaterialInstanceFactoryNode that has a primvar to UV index
+	// mapping that is usable with the primvar to UV index layout in MeshPrimvarToUVIndex
+	UInterchangeMaterialInstanceFactoryNode* GetOrCreateCompatibleMaterial(
+		const UInterchangeMaterialInstanceFactoryNode* OriginalFactoryNode,
+		const FMaterialPrimvarInfo* OriginalMaterialPrimvarInfo,
+		const TMap<FString, int32>& MeshPrimvarToUVIndex,
+		UInterchangeBaseNodeContainer* InNodeContainer
+	)
+	{
+		using namespace UE::Interchange;
+
+		if (!OriginalFactoryNode || !OriginalMaterialPrimvarInfo || !InNodeContainer)
+		{
+			return nullptr;
+		}
+
+		// First, let's create the target primvar UVIndex assignment that is compatible with this mesh.
+		// We use an array of TPairs here so that we can sort these into a deterministic order for hashing later.
+		TArray<TPair<FString, int32>> CompatiblePrimvarAndUVIndexPairs;
+		CompatiblePrimvarAndUVIndexPairs.Reserve(OriginalMaterialPrimvarInfo->PrimvarToUVIndex.Num());
+		for (const TPair<FString, int32>& Pair : OriginalMaterialPrimvarInfo->PrimvarToUVIndex)
+		{
+			const FString& MaterialPrimvar = Pair.Key;
+
+			bool bFoundUVIndex = false;
+
+			// Mesh has this primvar available at some UV index, point to it
+			if (const int32* FoundMeshUVIndex = MeshPrimvarToUVIndex.Find(MaterialPrimvar))
+			{
+				int32 MeshUVIndex = *FoundMeshUVIndex;
+				if (MeshUVIndex >= 0 && MeshUVIndex < USD_PREVIEW_SURFACE_MAX_UV_SETS)
+				{
+					CompatiblePrimvarAndUVIndexPairs.Add(TPair<FString, int32>{MaterialPrimvar, MeshUVIndex});
+					bFoundUVIndex = true;
+				}
+			}
+
+			if (!bFoundUVIndex)
+			{
+				// Point this primvar to read an unused UV index instead, since our mesh doesn't have this primvar
+				CompatiblePrimvarAndUVIndexPairs.Add(TPair<FString, int32>{MaterialPrimvar, UNUSED_UV_INDEX});
+			}
+		}
+
+		// Generate a deterministic hash based on the original material hash and this primvar UVIndex assignment
+		CompatiblePrimvarAndUVIndexPairs.Sort(
+			[](const TPair<FString, int32>& LHS, const TPair<FString, int32>& RHS)
+			{
+				if (LHS.Key == RHS.Key)
+				{
+					return LHS.Value < RHS.Value;
+				}
+				else
+				{
+					return LHS.Key < RHS.Key;
+				}
+			}
+		);
+		FSHAHash Hash;
+		FSHA1 SHA1;
+		const FString OriginalNodeUid = OriginalFactoryNode->GetUniqueID();
+		SHA1.UpdateWithString(*OriginalNodeUid, OriginalNodeUid.Len());
+		for (const TPair<FString, int32>& Pair : CompatiblePrimvarAndUVIndexPairs)
+		{
+			SHA1.UpdateWithString(*Pair.Key, Pair.Key.Len());
+			SHA1.Update((const uint8*)&Pair.Value, sizeof(Pair.Value));
+		}
+		SHA1.Final();
+		SHA1.GetHash(&Hash.Hash[0]);
+
+		const FString CompatibleNodeUid = OriginalNodeUid + UE::Interchange::USD::CompatibleMaterialUidSuffix + Hash.ToString();
+
+		// Check if we made this compatible material before
+		UInterchangeMaterialInstanceFactoryNode* CompatibleFactoryNode = Cast<UInterchangeMaterialInstanceFactoryNode>(
+			InNodeContainer->GetFactoryNode(CompatibleNodeUid)
+		);
+		if (CompatibleFactoryNode)
+		{
+			return CompatibleFactoryNode;
+		}
+
+		// We need to actually create a new compatible material
+		CompatibleFactoryNode = NewObject<UInterchangeMaterialInstanceFactoryNode>(InNodeContainer);
+
+		UObject* Object = nullptr;
+		CompatibleFactoryNode->CopyWithObject(OriginalFactoryNode, Object);
+
+		CompatibleFactoryNode->InitializeNode(
+			CompatibleNodeUid,
+			OriginalFactoryNode->GetDisplayLabel() + TEXT("_Compatible"),
+			OriginalFactoryNode->GetNodeContainerType()
+		);
+
+		InNodeContainer->AddNode(CompatibleFactoryNode);
+
+		// Cleanup our old primvar to UV Index mapping in case we don't overlap perfectly
+		for (const TPair<FString, int32>& OldPair : OriginalMaterialPrimvarInfo->PrimvarToUVIndex)
+		{
+			const FString& Primvar = OldPair.Key;
+			const FString AttributeKeyString = USD::PrimvarUVIndexAttributePrefix + Primvar;
+
+			CompatibleFactoryNode->RemoveAttribute(AttributeKeyString);
+		}
+
+		// We also need to move our pairs into an actual TMap now, as we'll need to query into it later
+		TMap<FString, int32> CompatiblePrimvarToUVIndex;
+		CompatiblePrimvarToUVIndex.Reserve(CompatiblePrimvarAndUVIndexPairs.Num());
+
+		// Set our new primvar to UV index mapping on the factory node (only for the sake of information / debugging)
+		for (const TPair<FString, int32>& NewPair : CompatiblePrimvarAndUVIndexPairs)
+		{
+			const FString& Primvar = NewPair.Key;
+			const int32 UVIndex = NewPair.Value;
+			const FString AttributeKeyString = USD::PrimvarUVIndexAttributePrefix + Primvar;
+
+			CompatibleFactoryNode->AddInt32Attribute(AttributeKeyString, UVIndex);
+
+			CompatiblePrimvarToUVIndex.Add(Primvar, UVIndex);
+		}
+
+		// Set our compatible UV indices into the actual material parameter attributes that will be set
+		// on the material instance
+		for (const TPair<FString, FString>& ParameterPair : OriginalMaterialPrimvarInfo->ParameterToPrimvar)
+		{
+			const FString& MaterialParameter = ParameterPair.Key;
+			const FString& Primvar = ParameterPair.Value;
+
+			int32 UVIndex = UNUSED_UV_INDEX;
+			if (int32* FoundUVIndex = CompatiblePrimvarToUVIndex.Find(Primvar))
+			{
+				UVIndex = *FoundUVIndex;
+			}
+
+			const bool bEnableTexture = UVIndex >= 0 && UVIndex < UNUSED_UV_INDEX;
+			const FString EnableKey = UInterchangeShaderPortsAPI::MakeInputValueKey(
+				USD::UseTextureParameterPrefix + MaterialParameter + USD::UseTextureParameterSuffix
+			);
+			ensure(CompatibleFactoryNode->AddFloatAttribute(EnableKey, bEnableTexture));
+
+			const FString UVIndexKey = UInterchangeShaderPortsAPI::MakeInputValueKey(MaterialParameter + USD::UVIndexParameterSuffix);
+			if (bEnableTexture)
+			{
+				ensure(CompatibleFactoryNode->AddFloatAttribute(UVIndexKey, UVIndex));
+			}
+			else
+			{
+				ensure(CompatibleFactoryNode->RemoveAttribute(UVIndexKey));
+			}
+		}
+
+		return CompatibleFactoryNode;
+	}
+
+	// Makes sure that all materials in the MeshNode's SlotDependencies array are primvar-compatible
+	// with the primvar to UV index mapping of the mesh, creating new materials and reassigning them to the
+	// SlotDependencies custom attribute if needed.
+	void GeneratePrimvarCompatibleMaterials(
+		const UInterchangeMeshNode* InMeshNode,
+		UInterchangeMeshFactoryNode* InMeshFactoryNode,
+		UInterchangeBaseNodeContainer* InNodeContainer,
+		TMap<FString, FMaterialPrimvarInfo>& InOutFactoryNodeUidToInfo
+	)
+	{
+		using namespace UE::Interchange;
+
+		// Rebuild the mesh's primvar to UV index map
+		TMap<FString, int32> MeshPrimvarToUVIndex;
+		{
+			TArray<UE::Interchange::FAttributeKey> MeshAttributeKeys;
+			InMeshNode->GetAttributeKeys(MeshAttributeKeys);
+
+			MeshPrimvarToUVIndex.Reserve(MeshAttributeKeys.Num());
+
+			for (const UE::Interchange::FAttributeKey& AttributeKey : MeshAttributeKeys)
+			{
+				const FString& AttributeKeyString = AttributeKey.ToString();
+
+				if (AttributeKeyString.StartsWith(USD::PrimvarUVIndexAttributePrefix))
+				{
+					FString ParameterName = AttributeKeyString;
+					ensure(ParameterName.RemoveFromStart(USD::PrimvarUVIndexAttributePrefix));
+
+					int32 UVIndex = -1;
+					if (InMeshNode->GetInt32Attribute(AttributeKeyString, UVIndex))
+					{
+						MeshPrimvarToUVIndex.Add(ParameterName, UVIndex);
 					}
 				}
 			}
+		}
+
+		// Build the reverse map too. This because we can still consider a mesh and material "compatible"
+		// if the material tries reading "st2" and the mesh doesn't have "st2" anywhere. We will just turn off
+		// usage of that particular texture on the material instance.
+		//
+		// For that case to still count as "compatible" though, then whatever UV index the material is trying to
+		// read "st2" from must not be used by any other primvar: Otherwise the mesh will still be providing
+		// some primvar data to the material, but the material will treat it as if it were "st2" when it is not
+		TArray<TSet<FString>> UVIndexToMeshPrimvars;
+		UVIndexToMeshPrimvars.SetNum(USD_PREVIEW_SURFACE_MAX_UV_SETS);
+		for (const TPair<FString, int32>& MeshPair : MeshPrimvarToUVIndex)
+		{
+			if (UVIndexToMeshPrimvars.IsValidIndex(MeshPair.Value))
+			{
+				UVIndexToMeshPrimvars[MeshPair.Value].Add(MeshPair.Key);
+			}
+		}
+
+		// Use primvar to UV index information to construct primvar-compatible materials if needed
+		TMap<FString, FString> SlotNameToMaterialFactoryNodeUid;
+		InMeshFactoryNode->GetSlotMaterialDependencies(SlotNameToMaterialFactoryNodeUid);
+		for (const TPair<FString, FString>& SlotPair : SlotNameToMaterialFactoryNodeUid)
+		{
+			const FString& SlotName = SlotPair.Key;
+			const FString& MaterialFactoryNodeUid = SlotPair.Value;
+
+			UInterchangeMaterialInstanceFactoryNode* OriginalMaterial = Cast<UInterchangeMaterialInstanceFactoryNode>(
+				InNodeContainer->GetFactoryNode(MaterialFactoryNodeUid)
+			);
+			if (!OriginalMaterial)
+			{
+				continue;
+			}
+
+			FMaterialPrimvarInfo* OriginalMaterialInfo = GetOrCreateMaterialPrimvarInfo(OriginalMaterial, InOutFactoryNodeUidToInfo);
+			if (!OriginalMaterialInfo)
+			{
+				continue;
+			}
+
+			// Check if the material's primvar-UVIndex mapping matches the mesh
+			bool bCompatible = true;
+			for (const TPair<FString, int32>& Pair : OriginalMaterialInfo->PrimvarToUVIndex)
+			{
+				const FString& MaterialPrimvar = Pair.Key;
+				int32 MaterialUVIndex = Pair.Value;
+
+				// If the mesh has the same primvar the material wants, it should be at the same UVIndex the material
+				// will read from
+				if (const int32* MeshUVIndex = MeshPrimvarToUVIndex.Find(MaterialPrimvar))
+				{
+					if (*MeshUVIndex != MaterialUVIndex)
+					{
+						bCompatible = false;
+						break;
+					}
+				}
+
+				// If the material is going to read from a given UVIndex that exists on the mesh, that UV set should
+				// contain the primvar data that the material expects to read
+				if (UVIndexToMeshPrimvars.IsValidIndex(MaterialUVIndex))
+				{
+					const TSet<FString>& CompatiblePrimvars = UVIndexToMeshPrimvars[MaterialUVIndex];
+					if (!CompatiblePrimvars.Contains(MaterialPrimvar))
+					{
+						bCompatible = false;
+						break;
+					}
+				}
+			}
+
+			if (!bCompatible)
+			{
+				// Generate a primvar-compatible material
+				UInterchangeMaterialInstanceFactoryNode* CompatibleMaterial = GetOrCreateCompatibleMaterial(	//
+					OriginalMaterial,
+					OriginalMaterialInfo,
+					MeshPrimvarToUVIndex,
+					InNodeContainer
+				);
+				if (CompatibleMaterial)
+				{
+					InMeshFactoryNode->SetSlotMaterialDependencyUid(SlotName, CompatibleMaterial->GetUniqueID());
+				}
+			}
+		}
+	}
+
+	void ProcessMeshNodes(UInterchangeBaseNodeContainer* NodeContainer, bool bGenerateCompatibleMaterials)
+	{
+		using namespace UE::Interchange;
+
+		// Collect the nodes in a separate container because GeneratePrimvarCompatibleMaterials may add additional nodes
+		// to the node container itself, which we can't do while we iterate over it
+		TMap<const UInterchangeMeshNode*, UInterchangeMeshFactoryNode*> MeshNodeToFactoryNode;
+
+		NodeContainer->IterateNodesOfType<UInterchangeMeshNode>(
+			[NodeContainer, &MeshNodeToFactoryNode](const FString& NodeUid, UInterchangeMeshNode* MeshNode)
+			{
+				const FString FactoryNodeUid = UInterchangeFactoryBaseNode::BuildFactoryNodeUid(NodeUid);
+				UInterchangeMeshFactoryNode* MeshFactoryNode = Cast<UInterchangeMeshFactoryNode>(NodeContainer->GetFactoryNode(FactoryNodeUid));
+				if (!MeshFactoryNode)
+				{
+					return;
+				}
+
+				MeshNodeToFactoryNode.Add(MeshNode, MeshFactoryNode);
+
+#if WITH_EDITORONLY_DATA
+				AddPrimvarAttributesToFactoryNodes(MeshNode, MeshFactoryNode, NodeContainer);
+#endif	  // WITH_EDITORONLY_DATA
+			}
 		);
+
+		if (bGenerateCompatibleMaterials)
+		{
+			// Build this map to prevent us from reparsing the same material in case multiple meshes
+			// are using them
+			TMap<FString, FMaterialPrimvarInfo> MaterialFactoryNodeUidToInfo;
+			for (const TPair<const UInterchangeMeshNode*, UInterchangeMeshFactoryNode*>& MeshPair : MeshNodeToFactoryNode)
+			{
+				GeneratePrimvarCompatibleMaterials(MeshPair.Key, MeshPair.Value, NodeContainer, MaterialFactoryNodeUidToInfo);
+			}
+		}
 	}
 
+#if WITH_EDITORONLY_DATA
 	class FUsdPrimvarBaker
 	{
 	public:
@@ -1035,6 +1423,7 @@ namespace UE::InterchangeUsdPipeline::Private
 UInterchangeUsdPipeline::UInterchangeUsdPipeline()
 	: GeometryPurpose((int32)(EUsdPurpose::Default | EUsdPurpose::Proxy | EUsdPurpose::Render | EUsdPurpose::Guide))
 	, bImportPseudoRoot(true)
+	, bGeneratePrimvarCompatibleMaterials(true)
 {
 }
 
@@ -1093,9 +1482,7 @@ void UInterchangeUsdPipeline::ExecutePipeline(
 		RemovePseudoRootFactoryNode(*NodeContainer);
 	}
 
-#if WITH_EDITORONLY_DATA
-	AddPrimvarAttributesToFactoryNodes(NodeContainer);
-#endif	  // WITH_EDITORONLY_DATA
+	ProcessMeshNodes(NodeContainer, bGeneratePrimvarCompatibleMaterials);
 #endif	  // USE_USD_SDK
 }
 
diff --git a/Engine/Plugins/Interchange/Runtime/Source/Pipelines/Public/InterchangeUSDPipeline.h b/Engine/Plugins/Interchange/Runtime/Source/Pipelines/Public/InterchangeUSDPipeline.h
--- a/Engine/Plugins/Interchange/Runtime/Source/Pipelines/Public/InterchangeUSDPipeline.h
+++ b/Engine/Plugins/Interchange/Runtime/Source/Pipelines/Public/InterchangeUSDPipeline.h
@@ -36,6 +36,16 @@ public:
 	UPROPERTY(EditAnywhere, BlueprintReadWrite, Category = "USD")
 	bool bImportPseudoRoot;
 
+	/**
+	 * Imported meshes will place their primvars on arbitrary UV set indices (e.g. "st2" on UV0, "st3" on UV1, etc.). Imported materials may be
+	 * expecting to read from specific primvars on particular UV indices (e.g. "st1" from UV0, "st2" from UV2, etc.). These mappings won't
+	 * necessarily correspond to each other. Setting this option to true means we'll generate additional material instances with primvar-UV index
+	 * mappings made to best match the exact meshes they are assigned to (e.g. "st1" disabled, "st2" from UV0). Setting this to false means we
+	 * will just assign the incompatible materials to the meshes anyway.
+	 */
+	UPROPERTY(EditAnywhere, BlueprintReadWrite, Category = "USD")
+	bool bGeneratePrimvarCompatibleMaterials;
+
 public:
 
 	virtual void ExecutePostFactoryPipeline(const UInterchangeBaseNodeContainer* BaseNodeContainer, const FString& NodeKey, UObject* CreatedAsset, bool bIsAReimport) override;
diff --git a/Engine/Plugins/Runtime/USDCore/Source/USDUtilities/Private/USDConversionUtils.cpp b/Engine/Plugins/Runtime/USDCore/Source/USDUtilities/Private/USDConversionUtils.cpp
--- a/Engine/Plugins/Runtime/USDCore/Source/USDUtilities/Private/USDConversionUtils.cpp
+++ b/Engine/Plugins/Runtime/USDCore/Source/USDUtilities/Private/USDConversionUtils.cpp
@@ -1028,6 +1028,12 @@ TMap<FString, int32> UsdUtils::CombinePrimvarsIntoUVSets(const TSet<FString>& Al
 	return PrimvarToUVIndex;
 }
 
+TMap<FString, int32> UsdUtils::GetPrimvarToUVIndexMap(const pxr::UsdPrim& UsdPrim, int32 MaxNumPrimvars)
+{
+	TArray<TUsdStore<pxr::UsdGeomPrimvar>> PrimvarsToUse = UsdUtils::GetUVSetPrimvars(UsdPrim, MaxNumPrimvars);
+	return UsdUtils::AssemblePrimvarsIntoPrimvarToUVIndexMap(PrimvarsToUse);
+}
+
 bool UsdUtils::IsAnimated(const pxr::UsdPrim& Prim)
 {
 	if (!Prim || !Prim.IsActive())
diff --git a/Engine/Plugins/Runtime/USDCore/Source/USDUtilities/Public/USDConversionUtils.h b/Engine/Plugins/Runtime/USDCore/Source/USDUtilities/Public/USDConversionUtils.h
--- a/Engine/Plugins/Runtime/USDCore/Source/USDUtilities/Public/USDConversionUtils.h
+++ b/Engine/Plugins/Runtime/USDCore/Source/USDUtilities/Public/USDConversionUtils.h
@@ -257,6 +257,15 @@ namespace UsdUtils
 	 */
 	USDUTILITIES_API TMap<FString, int32> CombinePrimvarsIntoUVSets(const TSet<FString>& AllPrimvars, const TSet<FString>& PreferredPrimvars);
 
+	/**
+	 * Calls GetUVSetPrimvars() followed by AssemblePrimvarsIntoPrimvarToUVIndexMap().
+	 * This exists so that it can be called without having to manipulate the USD types (e.g. for Interchange)
+	 */
+	USDUTILITIES_API TMap<FString, int32> GetPrimvarToUVIndexMap(
+		const pxr::UsdPrim& UsdPrim,
+		int32 MaxNumPrimvars = USD_PREVIEW_SURFACE_MAX_UV_SETS
+	);
+
 	template<typename T, typename U>
 	FString StringifyMap(const TMap<T, U>& Map)
 	{
diff --git a/Engine/Plugins/Runtime/USDCore/Source/UnrealUSDWrapper/Public/UnrealUSDWrapper.h b/Engine/Plugins/Runtime/USDCore/Source/UnrealUSDWrapper/Public/UnrealUSDWrapper.h
--- a/Engine/Plugins/Runtime/USDCore/Source/UnrealUSDWrapper/Public/UnrealUSDWrapper.h
+++ b/Engine/Plugins/Runtime/USDCore/Source/UnrealUSDWrapper/Public/UnrealUSDWrapper.h
@@ -29,6 +29,7 @@
 // We are capped to 4 because the geometry cache and skeletal mesh + morph target shaders have a hard-coded limit of 4 texture coordinates, so our
 // material instances can't cross that limit or else they wouldn't compile when assigned to those meshes
 #define USD_PREVIEW_SURFACE_MAX_UV_SETS 4
+#define UNUSED_UV_INDEX USD_PREVIEW_SURFACE_MAX_UV_SETS
 
 #if USE_USD_SDK
 PXR_NAMESPACE_OPEN_SCOPE
diff --git a/Engine/Source/Runtime/Interchange/Core/Public/Types/AttributeStorage.h b/Engine/Source/Runtime/Interchange/Core/Public/Types/AttributeStorage.h
--- a/Engine/Source/Runtime/Interchange/Core/Public/Types/AttributeStorage.h
+++ b/Engine/Source/Runtime/Interchange/Core/Public/Types/AttributeStorage.h
@@ -179,7 +179,7 @@ namespace UE
 				return Ar;
 			}
 
-			FORCEINLINE FString ToString() const
+			FORCEINLINE const FString& ToString() const
 			{
 				return Key;
 			}
