[BUGFIX] Interchange Metadata
- Fixed the changes to the metadata extraction in generic assets pipeline.

#rnx

diff --git a/Engine/Plugins/Interchange/Runtime/Source/Pipelines/Private/InterchangeGenericAssetsPipeline.cpp b/Engine/Plugins/Interchange/Runtime/Source/Pipelines/Private/InterchangeGenericAssetsPipeline.cpp
--- a/Engine/Plugins/Interchange/Runtime/Source/Pipelines/Private/InterchangeGenericAssetsPipeline.cpp
+++ b/Engine/Plugins/Interchange/Runtime/Source/Pipelines/Private/InterchangeGenericAssetsPipeline.cpp
@@ -1312,15 +1312,15 @@ void UInterchangeGenericAssetsPipeline::AddMetaData(UObject* CreatedAsset, const
 
 	const FString InterchangeMetaDataPrefix = TEXT("INTERCHANGE.");
 
-	bool bAssetIsActor = CreatedAsset->IsA<AActor>();
+	const bool bAssetIsActor = CreatedAsset->IsA<AActor>();
 #if WITH_EDITORONLY_DATA
 	FMetaData& MetaData = CreatedAsset->GetPackage()->GetMetaData();
 #endif
 
-	bool bProcessMetaData = bAssetIsActor;
-
 	TMap<FString, FString> MetaDataMap;
-	if (bProcessMetaData)
+#ifndef WITH_EDITORONLY_DATA
+	if (bAssetIsActor)
+#endif
 	{
 		TArray<FInterchangeUserDefinedAttributeInfo> UserAttributeInfos;
 		UInterchangeUserDefinedAttributesAPI::GetUserDefinedAttributeInfos(Node, UserAttributeInfos);
