[BUGFIX] FBX import crash
- Negative Material Index would be forced to zero mirroring Legacy FBX Import.

#rnx

diff --git a/Engine/Plugins/Interchange/Runtime/Source/Parsers/Fbx/Private/FbxMesh.cpp b/Engine/Plugins/Interchange/Runtime/Source/Parsers/Fbx/Private/FbxMesh.cpp
--- a/Engine/Plugins/Interchange/Runtime/Source/Parsers/Fbx/Private/FbxMesh.cpp
+++ b/Engine/Plugins/Interchange/Runtime/Source/Parsers/Fbx/Private/FbxMesh.cpp
@@ -1171,20 +1171,28 @@ bool FMeshDescriptionImporter::FillMeshDescriptionFromFbxMesh(FbxMesh* Mesh, con
 				{
 					if (LayerElementMaterial)
 					{
+						int32 MaterialRemapIndex = 0;
 						switch (MaterialMappingMode)
 						{
 							// material index is stored in the IndexArray, not the DirectArray (which is irrelevant with 2009.1)
 							case FbxLayerElement::eAllSame:
 							{
-								MaterialIndex = MaterialRemap[LayerElementMaterial->GetIndexArray().GetAt(0)];
+								MaterialRemapIndex = LayerElementMaterial->GetIndexArray().GetAt(0);
 							}
 							break;
 							case FbxLayerElement::eByPolygon:
 							{
-								MaterialIndex = MaterialRemap[LayerElementMaterial->GetIndexArray().GetAt(PolygonIndex)];
+								MaterialRemapIndex = LayerElementMaterial->GetIndexArray().GetAt(PolygonIndex);
 							}
 							break;
 						}
+
+						// Follows same behaviour as Legacy Import, i.e. forces index to zero if negative.
+						MaterialIndex = 0;
+						if (MaterialRemap.IsValidIndex(MaterialRemapIndex))
+						{
+							MaterialIndex = MaterialRemap[MaterialRemapIndex];
+						}
 					}
 				}
 
