Interchange: Fixed Linux linking errors due to missing MaterialX symbols.

#rnx

diff --git a/Engine/Plugins/Interchange/Runtime/Source/Import/InterchangeImport.Build.cs b/Engine/Plugins/Interchange/Runtime/Source/Import/InterchangeImport.Build.cs
--- a/Engine/Plugins/Interchange/Runtime/Source/Import/InterchangeImport.Build.cs
+++ b/Engine/Plugins/Interchange/Runtime/Source/Import/InterchangeImport.Build.cs
@@ -59,21 +59,12 @@ namespace UnrealBuildTool.Rules
 						"BSPUtils",
 						"InterchangeFbxParser",
 						"MaterialEditor",
+						"MaterialX",
 						"SkeletalMeshUtilitiesCommon",
 						"UnrealEd",
 						"VariantManager",
 					}
 				);
-
-				if (Target.Platform == UnrealTargetPlatform.Win64 || Target.Platform == UnrealTargetPlatform.Mac || Target.LinkType == TargetLinkType.Monolithic)
-				{
-					PrivateDependencyModuleNames.AddRange(
-						new string[]
-						{
-							"MaterialX"
-						}
-					);
-				}
 			}
 		}
 	}
