USD: Fix failing to import volume animations where the same volume is used for every frame.

diff --git a/Engine/Plugins/Interchange/Runtime/Source/Import/Private/Volume/InterchangeSparseVolumeTextureFactory.cpp b/Engine/Plugins/Interchange/Runtime/Source/Import/Private/Volume/InterchangeSparseVolumeTextureFactory.cpp
--- a/Engine/Plugins/Interchange/Runtime/Source/Import/Private/Volume/InterchangeSparseVolumeTextureFactory.cpp
+++ b/Engine/Plugins/Interchange/Runtime/Source/Import/Private/Volume/InterchangeSparseVolumeTextureFactory.cpp
@@ -580,16 +580,48 @@ UInterchangeFactoryBase::FImportAssetResult UInterchangeSparseVolumeTextureFacto
 		// same class, name and outer, which resets the SVT. We'll do this here too, we just have to be careful to
 		// pick the same class as the existing asset if we have one (during reimports): We will crash if we try
 		// overwriting an animated SVT asset with a static one, for example.
-		UClass* ClassToUse = ExistingAsset			  ? ExistingAsset->GetClass()
-							 : VolumeNodes.Num() == 1 ? UStaticSparseVolumeTexture::StaticClass()
-													  : UAnimatedSparseVolumeTexture::StaticClass();
+		UClass* ClassToUse = ExistingAsset ? ExistingAsset->GetClass() : nullptr;
+		if (!ClassToUse)
+		{
+			// Check the FactoryNode's attribute to decide this, as that is what is affected by the pipeline settings
+			FString AnimationID;
+			if (!FactoryNode->GetCustomAnimationID(AnimationID) || AnimationID.IsEmpty())
+			{
+				ClassToUse = UStaticSparseVolumeTexture::StaticClass();
+			}
+			else
+			{
+				ClassToUse = UAnimatedSparseVolumeTexture::StaticClass();
+			}
+		}
 
 		// If we have a mismatch we should return now before we stomp the existing asset for no reason.
 		// Note that it's fine to try and stuff a single frame payload into an animated SVT: It will just be a one-frame-animation
-		if (ClassToUse == UStaticSparseVolumeTexture::StaticClass() && VolumeNodes.Num() != 1)
+		if (ExistingAsset && ClassToUse == UStaticSparseVolumeTexture::StaticClass())
 		{
-			HandleFailure(LOCTEXT("SparseVolumeTextureFactory_Mismatch", "Cannot import multiple volume frames into a StaticSparseVolumeTexture."));
-			return ImportAssetResult;
+			// Check whether our current volume nodes actually describe an animation or not
+			bool bHasAnimation = VolumeNodes.Num() > 1;
+			if (!bHasAnimation)
+			{
+				int32 NumFrames = 0;
+				for (const UInterchangeVolumeNode* VolumeNode : VolumeNodes)
+				{
+					TArray<int32> FrameIndicesForVolume;
+					VolumeNode->GetCustomFrameIndicesInAnimation(FrameIndicesForVolume);
+
+					NumFrames += FrameIndicesForVolume.Num();
+				}
+				if (NumFrames > 1)
+				{
+					bHasAnimation = true;
+				}
+			}
+
+			if (bHasAnimation)
+			{
+				HandleFailure(LOCTEXT("SparseVolumeTextureFactory_Mismatch", "Cannot import a volume animation into a StaticSparseVolumeTexture."));
+				return ImportAssetResult;
+			}
 		}
 
 #if WITH_EDITOR
@@ -766,7 +798,7 @@ UInterchangeFactoryBase::FImportAssetResult UInterchangeSparseVolumeTextureFacto
 
 	bool bAtLeastOneValidPayloadKey = false;
 
-	TSet<int32> FrameIndices;
+	int32 NumFrames = 0;
 
 	TArray<FVolumePayloadKey> PayloadKeys;
 	PayloadKeys.Reserve(VolumeNodes.Num());
@@ -801,11 +833,11 @@ UInterchangeFactoryBase::FImportAssetResult UInterchangeSparseVolumeTextureFacto
 		{
 			// No animated frames --> Volume node wants to become a static SVT. Let's consider it as one frame,
 			// as that is how we'll process its payload later in FillSparseVolumeTextureWithPayloadData()
-			FrameIndices.Add(TNumericLimits<int32>::Lowest());
+			NumFrames += 1;
 		}
 		else
 		{
-			FrameIndices.Append(Payload.PayloadFrameIndices);
+			NumFrames += Payload.PayloadFrameIndices.Num();
 		}
 	}
 	if (!bAtLeastOneValidPayloadKey)
@@ -849,11 +881,11 @@ UInterchangeFactoryBase::FImportAssetResult UInterchangeSparseVolumeTextureFacto
 	bool bCanSetup = bAtLeastOneValidPayloadData;
 	if (UStaticSparseVolumeTexture* StaticVolumeTexture = Cast<UStaticSparseVolumeTexture>(SparseVolumeTexture))
 	{
-		bCanSetup = bCanSetup && FrameIndices.Num() == 1;
+		bCanSetup = bCanSetup && NumFrames == 1;
 	}
 	else if (UAnimatedSparseVolumeTexture* AnimatedVolumeTexture = Cast<UAnimatedSparseVolumeTexture>(SparseVolumeTexture))
 	{
-		bCanSetup = bCanSetup && FrameIndices.Num() >= 1;
+		bCanSetup = bCanSetup && NumFrames >= 1;
 	}
 	if (!bCanSetup)
 	{
diff --git a/Engine/Plugins/Interchange/Runtime/Source/Pipelines/Private/InterchangeSparseVolumeTexturePipeline.cpp b/Engine/Plugins/Interchange/Runtime/Source/Pipelines/Private/InterchangeSparseVolumeTexturePipeline.cpp
--- a/Engine/Plugins/Interchange/Runtime/Source/Pipelines/Private/InterchangeSparseVolumeTexturePipeline.cpp
+++ b/Engine/Plugins/Interchange/Runtime/Source/Pipelines/Private/InterchangeSparseVolumeTexturePipeline.cpp
@@ -470,7 +470,7 @@ void UInterchangeSparseVolumeTexturePipeline::ExecutePipeline(
 		const bool bAddSourceNodeName = false;
 		UInterchangeUserDefinedAttributesAPI::DuplicateAllUserDefinedAttribute(FirstVolume, FactoryNode, bAddSourceNodeName);
 
-		// Just record this for information's sake, as we won't really use this for anything after this point
+		// Providing the animationID is required to have the factory treat this node as an actual volume animation
 		FactoryNode->SetCustomAnimationID(AnimationID);
 
 		TSet<const UInterchangeVolumeNode*> AddedNodes;
