[Interchange] Animation curve type expectation missmatch fixing.

#rnx

diff --git a/Engine/Plugins/Interchange/Runtime/Source/Pipelines/Private/InterchangeGenericAnimationPipeline.cpp b/Engine/Plugins/Interchange/Runtime/Source/Pipelines/Private/InterchangeGenericAnimationPipeline.cpp
--- a/Engine/Plugins/Interchange/Runtime/Source/Pipelines/Private/InterchangeGenericAnimationPipeline.cpp
+++ b/Engine/Plugins/Interchange/Runtime/Source/Pipelines/Private/InterchangeGenericAnimationPipeline.cpp
@@ -840,24 +840,19 @@ void UInterchangeGenericAnimationPipeline::CreateAnimSequenceFactoryNode(UInterc
 
 							const bool bForceImportBoneCustomAttribute = CustomAttributeNamesToImport.Contains(AttributeInfo.Name);
 							
-							bool bAttributeProcessed = false;
 							EInterchangeAnimationPayLoadType AnimationPayloadType = EInterchangeAnimationPayLoadType::NONE;
-							if (!bForceImportBoneCustomAttribute && !bImportAllAttributesOnBone
-								&& SceneNode->GetAnimationCurveTypeForCurveName(AttributeInfo.Name, AnimationPayloadType))
+							if (SceneNode->GetAnimationCurveTypeForCurveName(AttributeInfo.Name, AnimationPayloadType))
 							{
 								if (AnimationPayloadType == EInterchangeAnimationPayLoadType::CURVE)
 								{
 									AnimSequenceFactoryNode->SetAnimatedAttributeCurveName(AttributeInfo.Name);
-									bAttributeProcessed = true;
 								}
 								else if (AnimationPayloadType == EInterchangeAnimationPayLoadType::STEPCURVE)
 								{
 									AnimSequenceFactoryNode->SetAnimatedAttributeStepCurveName(AttributeInfo.Name);
-									bAttributeProcessed = true;
 								}
 							}
-
-							if (!bAttributeProcessed)
+							else
 							{
 								//Material attribute curve
 								if (!bImportAllAttributesOnBone && bDecimalType && !bForceImportBoneCustomAttribute)
