[Interchange] Optimize material search by using the asset registry to scan root paths only once instead of at every call. The UI is now responsive even if the value require to refresh the import dialog.

#rnx

diff --git a/Engine/Plugins/Interchange/Runtime/Source/Pipelines/Private/InterchangeGenericMaterialPipeline.cpp b/Engine/Plugins/Interchange/Runtime/Source/Pipelines/Private/InterchangeGenericMaterialPipeline.cpp
--- a/Engine/Plugins/Interchange/Runtime/Source/Pipelines/Private/InterchangeGenericMaterialPipeline.cpp
+++ b/Engine/Plugins/Interchange/Runtime/Source/Pipelines/Private/InterchangeGenericMaterialPipeline.cpp
@@ -219,11 +219,15 @@ namespace UE::Interchange::InterchangeGenericMaterialPipeline::Private
 		FAssetRegistryModule& AssetRegistryModule = FModuleManager::LoadModuleChecked<FAssetRegistryModule>("AssetRegistry");
 		IAssetRegistry& AssetRegistry = AssetRegistryModule.Get();
 
+		//Make sure we scan the root paths only one time per editor session.
+		//There can be thousands of root path and it is a slow task even if everything is already scan.
+		static bool bRootPathScanned = false;
 		// Finish/update any scans
 		TArray<FString> ScanPaths;
-		if (BasePath.IsEmpty() || BasePath == TEXT("/"))
+		if (!bRootPathScanned && (BasePath.IsEmpty() || BasePath == TEXT("/")))
 		{
 			FPackageName::QueryRootContentPaths(ScanPaths);
+			bRootPathScanned = true;
 		}
 		else if (!BasePath.StartsWith(TEXT("/Temp"), ESearchCase::IgnoreCase)) //We must exclude Temp path to avoid asset registry scan path warnings
 		{
