Interchange MaterialX:
- Add support of NPR lib

#virtualized

diff --git a/Engine/Plugins/Interchange/Runtime/Source/Common/Public/MaterialX/InterchangeMaterialXDefinitions.h b/Engine/Plugins/Interchange/Runtime/Source/Common/Public/MaterialX/InterchangeMaterialXDefinitions.h
--- a/Engine/Plugins/Interchange/Runtime/Source/Common/Public/MaterialX/InterchangeMaterialXDefinitions.h
+++ b/Engine/Plugins/Interchange/Runtime/Source/Common/Public/MaterialX/InterchangeMaterialXDefinitions.h
@@ -146,6 +146,7 @@ namespace UE
 			static constexpr const TCHAR* Arcsine = TEXT("Arcsine");
 			static constexpr const TCHAR* Arctangent2 = TEXT("Arctangent2");
 			static constexpr const TCHAR* BlackBody = TEXT("BlackBody");
+			static constexpr const TCHAR* CameraVectorWS = TEXT("CameraVectorWS");
 			static constexpr const TCHAR* Ceil = TEXT("Ceil");
 			static constexpr const TCHAR* Clamp = TEXT("Clamp");
 			static constexpr const TCHAR* Contrast = TEXT("MaterialXContrast");
@@ -292,6 +293,7 @@ namespace UE
 			constexpr const TCHAR* MxAdobeRGBToRec709 = TEXT("/InterchangeAssets/Functions/MX_AdobeRGBToRec709.MX_AdobeRGBToRec709");
 			constexpr const TCHAR* MxArtisticIOR = TEXT("/InterchangeAssets/Functions/MX_Artistic_IOR.MX_Artistic_IOR");
 			constexpr const TCHAR* MxColorCorrect = TEXT("/InterchangeAssets/Functions/MX_ColorCorrect.MX_ColorCorrect");
+			constexpr const TCHAR* MxGoochShade = TEXT("/InterchangeAssets/Functions/MX_GoochShade.MX_GoochShade");
 			constexpr const TCHAR* MxHsvAdjust = TEXT("/InterchangeAssets/Functions/MX_HSVAdjust.MX_HSVAdjust");
 			constexpr const TCHAR* MxLinearAdobeRGBToRec709 = TEXT("/InterchangeAssets/Functions/MX_LinearAdobeRGBToRec709.MX_LinearAdobeRGBToRec709");
 			constexpr const TCHAR* MxLinearDisplayP3ToRec709 = TEXT("/InterchangeAssets/Functions/MX_LinearDisplayP3ToRec709.MX_LinearDisplayP3ToRec709");
@@ -815,6 +817,10 @@ MATERIALX_NAMESPACE_BEGIN
 		static constexpr const char* Rec709DisplayToLinRec709 = "rec709_display_to_lin_rec709";
 		static constexpr const char* SrgbDisplayP3ToLinRec709 = "srgb_displayp3_to_lin_rec709";
 		static constexpr const char* SrgbTextureToLinRec709 = "srgb_texture_to_lin_rec709";
+		// NPR
+		static constexpr const char* FacingRatio = "facingratio";
+		static constexpr const char* GoochShade = "gooch_shade";
+		static constexpr const char* ViewDirection = "viewdirection";
 	}
 
 	namespace NodeDefinition
diff --git a/Engine/Plugins/Interchange/Runtime/Source/Import/Private/MaterialX/MaterialXUtils/MaterialXManager.cpp b/Engine/Plugins/Interchange/Runtime/Source/Import/Private/MaterialX/MaterialXUtils/MaterialXManager.cpp
--- a/Engine/Plugins/Interchange/Runtime/Source/Import/Private/MaterialX/MaterialXUtils/MaterialXManager.cpp
+++ b/Engine/Plugins/Interchange/Runtime/Source/Import/Private/MaterialX/MaterialXUtils/MaterialXManager.cpp
@@ -240,6 +240,8 @@ FMaterialXManager::FMaterialXManager()
 		// Adjustment
 		{mx::Category::ColorCorrect,	        FMaterialXMaterialFunction{TInPlaceType<FString>{}, UE::MaterialFunctions::Path::MxColorCorrect}},
 		{mx::Category::HsvAdjust,	            FMaterialXMaterialFunction{TInPlaceType<FString>{}, UE::MaterialFunctions::Path::MxHsvAdjust}},
+		//NPR
+		{mx::Category::GoochShade,	            FMaterialXMaterialFunction{TInPlaceType<FString>{}, UE::MaterialFunctions::Path::MxGoochShade}},
 	}
 	, CategoriesToSkip {
 		mx::Category::ACEScgToLinRec709,
@@ -249,6 +251,7 @@ FMaterialXManager::FMaterialXManager()
 		mx::Category::Distance,
 		mx::Category::Extract,
 		mx::Category::Fractal3D,
+		mx::Category::GoochShade,
 		mx::Category::HsvAdjust,
 		mx::Category::LinAdobeRgbToLinRec709,
 		mx::Category::LinDisplayP3ToLinRec709,
@@ -464,6 +467,7 @@ namespace UE::Interchange::MaterialX
 																UE::MaterialFunctions::Path::HeightToNormalSmooth,
 																UE::MaterialFunctions::Path::MxColorCorrect,
 																UE::MaterialFunctions::Path::MxHsvAdjust,
+																UE::MaterialFunctions::Path::MxGoochShade,
 															  });
 
 		return bPackagesLoaded;
diff --git a/Engine/Plugins/Interchange/Runtime/Source/Import/Private/MaterialX/MaterialXUtils/MaterialXSurfaceShaderAbstract.cpp b/Engine/Plugins/Interchange/Runtime/Source/Import/Private/MaterialX/MaterialXUtils/MaterialXSurfaceShaderAbstract.cpp
--- a/Engine/Plugins/Interchange/Runtime/Source/Import/Private/MaterialX/MaterialXUtils/MaterialXSurfaceShaderAbstract.cpp
+++ b/Engine/Plugins/Interchange/Runtime/Source/Import/Private/MaterialX/MaterialXUtils/MaterialXSurfaceShaderAbstract.cpp
@@ -1472,8 +1472,29 @@ void FMaterialXSurfaceShaderAbstract::ConnectRefractInputToOutput(const FConnect
 	AddAttributeFromValueOrInterface(InputIorTarget, UE::Expressions::Inputs::RefractiveIndexTarget, RefractNode);
 }
 
-void FMaterialXSurfaceShaderAbstract::ConnectBumpInputToOutput(const FConnectNode& Connect)
+void FMaterialXSurfaceShaderAbstract::ConnectViewDirectionInputToOutput(const FConnectNode& Connect)
 {
+	using namespace UE::Interchange::Materials::Standard::Nodes;
+
+	UInterchangeShaderNode* ViewDirectionNode = CreateShaderNode(Connect.UpstreamNode, Connect.UpstreamNode->getName().c_str(), UE::Expressions::Names::CameraVectorWS);
+	UInterchangeShaderNode* NegateNode = CreateShaderNode(Connect.UpstreamNode, (Connect.UpstreamNode->getName() + "_Negate").c_str(), UE::Expressions::Names::Subtract);
+	NegateNode->AddFloatAttribute(UInterchangeShaderPortsAPI::MakeInputValueKey(UE::Expressions::Inputs::A), 0.f);
+	UInterchangeShaderPortsAPI::ConnectDefaultOuputToInput(NegateNode, UE::Expressions::Inputs::B, ViewDirectionNode->GetUniqueID());
+	
+	UInterchangeShaderNode* NodeToConnectTo = NegateNode;
+	if (mx::InputPtr InputSpace = Connect.UpstreamNode->getInput("space"))
+	{
+		const bool bIsObjectSpace = InputSpace->getValueString() != "world";
+		if (bIsObjectSpace)
+		{
+			UInterchangeShaderNode* TransformNode = CreateShaderNode(Connect.UpstreamNode, Connect.UpstreamNode->getName().c_str() + FString(TEXT("_Transform")), TransformVector::Name.ToString());
+			TransformNode->AddInt32Attribute(TransformVector::Attributes::TransformSourceType.ToString(), EMaterialVectorCoordTransformSource::TRANSFORMSOURCE_World);
+			TransformNode->AddInt32Attribute(TransformVector::Attributes::TransformType.ToString(), EMaterialVectorCoordTransform::TRANSFORM_Local);
+			UInterchangeShaderPortsAPI::ConnectDefaultOuputToInput(TransformNode, TransformVector::Inputs::Input.ToString(), NodeToConnectTo->GetUniqueID());
+			NodeToConnectTo = TransformNode;
+		}
+	}
+	UInterchangeShaderPortsAPI::ConnectDefaultOuputToInput(Connect.ParentShaderNode, Connect.InputChannelName, NodeToConnectTo->GetUniqueID());
 }
 
 UInterchangeShaderNode* FMaterialXSurfaceShaderAbstract::CreateMaskShaderNode(uint8 RGBA, mx::ElementPtr Element, const FString& NodeName, const FString& OutputName)
@@ -1722,6 +1743,7 @@ void FMaterialXSurfaceShaderAbstract::RegisterConnectNodeOutputToInputDelegates(
 	MatchingConnectNodeDelegates.Add(mx::Category::Swizzle,			FOnConnectNodeOutputToInput::CreateSP(this, &FMaterialXSurfaceShaderAbstract::ConnectSwizzleInputToOutput));
 	MatchingConnectNodeDelegates.Add(mx::Category::Switch,			FOnConnectNodeOutputToInput::CreateSP(this, &FMaterialXSurfaceShaderAbstract::ConnectSwitchInputToOutput));
 	MatchingConnectNodeDelegates.Add(mx::Category::Refract,			FOnConnectNodeOutputToInput::CreateSP(this, &FMaterialXSurfaceShaderAbstract::ConnectRefractInputToOutput));
+	MatchingConnectNodeDelegates.Add(mx::Category::ViewDirection,	FOnConnectNodeOutputToInput::CreateSP(this, &FMaterialXSurfaceShaderAbstract::ConnectViewDirectionInputToOutput));
 }
 
 void FMaterialXSurfaceShaderAbstract::SetMatchingInputsNames(MaterialX::NodePtr Node) const
diff --git a/Engine/Plugins/Interchange/Runtime/Source/Import/Public/MaterialX/MaterialXUtils/MaterialXSurfaceShaderAbstract.h b/Engine/Plugins/Interchange/Runtime/Source/Import/Public/MaterialX/MaterialXUtils/MaterialXSurfaceShaderAbstract.h
--- a/Engine/Plugins/Interchange/Runtime/Source/Import/Public/MaterialX/MaterialXUtils/MaterialXSurfaceShaderAbstract.h
+++ b/Engine/Plugins/Interchange/Runtime/Source/Import/Public/MaterialX/MaterialXUtils/MaterialXSurfaceShaderAbstract.h
@@ -324,8 +324,8 @@ protected:
 	/** <refract> */
 	void ConnectRefractInputToOutput(const FConnectNode& Connect);
 
-	/** <bump> */
-	void ConnectBumpInputToOutput(const FConnectNode& Connect);
+	/** <viewdirection> */
+	void ConnectViewDirectionInputToOutput(const FConnectNode& Connect);
 
 	/** End Connect MaterialX nodes*/
 
