[Fbx] Disable AutoRTFM instrumentation.

AutoRTFM instrumentation cannot be enabled with exceptions.

diff --git a/Engine/Plugins/Interchange/Runtime/Source/Parsers/Fbx/InterchangeFbxParser.Build.cs b/Engine/Plugins/Interchange/Runtime/Source/Parsers/Fbx/InterchangeFbxParser.Build.cs
--- a/Engine/Plugins/Interchange/Runtime/Source/Parsers/Fbx/InterchangeFbxParser.Build.cs
+++ b/Engine/Plugins/Interchange/Runtime/Source/Parsers/Fbx/InterchangeFbxParser.Build.cs
@@ -22,6 +22,8 @@ namespace UnrealBuildTool.Rules
 
 			bEnableExceptions = true;
 
+			bDisableAutoRTFMInstrumentation = true; // AutoRTFM cannot be used with exceptions
+
 			PublicDependencyModuleNames.AddRange(
 				new string[]
 				{
