[Interchange] Allowing NonJoint interrupted Skeletons at import. Also updating gltf Parsers accordingly.

#rnx

diff --git a/Engine/Plugins/Interchange/Runtime/Source/Import/Private/Animation/InterchangeAnimSequenceFactory.cpp b/Engine/Plugins/Interchange/Runtime/Source/Import/Private/Animation/InterchangeAnimSequenceFactory.cpp
--- a/Engine/Plugins/Interchange/Runtime/Source/Import/Private/Animation/InterchangeAnimSequenceFactory.cpp
+++ b/Engine/Plugins/Interchange/Runtime/Source/Import/Private/Animation/InterchangeAnimSequenceFactory.cpp
@@ -676,6 +676,35 @@ namespace UE::Interchange::Private
 								RootJointNode->GetCustomGlobalTransform(NodeContainer, GlobalOffsetTransform, LocalTransform);
 							}
 						}
+
+						// On the unusal edge use case where Non-Joint Node(s) are inserted into a Skeleton hierarchy, bake said transforms onto the Joint.
+						if (!SkeletonNode->GetUniqueID().Equals(SkeletonRootUid))
+						{
+							TArray<FTransform> NonJointTransforms; //Child to Parent
+							const UInterchangeSceneNode* ParentNode = Cast<UInterchangeSceneNode>(NodeContainer->GetNode(SkeletonNode->GetParentUid()));
+							while (ParentNode && !ParentNode->IsSpecializedTypeContains(UE::Interchange::FSceneNodeStaticData::GetJointSpecializeTypeString()))
+							{
+								FTransform NonJointLocalTransform;
+								if (ParentNode->GetCustomLocalTransform(NonJointLocalTransform))
+								{
+									NonJointTransforms.Add(NonJointLocalTransform);
+								}
+
+								ParentNode = Cast<UInterchangeSceneNode>(NodeContainer->GetNode(ParentNode->GetParentUid()));
+							}
+							
+							if (NonJointTransforms.Num() > 0)
+							{
+								FTransform NonJointTransformToApply;
+								for (int32 NonJointCounter = NonJointTransforms.Num() - 1; NonJointCounter >= 0; NonJointCounter--)
+								{
+									NonJointTransformToApply *= NonJointTransforms[NonJointCounter];
+								}
+
+								LocalTransform = NonJointTransformToApply * LocalTransform;
+							}
+						}
+
 						FTransform3f AnimKeyTransform = static_cast<FTransform3f>(LocalTransform);
 						const FName BoneName = FName(*(SkeletonNode->GetDisplayLabel()));
 						//Add only one transform key at time 0 since this node is not animated
diff --git a/Engine/Plugins/Interchange/Runtime/Source/Import/Private/Gltf/InterchangeGltfAnimation.cpp b/Engine/Plugins/Interchange/Runtime/Source/Import/Private/Gltf/InterchangeGltfAnimation.cpp
--- a/Engine/Plugins/Interchange/Runtime/Source/Import/Private/Gltf/InterchangeGltfAnimation.cpp
+++ b/Engine/Plugins/Interchange/Runtime/Source/Import/Private/Gltf/InterchangeGltfAnimation.cpp
@@ -538,6 +538,19 @@ namespace UE::Interchange::Gltf::Private
 			FVector3d(ScaleData));
 	}
 
+	void GetNonJointTransforms(const GLTF::FAsset& GltfAsset, int32 CurrentNodeIndex, FTransform& Transform)
+	{
+		if (GltfAsset.Nodes.IsValidIndex(CurrentNodeIndex))
+		{
+			if (GltfAsset.Nodes[CurrentNodeIndex].Type != GLTF::FNode::EType::Joint)
+			{
+				Transform = GltfAsset.Nodes[CurrentNodeIndex].Transform * Transform;
+
+				GetNonJointTransforms(GltfAsset, GltfAsset.Nodes[CurrentNodeIndex].ParentIndex, Transform);
+			}
+		}
+	}
+
 	bool GetBakedAnimationTransformPayloadData(const FString& PayLoadKey, const GLTF::FAsset& GltfAsset, FAnimationPayloadData& PayloadData)
 	{
 		const double BakeInterval = 1.0 / PayloadData.BakeFrequency;
@@ -563,6 +576,9 @@ namespace UE::Interchange::Gltf::Private
 		//channels are targetting the same node:
 		bool BaseTransformsAcquired = false;
 
+		bool bNonJointTransformNeedsApplication = false;
+		FTransform NonJointTransformToApply;
+
 		for (const int32& ChannelIndex : ChannelIndices)
 		{
 			const GLTF::FAnimation::FChannel& Channel = GltfAnimation.Channels[ChannelIndex];
@@ -580,6 +596,13 @@ namespace UE::Interchange::Gltf::Private
 				TranslationData.Init(JointBaseTranslation, BakeKeyCount);
 				RotationData.Init(JointBaseRotation, BakeKeyCount);
 				ScaleData.Init(JointBaseScale, BakeKeyCount);
+
+				if (Channel.Target.Node.Type == GLTF::FNode::EType::Joint &&
+					Channel.Target.Node.ParentIndex != Channel.Target.Node.ParentJointIndex)
+				{
+					GetNonJointTransforms(GltfAsset, Channel.Target.Node.ParentIndex, NonJointTransformToApply);
+					bNonJointTransformNeedsApplication = true;
+				}
 			}
 
 			TArray<float> Seconds;
@@ -656,12 +679,26 @@ namespace UE::Interchange::Gltf::Private
 			}
 		}
 
-		for (size_t Index = 0; Index < BakeKeyCount; Index++)
+		if (bNonJointTransformNeedsApplication)
+		{
+			FVector3d NonJointTranslation = NonJointTransformToApply.GetTranslation();
+			FQuat4d NonJointRotation = NonJointTransformToApply.GetRotation();
+			FVector3d NonJointScale = NonJointTransformToApply.GetScale3D();
+
+			for (size_t Index = 0; Index < BakeKeyCount; Index++)
+			{
+				PayloadData.Transforms.Add(FTransform(NonJointRotation * FQuat4d(RotationData[Index]), NonJointTranslation + FVector3d(TranslationData[Index]), NonJointScale * FVector3d(ScaleData[Index])));
+			}
+		}
+		else
 		{
-			PayloadData.Transforms.Add(FTransform(
-				FQuat4d(RotationData[Index]),
-				FVector3d(TranslationData[Index]),
-				FVector3d(ScaleData[Index])));
+			for (size_t Index = 0; Index < BakeKeyCount; Index++)
+			{
+				PayloadData.Transforms.Add(FTransform(
+					FQuat4d(RotationData[Index]),
+					FVector3d(TranslationData[Index]),
+					FVector3d(ScaleData[Index])));
+			}
 		}
 
 		return true;
diff --git a/Engine/Plugins/Interchange/Runtime/Source/Import/Private/Mesh/InterchangeSkeletonHelper.cpp b/Engine/Plugins/Interchange/Runtime/Source/Import/Private/Mesh/InterchangeSkeletonHelper.cpp
--- a/Engine/Plugins/Interchange/Runtime/Source/Import/Private/Mesh/InterchangeSkeletonHelper.cpp
+++ b/Engine/Plugins/Interchange/Runtime/Source/Import/Private/Mesh/InterchangeSkeletonHelper.cpp
@@ -242,7 +242,8 @@ namespace UE::Interchange::Private
 		, const bool bUseTimeZeroAsBindPose
 		, bool& bOutDiffPose
 		, TArray<FString>& OutBoneNotBindNames
-		, bool bImportSockets)
+		, bool bImportSockets
+		, bool bAllowNonJoints)
 	{
 		const UInterchangeSceneNode* JointNode = Cast<UInterchangeSceneNode>(NodeContainer->GetNode(JointNodeId));
 		if (!JointNode)
@@ -250,6 +251,32 @@ namespace UE::Interchange::Private
 			UE_LOG(LogInterchangeImport, Warning, TEXT("Invalid Skeleton Joint"));
 			return;
 		}
+		
+		//Do not add non-joint Nodes as Bones.
+		//	Exception is when a Skeleton was created in order to create AnimSeqence from LevelSequence.
+		//	Which can be tracked against by checking the first Node that this function is executed on: if it does not have Joint Specialization then we add all non-joint nodes.
+
+		if (ParentIndex == INDEX_NONE)
+		{
+			if (JointNode->IsSpecializedTypeContains(UE::Interchange::FSceneNodeStaticData::GetJointSpecializeTypeString()))
+			{
+				bAllowNonJoints = false;
+			}
+			else
+			{
+				bAllowNonJoints = true;
+			}
+		}
+
+		if (!bAllowNonJoints && !JointNode->IsSpecializedTypeContains(UE::Interchange::FSceneNodeStaticData::GetJointSpecializeTypeString()))
+		{
+			const TArray<FString> ChildrenIds = NodeContainer->GetNodeChildrenUids(JointNodeId);
+			for (int32 ChildIndex = 0; ChildIndex < ChildrenIds.Num(); ++ChildIndex)
+			{
+				RecursiveAddBones(NodeContainer, ChildrenIds[ChildIndex], JointInfos, ParentIndex, RefBonesBinary, bUseTimeZeroAsBindPose, bOutDiffPose, OutBoneNotBindNames, bAllowNonJoints);
+			}
+			return;
+		}
 
 		int32 JointInfoIndex = JointInfos.Num();
 		FJointInfo& Info = JointInfos.AddZeroed_GetRef();
@@ -328,7 +355,7 @@ namespace UE::Interchange::Private
 					//We dont want a socket to be a skeleton joint
 					continue;
 				}
-				RecursiveAddBones(NodeContainer, ChildrenIds[ChildIndex], JointInfos, JointInfoIndex, RefBonesBinary, bUseTimeZeroAsBindPose, bOutDiffPose, OutBoneNotBindNames, bImportSockets);
+				RecursiveAddBones(NodeContainer, ChildrenIds[ChildIndex], JointInfos, JointInfoIndex, RefBonesBinary, bUseTimeZeroAsBindPose, bOutDiffPose, OutBoneNotBindNames, bImportSockets, bAllowNonJoints);
 			}
 		}
 	}
diff --git a/Engine/Plugins/Interchange/Runtime/Source/Import/Public/InterchangeSkeletonHelper.h b/Engine/Plugins/Interchange/Runtime/Source/Import/Public/InterchangeSkeletonHelper.h
--- a/Engine/Plugins/Interchange/Runtime/Source/Import/Public/InterchangeSkeletonHelper.h
+++ b/Engine/Plugins/Interchange/Runtime/Source/Import/Public/InterchangeSkeletonHelper.h
@@ -59,7 +59,8 @@ namespace UE::Interchange::Private
 			, const bool bUseTimeZeroAsBindPose
 			, bool& bOutDiffPose
 			, TArray<FString>& OutBoneNotBindNames
-			, bool bImportSockets);
+			, bool bImportSockets
+			, bool bAllowNonJoints = true /*It will be set on based on Root node's details */);
 
 	private:
 		static FName SkeletalLodGetBoneName(const TArray<FMeshBoneInfo>& SkeletalLodRawInfos, int32 BoneIndex);
diff --git a/Engine/Plugins/Interchange/Runtime/Source/Parsers/GLTFCore/Private/GLTF/GLTFReader.cpp b/Engine/Plugins/Interchange/Runtime/Source/Parsers/GLTFCore/Private/GLTF/GLTFReader.cpp
--- a/Engine/Plugins/Interchange/Runtime/Source/Parsers/GLTFCore/Private/GLTF/GLTFReader.cpp
+++ b/Engine/Plugins/Interchange/Runtime/Source/Parsers/GLTFCore/Private/GLTF/GLTFReader.cpp
@@ -530,8 +530,6 @@ namespace GLTF
 			{
 				const int32 NodeIndex = Value->AsNumber();
 				Scene.Nodes.Add(NodeIndex);
-
-				BuildParentIndices(INDEX_NONE, NodeIndex);
 			}
 		}
 
@@ -1174,6 +1172,14 @@ namespace GLTF
 
 			SetupNodesType();
 
+			for (FScene& Scene : Asset->Scenes)
+			{
+				for (int32 NodeIndex : Scene.Nodes)
+				{
+					BuildParentIndices(INDEX_NONE, INDEX_NONE, NodeIndex);
+				}
+			}
+
 			BuildRootJoints();
 
 			GenerateInverseBindPosesPerSkinIndices();
@@ -1364,13 +1370,12 @@ namespace GLTF
 				{
 					FNode& CurrentNode = Asset->Nodes[Skin.Joints[JointCounter]];
 
-					if (CurrentNode.ParentIndex != INDEX_NONE &&
-						Asset->Nodes.IsValidIndex(CurrentNode.ParentIndex) &&
-						Asset->Nodes[CurrentNode.ParentIndex].Type == FNode::EType::Joint &&
-						(Asset->Nodes[CurrentNode.ParentIndex].SkinIndexToGlobalInverseBindMatrix.Contains(SkinIndex) || Asset->Nodes[CurrentNode.ParentIndex].SkinIndexToGlobalInverseBindMatrix.Num() > 0)
+					if (CurrentNode.ParentJointIndex != INDEX_NONE &&
+						Asset->Nodes.IsValidIndex(CurrentNode.ParentJointIndex) &&
+						(Asset->Nodes[CurrentNode.ParentJointIndex].SkinIndexToGlobalInverseBindMatrix.Contains(SkinIndex) || Asset->Nodes[CurrentNode.ParentJointIndex].SkinIndexToGlobalInverseBindMatrix.Num() > 0)
 						)
 					{
-						FNode& ParentNode = Asset->Nodes[CurrentNode.ParentIndex];
+						FNode& ParentNode = Asset->Nodes[CurrentNode.ParentJointIndex];
 
 						//LocalBindPose; //bind pose would be CurrentNode.GlobalInverseBindTransform.Inverse() * ParentNode.GlobalInverseBindTransform
 						FMatrix ParentGlobalInverseBindMatrix;
@@ -1466,18 +1471,33 @@ namespace GLTF
 		}
 	}
 
-	void FFileReader::BuildParentIndices(int32 ParentNodeIndex, int32 CurrentNodeIndex) const
+	void FFileReader::BuildParentIndices(int32 ParentNodeIndex, int32 ParentJointIndex, int32 CurrentNodeIndex) const
 	{
 		if (!Asset->Nodes.IsValidIndex(CurrentNodeIndex))
 		{
 			return;
 		}
+
 		GLTF::FNode& Node = Asset->Nodes[CurrentNodeIndex];
+
+		if (Node.ParentIndex != INDEX_NONE)
+		{
+			//it's already so stop the sequence here.
+			return;
+		}
+
 		Node.ParentIndex = ParentNodeIndex;
 
+		//TODO1: we also should differentiate between Skins
+		//		currently while finding roots and joints, but especially roots, we just look at the type of the Node
+		//		However this way we are taking into account that different skins can have different Skeleton compositions.
+		//		We have to update the FNodes' to contain information as to which Skin are they Joint in.
+		//		Then root calculations should take into account the SkinIdx information and create Roots based on said information.
+		Node.ParentJointIndex = ParentJointIndex;
+
 		for (const int32 ChildNodeIndex : Node.Children)
 		{
-			BuildParentIndices(CurrentNodeIndex, ChildNodeIndex);
+			BuildParentIndices(CurrentNodeIndex, Node.Type == GLTF::FNode::EType::Joint ? CurrentNodeIndex : ParentJointIndex, ChildNodeIndex);
 		}
 	}
 
@@ -1487,9 +1507,10 @@ namespace GLTF
 		{
 			return INDEX_NONE;
 		}
-		while (Asset->Nodes.IsValidIndex(Asset->Nodes[CurrentIndex].ParentIndex) && Asset->Nodes[Asset->Nodes[CurrentIndex].ParentIndex].Type == GLTF::FNode::EType::Joint)
+
+		while (Asset->Nodes.IsValidIndex(Asset->Nodes[CurrentIndex].ParentJointIndex))
 		{
-			CurrentIndex = Asset->Nodes[CurrentIndex].ParentIndex;
+			CurrentIndex = Asset->Nodes[CurrentIndex].ParentJointIndex;
 		}
 		return CurrentIndex;
 	}
@@ -1618,6 +1639,7 @@ namespace GLTF
 						OriginalNode.Children.Reset();
 						OriginalNode.Children.Add(Node.Index);
 						Node.ParentIndex = OriginalNode.Index;
+						Node.ParentJointIndex = OriginalNode.ParentJointIndex;
 
 						Node.Transform = OriginalNode.Transform;
 						Node.bHasLocalBindPose = true;
diff --git a/Engine/Plugins/Interchange/Runtime/Source/Parsers/GLTFCore/Public/GLTFNode.h b/Engine/Plugins/Interchange/Runtime/Source/Parsers/GLTFCore/Public/GLTFNode.h
--- a/Engine/Plugins/Interchange/Runtime/Source/Parsers/GLTFCore/Public/GLTFNode.h
+++ b/Engine/Plugins/Interchange/Runtime/Source/Parsers/GLTFCore/Public/GLTFNode.h
@@ -73,6 +73,7 @@ namespace GLTF
 
 		int32 Index;		// index of FNode in GLTFAsset.Nodes array.
 		int32 ParentIndex;
+		int32 ParentJointIndex;
 		int32 RootJointIndex; //only valid if node is of Joint Type
 
 		TArray<float> MorphTargetWeights; //for the instantiated mesh with morph targets.
@@ -97,6 +98,7 @@ namespace GLTF
 		    , LightIndex(INDEX_NONE)
 			, Index(INDEX_NONE)
 			, ParentIndex(INDEX_NONE)
+			, ParentJointIndex(INDEX_NONE)
 			, RootJointIndex(INDEX_NONE)
 			, bHasLocalBindPose(false)
 		{
diff --git a/Engine/Plugins/Interchange/Runtime/Source/Parsers/GLTFCore/Public/GLTFReader.h b/Engine/Plugins/Interchange/Runtime/Source/Parsers/GLTFCore/Public/GLTFReader.h
--- a/Engine/Plugins/Interchange/Runtime/Source/Parsers/GLTFCore/Public/GLTFReader.h
+++ b/Engine/Plugins/Interchange/Runtime/Source/Parsers/GLTFCore/Public/GLTFReader.h
@@ -67,7 +67,7 @@ namespace GLTF
 		void GenerateLocalBindPosesPerSkinIndices() const; //Per node
 		void SetLocalBindPosesForJoints() const; //Per node
 
-		void BuildParentIndices(int32 ParentNodeIndex, int32 CurrentNodeIndex) const;
+		void BuildParentIndices(int32 ParentNodeIndex, int32 ParentJointIndex, int32 CurrentNodeIndex) const;
 		int32 FindRootJointIndex(int32 CurrentNodeIndex) const;
 		void BuildRootJoints() const;
 
