Interchange:
- Allow users to update materials referred by material slots on reimport

diff --git a/Engine/Plugins/Interchange/Runtime/Source/Import/Private/Mesh/InterchangeMeshHelper.cpp b/Engine/Plugins/Interchange/Runtime/Source/Import/Private/Mesh/InterchangeMeshHelper.cpp
--- a/Engine/Plugins/Interchange/Runtime/Source/Import/Private/Mesh/InterchangeMeshHelper.cpp
+++ b/Engine/Plugins/Interchange/Runtime/Source/Import/Private/Mesh/InterchangeMeshHelper.cpp
@@ -417,6 +417,7 @@ namespace UE::Interchange::Private::MeshHelper
 							FMeshMaterialViewer<MaterialType>& ExistingSkeletalMaterial = ExistMaterialArrayViewer[ExistingMaterialIndex];
 							ExistingSkeletalMaterial.SetMaterialSlotName(MaterialSlotName);
 							ExistingSkeletalMaterial.SetImportedMaterialSlotName(MaterialSlotName);
+							ExistingSkeletalMaterial.SetMaterialInterface(NewMaterial);
 							ExistingMaterialMatchImportedMaterialIndex[ExistingMaterialIndex] = ImportedMaterialIndex;
 							ImportedMaterialMatchExistingMaterialIndex[ImportedMaterialIndex] = ExistingMaterialIndex;
 							break;
