[Interchange-glTF] Fixing Normal and Tangent imports.

#rnx

diff --git a/Engine/Plugins/Interchange/Runtime/Source/Parsers/GLTFCore/Private/GLTFMeshFactory.cpp b/Engine/Plugins/Interchange/Runtime/Source/Parsers/GLTFCore/Private/GLTFMeshFactory.cpp
--- a/Engine/Plugins/Interchange/Runtime/Source/Parsers/GLTFCore/Private/GLTFMeshFactory.cpp
+++ b/Engine/Plugins/Interchange/Runtime/Source/Parsers/GLTFCore/Private/GLTFMeshFactory.cpp
@@ -455,9 +455,9 @@ namespace GLTF
 		FMatrix TotalMatrixForNormal;
 		TotalMatrixForNormal = TotalMatrix.Inverse();
 		TotalMatrixForNormal = TotalMatrixForNormal.GetTransposed();
-		TotalMatrixForNormal.RemoveScaling(); //based on accepted solution in StaticMeshOperations::AppendMeshDescriptions / StaticMeshOperations::ApplyTransform
 
 		const bool bIsMirrored = TotalMatrix.Determinant() < 0.f;
+		const float BinormalMultipler = bIsMirrored ? -1.f : 1.f;
 
 		auto TransformPosition = [](const FMatrix& Matrix, FVector3f& Position)
 		{
@@ -488,7 +488,8 @@ namespace GLTF
 			{
 				for (FVector3f& Normal : Normals)
 				{
-					TransformPosition(TotalMatrixForNormal, Normal);
+					FVector TransformedPosition = TotalMatrixForNormal.TransformVector(FVector(Normal)).GetSafeNormal();
+					Normal = static_cast<FVector3f>(TransformedPosition);
 				}
 			}
 
@@ -525,9 +526,12 @@ namespace GLTF
 
 			if (TotalMatrixForNormal != FMatrix::Identity)
 			{
-				for (FVector4f& Tangent : Tangents)
+				for (FVector4f& TangentBinormal : Tangents)
 				{
-					TransformPosition4f(TotalMatrixForNormal, Tangent);
+					const float Binormal = TangentBinormal.W * BinormalMultipler;
+					FVector4 TransformedTangentBinormal = TotalMatrixForNormal.TransformVector(FVector(TangentBinormal)).GetSafeNormal();
+					TransformedTangentBinormal.W = Binormal;
+					TangentBinormal = FVector4f(TransformedTangentBinormal);
 				}
 			}
 
