Interchange: Allow grouping imported assets in subfolders according to their class names.

#virtualized

diff --git a/Engine/Plugins/Fab/Source/Fab/Private/Pipelines/InterchangeMegascansPipeline.cpp b/Engine/Plugins/Fab/Source/Fab/Private/Pipelines/InterchangeMegascansPipeline.cpp
--- a/Engine/Plugins/Fab/Source/Fab/Private/Pipelines/InterchangeMegascansPipeline.cpp
+++ b/Engine/Plugins/Fab/Source/Fab/Private/Pipelines/InterchangeMegascansPipeline.cpp
@@ -7,6 +7,7 @@
 #include "InterchangeMaterialInstanceNode.h"
 #include "InterchangeMeshFactoryNode.h"
 #include "InterchangeMeshNode.h"
+#include "InterchangePipelineHelper.h"
 #include "InterchangePipelineMeshesUtilities.h"
 #include "InterchangeSceneNode.h"
 #include "InterchangeShaderGraphNode.h"
@@ -21,6 +22,7 @@
 
 #include "Nodes/InterchangeInstancedFoliageTypeFactoryNode.h"
 #include "Nodes/InterchangeUserDefinedAttribute.h"
+#include "Nodes/InterchangeSourceNode.h"
 
 #include "Serialization/JsonReader.h"
 #include "Serialization/JsonSerializer.h"
@@ -397,6 +399,10 @@ void UInterchangeMegascansPipeline::SetFoliageType(const UInterchangeStaticMeshF
 
 	UInterchangeInstancedFoliageTypeFactoryNode* InstancedFoliageTypeFactoryNode = NewObject<UInterchangeInstancedFoliageTypeFactoryNode>(BaseNodeContainer, NAME_None);
 	InstancedFoliageTypeFactoryNode->InitializeNode(FoliageTypeNodeUid, FoliageTypeNodeDisplayLabel, EInterchangeNodeContainerType::FactoryData);
+
+	UInterchangeSourceNode* SourceNode = UInterchangeSourceNode::FindOrCreateUniqueInstance(BaseNodeContainer);
+	UE::Interchange::PipelineHelper::FillSubPathFromSourceNode(InstancedFoliageTypeFactoryNode, SourceNode);
+
 	BaseNodeContainer->AddNode(InstancedFoliageTypeFactoryNode);
 	BaseNodeContainer->SetNodeParentUid(FoliageTypeNodeUid, StaticMeshFactoryNodeUid);
 
diff --git a/Engine/Plugins/Interchange/Runtime/Source/Pipelines/InterchangePipelines.Build.cs b/Engine/Plugins/Interchange/Runtime/Source/Pipelines/InterchangePipelines.Build.cs
--- a/Engine/Plugins/Interchange/Runtime/Source/Pipelines/InterchangePipelines.Build.cs
+++ b/Engine/Plugins/Interchange/Runtime/Source/Pipelines/InterchangePipelines.Build.cs
@@ -26,6 +26,7 @@ namespace UnrealBuildTool.Rules
 					"ApplicationCore",
 					"CinematicCamera",
 					"DeveloperSettings",
+					"Foliage",
 					"GeometryCache",
 					"GLTFCore",
 					"InputCore",
@@ -33,6 +34,7 @@ namespace UnrealBuildTool.Rules
 					"LevelSequence",
 					"Slate",
 					"SlateCore",
+					"VariantManagerContent",
 				}
 			);
 
diff --git a/Engine/Plugins/Interchange/Runtime/Source/Pipelines/Private/InterchangeGenericAssetsPipeline.cpp b/Engine/Plugins/Interchange/Runtime/Source/Pipelines/Private/InterchangeGenericAssetsPipeline.cpp
--- a/Engine/Plugins/Interchange/Runtime/Source/Pipelines/Private/InterchangeGenericAssetsPipeline.cpp
+++ b/Engine/Plugins/Interchange/Runtime/Source/Pipelines/Private/InterchangeGenericAssetsPipeline.cpp
@@ -1,4 +1,4 @@
-// Copyright Epic Games, Inc. All Rights Reserved. 
+// Copyright Epic Games, Inc. All Rights Reserved.
 
 #include "InterchangeGenericAssetsPipeline.h"
 
@@ -8,17 +8,21 @@
 #include "AssetRegistry/AssetRegistryModule.h"
 #include "CoreMinimal.h"
 #include "Engine/SkinnedAssetCommon.h"
+#include "Engine/Texture.h"
 #include "Framework/Docking/TabManager.h"
+#include "GameFramework/Actor.h"
 #include "GeometryCache.h"
 #include "HAL/PlatformApplicationMisc.h"
 #include "InterchangeAnimationTrackSetNode.h"
 #include "InterchangeAnimSequenceFactoryNode.h"
+#include "InterchangeAssetUserData.h"
 #include "InterchangeBlueprintPipelineBase.h"
 #include "InterchangeCommonPipelineDataFactoryNode.h"
 #include "InterchangeGenericAnimationPipeline.h"
 #include "InterchangeGenericMaterialPipeline.h"
 #include "InterchangeGenericMeshPipeline.h"
 #include "InterchangeGenericTexturePipeline.h"
+#include "InterchangeManager.h"
 #include "InterchangeMaterialFactoryNode.h"
 #include "InterchangeMeshNode.h"
 #include "InterchangePipelineHelper.h"
@@ -33,17 +37,20 @@
 #include "InterchangeSourceData.h"
 #include "InterchangeStaticMeshFactoryNode.h"
 #include "InterchangeStaticMeshLodDataNode.h"
-#include "InterchangeAssetUserData.h"
+#include "LevelSequence.h"
 #include "Misc/PackageName.h"
 #include "Misc/Paths.h"
 #include "Nodes/InterchangeBaseNode.h"
 #include "Nodes/InterchangeBaseNodeContainer.h"
+#include "Nodes/InterchangeSourceNode.h"
 #include "Nodes/InterchangeUserDefinedAttribute.h"
+#include "PhysicsEngine/PhysicsAsset.h"
+#include "SparseVolumeTexture/SparseVolumeTexture.h"
 #include "Styling/AppStyle.h"
 #include "Templates/SubclassOf.h"
+#include "UObject/MetaData.h"
 #include "UObject/Object.h"
 #include "UObject/ObjectMacros.h"
-#include "UObject/MetaData.h"
 #include "UObject/ReferencerFinder.h"
 #include "Widgets/Images/SImage.h"
 #include "Widgets/Input/SButton.h"
@@ -52,7 +59,6 @@
 #include "Widgets/Layout/SSeparator.h"
 #include "Widgets/Text/STextBlock.h"
 #include "Widgets/Views/STableRow.h"
-#include "GameFramework/Actor.h"
 
 #include UE_INLINE_GENERATED_CPP_BY_NAME(InterchangeGenericAssetsPipeline)
 
@@ -395,6 +401,8 @@ bool UInterchangeGenericAssetsPipeline::IsPropertyChangeNeedRefresh(const FPrope
 	static const TSet<FName> NeedRefreshProperties =
 	{
 		GET_MEMBER_NAME_CHECKED(UInterchangeGenericAssetsPipeline, bUseSourceNameForAsset),
+		GET_MEMBER_NAME_CHECKED(UInterchangeGenericAssetsPipeline, bSceneNameSubFolder),
+		GET_MEMBER_NAME_CHECKED(UInterchangeGenericAssetsPipeline, bAssetTypeSubFolders),
 		GET_MEMBER_NAME_CHECKED(UInterchangeGenericAssetsPipeline, AssetName)
 	};
 
@@ -1262,11 +1270,36 @@ void UInterchangeGenericAssetsPipeline::ExecutePipeline(UInterchangeBaseNodeCont
 	}
 
 	ImplementUseSourceNameForAssetOption(InBaseNodeContainer, InSourceDatas);
-	//Make sure all factory nodes have the specified strategy
-	InBaseNodeContainer->IterateNodesOfType<UInterchangeFactoryBaseNode>([ReimportStrategyClosure = ReimportStrategy](const FString& NodeUid, UInterchangeFactoryBaseNode* FactoryNode)
+
+	// Setup the source node sub paths, which will be used here but also by all other pipelines that run after us and create asset factory nodes
+	UInterchangeSourceNode* SourceNode = UInterchangeSourceNode::FindOrCreateUniqueInstance(InBaseNodeContainer);
+	if (bAssetTypeSubFolders)
+	{
+		SourceNode->SetCustomUseAssetTypeSubPathSuffix(bAssetTypeSubFolders);
+	}
+	if (bSceneNameSubFolder)
+	{
+		FString SceneNameFolderPrefix = FPaths::GetBaseFilename(InSourceDatas[0]->GetFilename());
+
+		UInterchangeManager::GetInterchangeManager().SanitizeNameInline(
+			SceneNameFolderPrefix,
+			ESanitizeNameTypeFlags::ObjectName | ESanitizeNameTypeFlags::ObjectPath | ESanitizeNameTypeFlags::LongPackage
+		);
+
+		SourceNode->SetCustomSubPathPrefix(SceneNameFolderPrefix);
+	}
+
+	InBaseNodeContainer->IterateNodesOfType<UInterchangeFactoryBaseNode>(
+		[SourceNode, ReimportStrategyClosure = ReimportStrategy](const FString& NodeUid, UInterchangeFactoryBaseNode* FactoryNode)
 		{
+			// This obviously only covers the factory nodes that were created so far. Any pipeline after this will need
+			// to call this themselves if they create additional asset factory nodes
+			UE::Interchange::PipelineHelper::FillSubPathFromSourceNode(FactoryNode, SourceNode);
+
+			// Make sure all factory nodes have the specified strategy
 			FactoryNode->SetReimportStrategyFlags(ReimportStrategyClosure);
-		});
+		}
+	);
 }
 
 void UInterchangeGenericAssetsPipeline::ExecutePostFactoryPipeline(const UInterchangeBaseNodeContainer* InBaseNodeContainer, const FString& NodeKey, UObject* CreatedAsset, bool bIsAReimport)
diff --git a/Engine/Plugins/Interchange/Runtime/Source/Pipelines/Private/InterchangeGenericScenesPipeline.cpp b/Engine/Plugins/Interchange/Runtime/Source/Pipelines/Private/InterchangeGenericScenesPipeline.cpp
--- a/Engine/Plugins/Interchange/Runtime/Source/Pipelines/Private/InterchangeGenericScenesPipeline.cpp
+++ b/Engine/Plugins/Interchange/Runtime/Source/Pipelines/Private/InterchangeGenericScenesPipeline.cpp
@@ -20,6 +20,7 @@
 #include "InterchangeManager.h"
 #include "InterchangeMeshActorFactoryNode.h"
 #include "InterchangeMeshNode.h"
+#include "InterchangePipelineHelper.h"
 #include "InterchangePipelineLog.h"
 #include "InterchangePipelineMeshesUtilities.h"
 #include "InterchangeSceneImportAsset.h"
@@ -33,6 +34,7 @@
 #include "InterchangeStaticMeshLodDataNode.h"
 #include "InterchangeVariantSetNode.h"
 #include "InterchangeVolumeNode.h"
+#include "Nodes/InterchangeSourceNode.h"
 #include "Nodes/InterchangeUserDefinedAttribute.h"
 #include "Volume/InterchangeVolumeDefinitions.h"
 
@@ -232,7 +234,7 @@ void UInterchangeGenericLevelPipeline::ExecutePipeline(UInterchangeBaseNodeConta
 {
 	if (!InBaseNodeContainer)
 	{
-		UE_LOG(LogInterchangePipeline, Warning, TEXT("UInterchangeGenericAssetsPipeline: Cannot execute pre-import pipeline because InBaseNodeContrainer is null"));
+		UE_LOG(LogInterchangePipeline, Warning, TEXT("UInterchangeGenericLevelPipeline: Cannot execute pre-import pipeline because InBaseNodeContrainer is null"));
 		return;
 	}
 
@@ -298,6 +300,8 @@ void UInterchangeGenericLevelPipeline::ExecutePipeline(UInterchangeBaseNodeConta
 	UInterchangeLevelFactoryNode* ParentLevelFactoryNode = nullptr;
 	LevelInstanceActorFactoryNode = nullptr;
 
+	UInterchangeSourceNode* SourceNode = UInterchangeSourceNode::FindOrCreateUniqueInstance(InBaseNodeContainer);
+
 	if (TargetWorld)
 	{
 		//create the parent level factory node
@@ -307,6 +311,7 @@ void UInterchangeGenericLevelPipeline::ExecutePipeline(UInterchangeBaseNodeConta
 		ParentLevelFactoryNode->InitializeNode(ParentNodeUid, ParentDisplayLabel, EInterchangeNodeContainerType::FactoryData);
 		ParentLevelFactoryNode->SetCustomShouldCreateLevel(false);
 		ParentLevelFactoryNode->SetCustomReferenceObject(TargetWorld);
+		UE::Interchange::PipelineHelper::FillSubPathFromSourceNode(ParentLevelFactoryNode, SourceNode);
 		BaseNodeContainer->AddNode(ParentLevelFactoryNode);
 
 		UWorld* ReimportReferenceWorld = nullptr;
@@ -328,6 +333,7 @@ void UInterchangeGenericLevelPipeline::ExecutePipeline(UInterchangeBaseNodeConta
 			LevelFactoryNode->InitializeNode(NodeUid, DisplayLabel, EInterchangeNodeContainerType::FactoryData);
 			LevelFactoryNode->SetCustomCreateWorldPartitionLevel(false);
 			LevelFactoryNode->SetCustomShouldCreateLevel(true);
+			UE::Interchange::PipelineHelper::FillSubPathFromSourceNode(LevelFactoryNode, SourceNode);
 
 			if (ReimportReferenceWorld)
 			{
@@ -376,6 +382,8 @@ void UInterchangeGenericLevelPipeline::ExecutePipeline(UInterchangeBaseNodeConta
 		ensure(!BaseNodeContainer->IsNodeUidValid(FactoryNodeUid));
 		SceneImportFactoryNode = NewObject<UInterchangeSceneImportAssetFactoryNode>(BaseNodeContainer, NAME_None);
 		SceneImportFactoryNode->InitializeNode(FactoryNodeUid, DisplayLabel, EInterchangeNodeContainerType::FactoryData);
+		UE::Interchange::PipelineHelper::FillSubPathFromSourceNode(SceneImportFactoryNode, SourceNode);
+
 		// Add dependency to all the factory nodes created so far
 		BaseNodeContainer->IterateNodesOfType<UInterchangeFactoryBaseNode>(
 			[this](const FString& NodeUid, UInterchangeFactoryBaseNode* Node)
@@ -1188,6 +1196,9 @@ void UInterchangeGenericLevelPipeline::ExecuteSceneVariantSetNodePreImport(const
 	FactoryNode->InitializeNode(FactoryNodeUid, SceneVariantSetNode.GetDisplayLabel(), EInterchangeNodeContainerType::FactoryData);
 	FactoryNode->SetEnabled(bEnableSceneVariantSet);
 
+	UInterchangeSourceNode* SourceNode = UInterchangeSourceNode::FindOrCreateUniqueInstance(BaseNodeContainer);
+	UE::Interchange::PipelineHelper::FillSubPathFromSourceNode(FactoryNode, SourceNode);
+
 	TArray<FString> VariantSetUids;
 	SceneVariantSetNode.GetCustomVariantSetUids(VariantSetUids);
 
@@ -1432,7 +1443,7 @@ void UInterchangeGenericLevelPipeline::ExecutePostBroadcastPipeline(const UInter
 				{
 					if (!EditorUtilities->SaveAsset(WorldAsset))
 					{
-						UE_LOG(LogInterchangePipeline, Warning, TEXT("UInterchangeGenericAssetsPipeline: Cannot save the level instance actor (%s) referenced world (%s)"), *LevelInstanceActor->GetName(), *WorldAsset->GetName());
+						UE_LOG(LogInterchangePipeline, Warning, TEXT("UInterchangeGenericLevelPipeline: Cannot save the level instance actor (%s) referenced world (%s)"), *LevelInstanceActor->GetName(), *WorldAsset->GetName());
 					}
 				}
 			}
diff --git a/Engine/Plugins/Interchange/Runtime/Source/Pipelines/Private/InterchangePipelineHelper.cpp b/Engine/Plugins/Interchange/Runtime/Source/Pipelines/Private/InterchangePipelineHelper.cpp
--- a/Engine/Plugins/Interchange/Runtime/Source/Pipelines/Private/InterchangePipelineHelper.cpp
+++ b/Engine/Plugins/Interchange/Runtime/Source/Pipelines/Private/InterchangePipelineHelper.cpp
@@ -1,12 +1,29 @@
-// Copyright Epic Games, Inc. All Rights Reserved. 
+// Copyright Epic Games, Inc. All Rights Reserved.
 
 #include "InterchangePipelineHelper.h"
 
+#include "Animation/AnimSequence.h"
+#include "Animation/Skeleton.h"
 #include "CoreMinimal.h"
+#include "Engine/SkeletalMesh.h"
+#include "Engine/SpecularProfile.h"
+#include "Engine/StaticMesh.h"
+#include "Engine/Texture.h"
 #include "Framework/Application/SlateApplication.h"
 #include "Framework/Docking/TabManager.h"
+#include "GeometryCache.h"
+#include "FoliageType_InstancedStaticMesh.h"
 #include "HAL/PlatformApplicationMisc.h"
 #include "InterchangePipelineBase.h"
+#include "InterchangeSceneImportAsset.h"
+#include "LevelSequence.h"
+#include "LevelVariantSets.h"
+#include "Materials/MaterialInterface.h"
+#include "Nodes/InterchangeFactoryBaseNode.h"
+#include "Nodes/InterchangeSourceNode.h"
+#include "PhysicsEngine/PhysicsAsset.h"
+#include "Sound/SoundBase.h"
+#include "SparseVolumeTexture/SparseVolumeTexture.h"
 #include "Widgets/SCompoundWidget.h"
 #include "Widgets/SWindow.h"
 
@@ -45,4 +62,93 @@ namespace UE::Interchange::PipelineHelper
 		// @todo: we can make this slow as showing progress bar later
 		FSlateApplication::Get().AddModalWindow(Window, ParentWindow, false);
 	}
+
+	bool FillSubPathFromSourceNode(UInterchangeFactoryBaseNode* FactoryNode, const UInterchangeSourceNode* SourceNode)
+	{
+		if (!FactoryNode || !SourceNode)
+		{
+			return false;
+		}
+
+		FString Prefix;
+		bool bHasPrefix = SourceNode->GetCustomSubPathPrefix(Prefix);
+		bool bUseAssetTypeSuffix = SourceNode->GetCustomUseAssetTypeSubPathSuffix(bUseAssetTypeSuffix) && bUseAssetTypeSuffix;
+
+		if (!bUseAssetTypeSuffix && !bHasPrefix)
+		{
+			// Nothing to do
+			return true;
+		}
+
+		FString Suffix;
+		const UClass* Class = FactoryNode->GetObjectClass();
+		if (bUseAssetTypeSuffix && Class)
+		{
+			if (
+				Class->IsChildOf(UMaterialInterface::StaticClass()) ||
+				Class->IsChildOf(USpecularProfile::StaticClass())
+			)
+			{
+				Suffix = TEXT("Materials");
+			}
+			else if (Class->IsChildOf(UStaticMesh::StaticClass()))
+			{
+				Suffix = TEXT("StaticMeshes");
+			}
+			else if (
+				Class->IsChildOf(UTexture::StaticClass()) ||
+				Class->IsChildOf(USparseVolumeTexture::StaticClass())
+			)
+			{
+				Suffix = TEXT("Textures");
+			}
+			else if (
+				Class->IsChildOf(USkeletalMesh::StaticClass()) ||
+				Class->IsChildOf(USkeleton::StaticClass()) ||
+				Class->IsChildOf(UPhysicsAsset::StaticClass()) ||
+				Class->IsChildOf(UAnimSequence::StaticClass())
+			)
+			{
+				Suffix = TEXT("SkeletalMeshes");
+			}
+			else if (Class->IsChildOf(ULevelSequence::StaticClass()))
+			{
+				Suffix = TEXT("LevelSequences");
+			}
+			else if (Class->IsChildOf(UGeometryCache::StaticClass()))
+			{
+				Suffix = TEXT("GeometryCaches");
+			}
+			else if (Class->IsChildOf(UWorld::StaticClass()))
+			{
+				Suffix = TEXT("Levels");
+			}
+			else if (Class->IsChildOf(USoundBase::StaticClass()))
+			{
+				Suffix = TEXT("Sounds");
+			}
+			else if (Class->IsChildOf(ULevelVariantSets::StaticClass()))
+			{
+				Suffix = TEXT("Variants");
+			}
+			else if (Class->IsChildOf(UFoliageType_InstancedStaticMesh::StaticClass()))
+			{
+				Suffix = TEXT("Foliage");
+			}
+			else if (Class->IsChildOf(UInterchangeSceneImportAsset::StaticClass()))
+			{
+				// No suffix. Always place this outside of the asset type folders
+			}
+			// TODO: Category for GroomAsset, GroomCache, GroomBindingAsset.
+			// Not done yet because base Interchange doesn't support those assets anyway, and it would require
+			// another plugin dependency on the HairStrands plugin
+			else
+			{
+				// Fallback to using the class name as the subfolder
+				Suffix = Class->GetName();
+			}
+		}
+
+		return FactoryNode->SetCustomSubPath(FPaths::Combine(Prefix, Suffix));
+	}
 }
diff --git a/Engine/Plugins/Interchange/Runtime/Source/Pipelines/Private/InterchangeSparseVolumeTexturePipeline.cpp b/Engine/Plugins/Interchange/Runtime/Source/Pipelines/Private/InterchangeSparseVolumeTexturePipeline.cpp
--- a/Engine/Plugins/Interchange/Runtime/Source/Pipelines/Private/InterchangeSparseVolumeTexturePipeline.cpp
+++ b/Engine/Plugins/Interchange/Runtime/Source/Pipelines/Private/InterchangeSparseVolumeTexturePipeline.cpp
@@ -2,9 +2,11 @@
 
 #include "InterchangeSparseVolumeTexturePipeline.h"
 
+#include "InterchangePipelineHelper.h"
 #include "InterchangePipelineLog.h"
 #include "InterchangeSparseVolumeTextureFactoryNode.h"
 #include "InterchangeVolumeNode.h"
+#include "Nodes/InterchangeSourceNode.h"
 #include "Nodes/InterchangeUserDefinedAttribute.h"
 #include "Volume/InterchangeVolumeDefinitions.h"
 
@@ -38,6 +40,9 @@ namespace UE::Interchange::Private
 			TextureFactoryNode = NewObject<UInterchangeSparseVolumeTextureFactoryNode>(BaseNodeContainer);
 			TextureFactoryNode->InitializeNode(NodeUid, DisplayLabel, NodeContainerType);
 
+			UInterchangeSourceNode* SourceNode = UInterchangeSourceNode::FindOrCreateUniqueInstance(BaseNodeContainer);
+			UE::Interchange::PipelineHelper::FillSubPathFromSourceNode(TextureFactoryNode, SourceNode);
+
 			BaseNodeContainer->AddNode(TextureFactoryNode);
 		}
 
diff --git a/Engine/Plugins/Interchange/Runtime/Source/Pipelines/Private/InterchangeglTFPipeline.cpp b/Engine/Plugins/Interchange/Runtime/Source/Pipelines/Private/InterchangeglTFPipeline.cpp
--- a/Engine/Plugins/Interchange/Runtime/Source/Pipelines/Private/InterchangeglTFPipeline.cpp
+++ b/Engine/Plugins/Interchange/Runtime/Source/Pipelines/Private/InterchangeglTFPipeline.cpp
@@ -2,13 +2,15 @@
 
 #include "InterchangeglTFPipeline.h"
 
+#include "InterchangePipelineHelper.h"
 #include "InterchangePipelineLog.h"
 
+#include "InterchangeManager.h"
+#include "InterchangeMaterialFactoryNode.h"
+#include "InterchangeMaterialInstanceNode.h"
 #include "InterchangeMeshFactoryNode.h"
 #include "InterchangeShaderGraphNode.h"
-#include "InterchangeMaterialInstanceNode.h"
-#include "InterchangeMaterialFactoryNode.h"
-#include "InterchangeManager.h"
+#include "Nodes/InterchangeSourceNode.h"
 
 #include "Gltf/InterchangeGLTFMaterial.h"
 
@@ -341,6 +343,9 @@ void UInterchangeGLTFPipeline::ExecutePipeline(UInterchangeBaseNodeContainer* No
 			UInterchangeMaterialInstanceFactoryNode* MaterialInstanceFactoryNode = NewObject<UInterchangeMaterialInstanceFactoryNode>(NodeContainer);
 			GLTFPipelineSettings->BuildMaterialInstance(ShaderGraphNode.Value, MaterialInstanceFactoryNode);
 
+			UInterchangeSourceNode* SourceNode = UInterchangeSourceNode::FindOrCreateUniqueInstance(NodeContainer);
+			UE::Interchange::PipelineHelper::FillSubPathFromSourceNode(MaterialInstanceFactoryNode, SourceNode);
+
 			NodeContainer->ReplaceNode(ShaderGraphNode.Key, MaterialInstanceFactoryNode);
 		}
 	}
diff --git a/Engine/Plugins/Interchange/Runtime/Source/Pipelines/Public/InterchangeGenericAssetsPipeline.h b/Engine/Plugins/Interchange/Runtime/Source/Pipelines/Public/InterchangeGenericAssetsPipeline.h
--- a/Engine/Plugins/Interchange/Runtime/Source/Pipelines/Public/InterchangeGenericAssetsPipeline.h
+++ b/Engine/Plugins/Interchange/Runtime/Source/Pipelines/Public/InterchangeGenericAssetsPipeline.h
@@ -57,6 +57,14 @@ public:
 	UPROPERTY(EditAnywhere, BlueprintReadWrite, Category = "Common")
 	bool bUseSourceNameForAsset = true;
 
+	/** Create an additional Content folder inside of the chosen import directory, and name it after the imported scene */
+	UPROPERTY(EditAnywhere, BlueprintReadWrite, Category = "Common")
+	bool bSceneNameSubFolder = false;
+
+	/** Group the assets according to their type into additional Content folders created on the import directory (/Materials, /StaticMeshes, /SkeletalMeshes, etc.) */
+	UPROPERTY(EditAnywhere, BlueprintReadWrite, Category = "Common")
+	bool bAssetTypeSubFolders = false;
+
 	/** If set, and there is only one asset and one source, the imported asset is given this name. */
 	UPROPERTY(EditAnywhere, BlueprintReadWrite, Category = "Common", meta = (StandAlonePipelineProperty = "True"))
 	FString AssetName;
diff --git a/Engine/Plugins/Interchange/Runtime/Source/Pipelines/Public/InterchangePipelineHelper.h b/Engine/Plugins/Interchange/Runtime/Source/Pipelines/Public/InterchangePipelineHelper.h
--- a/Engine/Plugins/Interchange/Runtime/Source/Pipelines/Public/InterchangePipelineHelper.h
+++ b/Engine/Plugins/Interchange/Runtime/Source/Pipelines/Public/InterchangePipelineHelper.h
@@ -5,7 +5,13 @@
 #include "CoreMinimal.h"
 #include "InterchangePipelineBase.h"
 
+class UInterchangeFactoryBaseNode;
+class UInterchangeSourceNode;
+
 namespace UE::Interchange::PipelineHelper
 {
 	void ShowModalDialog(TSharedRef<SInterchangeBaseConflictWidget> ConflictWidget, const FText& Title, const FVector2D& WindowSize);
+
+	// Fills in the FactoryNode's SubPath custom attribute with the SourceNode's sub path prefixes and suffixes
+	INTERCHANGEPIPELINES_API bool FillSubPathFromSourceNode(UInterchangeFactoryBaseNode* FactoryNode, const UInterchangeSourceNode* SourceNode);
 }
diff --git a/Engine/Source/Runtime/Interchange/Core/Private/Nodes/InterchangeSourceNode.cpp b/Engine/Source/Runtime/Interchange/Core/Private/Nodes/InterchangeSourceNode.cpp
--- a/Engine/Source/Runtime/Interchange/Core/Private/Nodes/InterchangeSourceNode.cpp
+++ b/Engine/Source/Runtime/Interchange/Core/Private/Nodes/InterchangeSourceNode.cpp
@@ -188,3 +188,21 @@ bool UInterchangeSourceNode::SetCustomUseLegacySkeletalMeshBakeTransform(const b
 	IMPLEMENT_NODE_ATTRIBUTE_SETTER_NODELEGATE(UseLegacySkeletalMeshBakeTransform, bool);
 }
 
+bool UInterchangeSourceNode::GetCustomSubPathPrefix(FString& AttributeValue) const
+{
+	IMPLEMENT_NODE_ATTRIBUTE_GETTER(SubPathPrefix, FString);
+}
+bool UInterchangeSourceNode::SetCustomSubPathPrefix(const FString& AttributeValue)
+{
+	IMPLEMENT_NODE_ATTRIBUTE_SETTER_NODELEGATE(SubPathPrefix, FString);
+}
+
+bool UInterchangeSourceNode::GetCustomUseAssetTypeSubPathSuffix(bool& AttributeValue) const
+{
+	IMPLEMENT_NODE_ATTRIBUTE_GETTER(UseAssetTypeSubPathSuffix, bool);
+}
+bool UInterchangeSourceNode::SetCustomUseAssetTypeSubPathSuffix(const bool& AttributeValue)
+{
+	IMPLEMENT_NODE_ATTRIBUTE_SETTER_NODELEGATE(UseAssetTypeSubPathSuffix, bool);
+}
+
diff --git a/Engine/Source/Runtime/Interchange/Core/Public/Nodes/InterchangeSourceNode.h b/Engine/Source/Runtime/Interchange/Core/Public/Nodes/InterchangeSourceNode.h
--- a/Engine/Source/Runtime/Interchange/Core/Public/Nodes/InterchangeSourceNode.h
+++ b/Engine/Source/Runtime/Interchange/Core/Public/Nodes/InterchangeSourceNode.h
@@ -21,8 +21,8 @@ struct FFrame;
 
 namespace UE::Interchange
 {
-	
-	struct FSourceNodeExtraInfoStaticData 
+
+	struct FSourceNodeExtraInfoStaticData
 	{
 		static INTERCHANGECORE_API const FString& GetApplicationVendorExtraInfoKey();
 		static INTERCHANGECORE_API const FString& GetApplicationNameExtraInfoKey();
@@ -157,18 +157,51 @@ public:
 	UFUNCTION(BlueprintCallable, Category = "Interchange | Node | Source")
 	INTERCHANGECORE_API bool SetCustomUseLegacySkeletalMeshBakeTransform(const bool& AttributeValue);
 
+	/**
+	 * Gets a prefix that should be added to factory node SubPath custom attributes.
+	 * For example this can contain the imported scene's name, so that we create an additional content folder named
+	 * after it to contain the imported assets.
+	 */
+	UFUNCTION(BlueprintCallable, Category = "Interchange | Node | Source")
+	INTERCHANGECORE_API bool GetCustomSubPathPrefix(FString& Prefix) const;
+
+	/** Sets the prefix that should be added to factory node SubPath custom attributes. */
+	UFUNCTION(BlueprintCallable, Category = "Interchange | Node | Source")
+	INTERCHANGECORE_API bool SetCustomSubPathPrefix(const FString& Prefix);
+
+	/**
+	 * Gets whether factory nodes for this import should have a suffix named after their asset category added to their
+	 * custom sub path attribute. For example, if this is set then imported StaticMesh assets will be placed inside of
+	 * an additional content folder named "StaticMeshes".
+	 *
+	 * Note that this is done automatically for all factory nodes created by the generic assets pipeline, but must be
+	 * handled manually by calling FillSubPathFromSourceNode() from InterchangePipelineHelper.h for any factory nodes
+	 * that other pipelines may create.
+	 */
+	UFUNCTION(BlueprintCallable, Category = "Interchange | Node | Source")
+	INTERCHANGECORE_API bool GetCustomUseAssetTypeSubPathSuffix(bool& Suffix) const;
+
+	/**
+	 * Sets whether factory nodes for this import should have a suffix named after their asset category added to their
+	 * custom sub path attribute.
+	 */
+	UFUNCTION(BlueprintCallable, Category = "Interchange | Node | Source")
+	INTERCHANGECORE_API bool SetCustomUseAssetTypeSubPathSuffix(const bool& Suffix);
+
 private:
-	const UE::Interchange::FAttributeKey Macro_CustomSourceFrameRateNumeratorKey = UE::Interchange::FAttributeKey(TEXT("SourceFrameRateNumerator"));
-	const UE::Interchange::FAttributeKey Macro_CustomSourceFrameRateDenominatorKey = UE::Interchange::FAttributeKey(TEXT("SourceFrameRateDenominator"));
-	const UE::Interchange::FAttributeKey Macro_CustomSourceTimelineStartKey = UE::Interchange::FAttributeKey(TEXT("SourceTimelineStart"));
-	const UE::Interchange::FAttributeKey Macro_CustomSourceTimelineEndKey = UE::Interchange::FAttributeKey(TEXT("SourceTimelineEnd"));
-	const UE::Interchange::FAttributeKey Macro_CustomAnimatedTimeStartKey = UE::Interchange::FAttributeKey(TEXT("AnimatedTimeStart"));
-	const UE::Interchange::FAttributeKey Macro_CustomAnimatedTimeEndKey = UE::Interchange::FAttributeKey(TEXT("AnimatedTimeEnd"));
-	const UE::Interchange::FAttributeKey Macro_CustomImportUnusedMaterialKey = UE::Interchange::FAttributeKey(TEXT("ImportUnusedMaterial"));
-	const UE::Interchange::FAttributeKey Macro_CustomAxisConversionInverseTransformKey = UE::Interchange::FAttributeKey(TEXT("AxisConversionInverseTransform"));
-	const UE::Interchange::FAttributeKey Macro_CustomUseLegacySkeletalMeshBakeTransformKey = UE::Interchange::FAttributeKey(TEXT("UseLegacySkeletalMeshBakeTransform"));
+	IMPLEMENT_NODE_ATTRIBUTE_KEY(SourceFrameRateNumerator);
+	IMPLEMENT_NODE_ATTRIBUTE_KEY(SourceFrameRateDenominator);
+	IMPLEMENT_NODE_ATTRIBUTE_KEY(SourceTimelineStart);
+	IMPLEMENT_NODE_ATTRIBUTE_KEY(SourceTimelineEnd);
+	IMPLEMENT_NODE_ATTRIBUTE_KEY(AnimatedTimeStart);
+	IMPLEMENT_NODE_ATTRIBUTE_KEY(AnimatedTimeEnd);
+	IMPLEMENT_NODE_ATTRIBUTE_KEY(ImportUnusedMaterial);
+	IMPLEMENT_NODE_ATTRIBUTE_KEY(AxisConversionInverseTransform);
+	IMPLEMENT_NODE_ATTRIBUTE_KEY(UseLegacySkeletalMeshBakeTransform);
+	IMPLEMENT_NODE_ATTRIBUTE_KEY(SubPathPrefix);
+	IMPLEMENT_NODE_ATTRIBUTE_KEY(UseAssetTypeSubPathSuffix);
 
 	// Extra InformationTo show in the Config Panel.
-	UE::Interchange::TMapAttributeHelper<FString, FString> ExtraInformation; 
+	UE::Interchange::TMapAttributeHelper<FString, FString> ExtraInformation;
 };
 
