[Interchange] Removing unused function post FMeshDescription rework.

#rnx

diff --git a/Engine/Plugins/Interchange/Runtime/Source/Import/Private/Mesh/InterchangeSkeletalMeshFactory.cpp b/Engine/Plugins/Interchange/Runtime/Source/Import/Private/Mesh/InterchangeSkeletalMeshFactory.cpp
--- a/Engine/Plugins/Interchange/Runtime/Source/Import/Private/Mesh/InterchangeSkeletalMeshFactory.cpp
+++ b/Engine/Plugins/Interchange/Runtime/Source/Import/Private/Mesh/InterchangeSkeletalMeshFactory.cpp
@@ -576,163 +576,6 @@ namespace UE
 				DestinationMeshDescription = MoveTemp(FixedMeshDescription);
 			}
 
-			//TODO_REFACTOR_SKELETAL_MESHDESCRIPTION_USAGE:
-			//THIS CAN BE REMOVED ONCE THE FMESHDESCRIPTION REFACTOR IS COMPLETE.
-			//KEEPING IT HERE FOR NOW SO THAT DIFF AGAINST HAVE REVISION IS EASIER TO READ.
-			void ProcessImportMeshInfluences(const int32 WedgeCount, TArray<SkeletalMeshImportData::FRawBoneInfluence>& Influences)
-			{
-				TRACE_CPUPROFILER_EVENT_SCOPE(ProcessImportMeshInfluences)
-				// Sort influences by vertex index.
-				struct FCompareVertexIndex
-				{
-					bool operator()(const SkeletalMeshImportData::FRawBoneInfluence& A, const SkeletalMeshImportData::FRawBoneInfluence& B) const
-					{
-						if (A.VertexIndex > B.VertexIndex) return false;
-						else if (A.VertexIndex < B.VertexIndex) return true;
-						else if (A.Weight < B.Weight) return false;
-						else if (A.Weight > B.Weight) return true;
-						else if (A.BoneIndex > B.BoneIndex) return false;
-						else if (A.BoneIndex < B.BoneIndex) return true;
-						else									  return  false;
-					}
-				};
-				Influences.Sort(FCompareVertexIndex());
-
-				TArray <SkeletalMeshImportData::FRawBoneInfluence> NewInfluences;
-				int32	LastNewInfluenceIndex = 0;
-				int32	LastVertexIndex = INDEX_NONE;
-				int32	InfluenceCount = 0;
-
-				float TotalWeight = 0.f;
-				int MaxVertexInfluence = 0;
-				float MaxIgnoredWeight = 0.0f;
-
-				//We have to normalize the data before filtering influences
-				//Because influence filtering is base on the normalize value.
-				//Some DCC like Daz studio don't have normalized weight
-				for (int32 i = 0; i < Influences.Num(); i++)
-				{
-					// if less than min weight, or it's more than 8, then we clear it to use weight
-					InfluenceCount++;
-					TotalWeight += Influences[i].Weight;
-					// we have all influence for the same vertex, normalize it now
-					if (i + 1 >= Influences.Num() || Influences[i].VertexIndex != Influences[i + 1].VertexIndex)
-					{
-						// Normalize the last set of influences.
-						if (InfluenceCount && (TotalWeight != 1.0f))
-						{
-							float OneOverTotalWeight = 1.f / TotalWeight;
-							for (int r = 0; r < InfluenceCount; r++)
-							{
-								Influences[i - r].Weight *= OneOverTotalWeight;
-							}
-						}
-
-						if (MaxVertexInfluence < InfluenceCount)
-						{
-							MaxVertexInfluence = InfluenceCount;
-						}
-
-						// clear to count next one
-						InfluenceCount = 0;
-						TotalWeight = 0.f;
-					}
-
-					if (InfluenceCount > MAX_TOTAL_INFLUENCES && Influences[i].Weight > MaxIgnoredWeight)
-					{
-						MaxIgnoredWeight = Influences[i].Weight;
-					}
-				}
-
-				// warn about too many influences
-				if (MaxVertexInfluence > MAX_TOTAL_INFLUENCES)
-				{
-					//TODO log a display message to the user
-					//UE_LOG(LogLODUtilities, Display, TEXT("Skeletal mesh (%s) influence count of %d exceeds max count of %d. Influence truncation will occur. Maximum Ignored Weight %f"), *MeshName, MaxVertexInfluence, MAX_TOTAL_INFLUENCES, MaxIgnoredWeight);
-				}
-
-				for (int32 i = 0; i < Influences.Num(); i++)
-				{
-					// we found next verts, normalize it now
-					if (LastVertexIndex != Influences[i].VertexIndex)
-					{
-						// Normalize the last set of influences.
-						if (InfluenceCount && (TotalWeight != 1.0f))
-						{
-							float OneOverTotalWeight = 1.f / TotalWeight;
-							for (int r = 0; r < InfluenceCount; r++)
-							{
-								NewInfluences[LastNewInfluenceIndex - r].Weight *= OneOverTotalWeight;
-							}
-						}
-
-						// now we insert missing verts
-						if (LastVertexIndex != INDEX_NONE)
-						{
-							int32 CurrentVertexIndex = Influences[i].VertexIndex;
-							for (int32 j = LastVertexIndex + 1; j < CurrentVertexIndex; j++)
-							{
-								// Add a 0-bone weight if none other present (known to happen with certain MAX skeletal setups).
-								LastNewInfluenceIndex = NewInfluences.AddUninitialized();
-								NewInfluences[LastNewInfluenceIndex].VertexIndex = j;
-								NewInfluences[LastNewInfluenceIndex].BoneIndex = 0;
-								NewInfluences[LastNewInfluenceIndex].Weight = 1.f;
-							}
-						}
-
-						// clear to count next one
-						InfluenceCount = 0;
-						TotalWeight = 0.f;
-						LastVertexIndex = Influences[i].VertexIndex;
-					}
-
-					// if less than min weight, or it's more than 12, then we clear it to use weight
-					if (Influences[i].Weight >= UE::AnimationCore::BoneWeightThreshold &&
-						InfluenceCount < MAX_TOTAL_INFLUENCES)
-					{
-						LastNewInfluenceIndex = NewInfluences.Add(Influences[i]);
-						InfluenceCount++;
-						TotalWeight += Influences[i].Weight;
-					}
-				}
-
-				Influences = NewInfluences;
-
-				// Ensure that each vertex has at least one influence as e.g. CreateSkinningStream relies on it.
-				// The below code relies on influences being sorted by vertex index.
-				if (Influences.Num() == 0)
-				{
-					// warn about no influences
-					//TODO add a user log
-					//UE_LOG(LogLODUtilities, Warning, TEXT("Warning skeletal mesh (%s) has no vertex influences"), *MeshName);
-					// add one for each wedge entry
-					Influences.AddUninitialized(WedgeCount);
-					for (int32 WedgeIdx = 0; WedgeIdx < WedgeCount; WedgeIdx++)
-					{
-						Influences[WedgeIdx].VertexIndex = WedgeIdx;
-						Influences[WedgeIdx].BoneIndex = 0;
-						Influences[WedgeIdx].Weight = 1.0f;
-					}
-					for (int32 i = 0; i < Influences.Num(); i++)
-					{
-						int32 CurrentVertexIndex = Influences[i].VertexIndex;
-
-						if (LastVertexIndex != CurrentVertexIndex)
-						{
-							for (int32 j = LastVertexIndex + 1; j < CurrentVertexIndex; j++)
-							{
-								// Add a 0-bone weight if none other present (known to happen with certain MAX skeletal setups).
-								Influences.InsertUninitialized(i, 1);
-								Influences[i].VertexIndex = j;
-								Influences[i].BoneIndex = 0;
-								Influences[i].Weight = 1.f;
-							}
-							LastVertexIndex = CurrentVertexIndex;
-						}
-					}
-				}
-			}
-
 			struct FContentInfo
 			{
 				bool bApplyGeometry = false;
