USD: Fix static analysic warning with recent cl 41322392.

diff --git a/Engine/Plugins/Interchange/Runtime/Source/Pipelines/Private/InterchangePipelineHelper.cpp b/Engine/Plugins/Interchange/Runtime/Source/Pipelines/Private/InterchangePipelineHelper.cpp
--- a/Engine/Plugins/Interchange/Runtime/Source/Pipelines/Private/InterchangePipelineHelper.cpp
+++ b/Engine/Plugins/Interchange/Runtime/Source/Pipelines/Private/InterchangePipelineHelper.cpp
@@ -72,7 +72,12 @@ namespace UE::Interchange::PipelineHelper
 
 		FString Prefix;
 		bool bHasPrefix = SourceNode->GetCustomSubPathPrefix(Prefix);
-		bool bUseAssetTypeSuffix = SourceNode->GetCustomUseAssetTypeSubPathSuffix(bUseAssetTypeSuffix) && bUseAssetTypeSuffix;
+
+		bool bUseAssetTypeSuffix = false;
+		if (bool bAttributeValue = false; SourceNode->GetCustomUseAssetTypeSubPathSuffix(bAttributeValue))
+		{
+			bUseAssetTypeSuffix = bAttributeValue;
+		}
 
 		if (!bUseAssetTypeSuffix && !bHasPrefix)
 		{
