[Interchange] MorphTarget VertexInstanceNormal setting fix. MorphTaret duplicate name handling fix.

#rnx

diff --git a/Engine/Plugins/Interchange/Runtime/Source/Import/Private/Mesh/InterchangeMeshHelper.cpp b/Engine/Plugins/Interchange/Runtime/Source/Import/Private/Mesh/InterchangeMeshHelper.cpp
--- a/Engine/Plugins/Interchange/Runtime/Source/Import/Private/Mesh/InterchangeMeshHelper.cpp
+++ b/Engine/Plugins/Interchange/Runtime/Source/Import/Private/Mesh/InterchangeMeshHelper.cpp
@@ -576,35 +576,36 @@ namespace UE::Interchange::Private::MeshHelper
 		SkeletonMorphCurveMetadataNames.Reserve(MorphTargetCount);
 		FSkeletalMeshAttributes DestinationMeshAttributes(DestinationMeshDescription);
 
+		TVertexAttributesConstRef<FVector3f> DestinationMeshVertexPositions = DestinationMeshAttributes.GetVertexPositions();
+		TVertexInstanceAttributesConstRef<FVector3f> DestinationMeshVertexInstanceNormals = DestinationMeshAttributes.GetVertexInstanceNormals();
+
 		const bool bUseMorphTargetNormals = true;
 
+		TSet<FName> UsedMorphTargetNames;
+		UsedMorphTargetNames.Append(DestinationMeshAttributes.GetMorphTargetNames());
+		UsedMorphTargetNames.Reserve(UsedMorphTargetNames.Num() + MorphTargetCount);
+
 		for (TPair<FString, TArray<FString>>& NamePerUIDs : KeysPerName)
 		{
-			const FString& MorphTargetName = NamePerUIDs.Key;
+			const FString& MorphTargetNameBase = NamePerUIDs.Key;
+			FString MorphTargetName(MorphTargetNameBase);
 
-			if (!DestinationMeshAttributes.GetMorphTargetNames().Contains(MorphTargetName))
+			for (int32 Suffix = 1; UsedMorphTargetNames.Contains(*MorphTargetName); Suffix++)
 			{
-				DestinationMeshAttributes.RegisterMorphTargetAttribute(*MorphTargetName, bUseMorphTargetNormals);
+				MorphTargetName = FString::Printf(TEXT("%s_%d"), *MorphTargetNameBase, Suffix);
 			}
-			else
+			UsedMorphTargetNames.Add(*MorphTargetName);
+
+			if (MorphTargetName != MorphTargetNameBase)
 			{
-				if (!DestinationMeshAttributes.HasMorphTargetNormalsAttribute(*MorphTargetName))
-				{
-					DestinationMeshAttributes.RegisterMorphTargetAttribute(*MorphTargetName, bUseMorphTargetNormals);
-				}
+				UE_LOG(LogInterchangeImport, Warning, TEXT("Duplicate morph target '%s' found, renamed to '%s'."), *MorphTargetNameBase, *MorphTargetName);
 			}
 
+			DestinationMeshAttributes.RegisterMorphTargetAttribute(*MorphTargetName, bUseMorphTargetNormals);
+
 			TVertexAttributesRef<FVector3f> DestinationMeshMorphPosDeltas = DestinationMeshAttributes.GetVertexMorphPositionDelta(*MorphTargetName);
 			TVertexInstanceAttributesRef<FVector3f> DestinationMeshMorphNormals = DestinationMeshAttributes.GetVertexInstanceMorphNormalDelta(*MorphTargetName);
 
-			TVertexAttributesConstRef<FVector3f> DestinationMeshVertexPositions = DestinationMeshAttributes.GetVertexPositions();
-			TVertexInstanceAttributesConstRef<FVector3f> DestinationMeshVertexInstanceNormals = DestinationMeshAttributes.GetVertexInstanceNormals();
-
-			bool bHasNormals = DestinationMeshVertexInstanceNormals.IsValid() && (DestinationMeshVertexInstanceNormals.GetNumElements() > 0);
-
-			TArray<FVertexID> VertexIDsUsed;
-			TArray<FVertexInstanceID> VertexInstanceIDsUsed;
-
 			for (const FString& MorphTargetKey : NamePerUIDs.Value)
 			{
 				const TOptional<UE::Interchange::FMeshPayloadData>& MorphTargetPayloadData = LodMorphTargetMeshDescriptions.FindChecked(MorphTargetKey);
@@ -621,34 +622,31 @@ namespace UE::Interchange::Private::MeshHelper
 				TVertexAttributesConstRef<FVector3f> MorphTargetMeshVertexPositions = MorphTargetMeshAttributes.GetVertexPositions();
 				TVertexInstanceAttributesConstRef<FVector3f> MorphTargetMeshVertexInstanceNormals = MorphTargetMeshAttributes.GetVertexInstanceNormals();
 
-				bool bSetNormals = bUseMorphTargetNormals && bHasNormals && MorphTargetMeshVertexInstanceNormals.IsValid() && (MorphTargetMeshVertexInstanceNormals.GetNumElements() > 0);
+				bool bSetNormals = bUseMorphTargetNormals && MorphTargetMeshVertexInstanceNormals.IsValid() && (MorphTargetMeshVertexInstanceNormals.GetNumElements() > 0);
 
 				//Populate the deltas in the target mesh description
 				//Note: We don't have to apply GlobalTransform to MorphTarget Position/Normal here anymore, as it is now passed to PayloadRequest.
+				//VertexInstanceNormals will always be present as the Attribute registration will create it 
+				// and it will have the elements inside it due to VertexInstance creation automatically adds emtpy elements. (only exception if we don't have vertexinstances set to begin with)
+
 				if (bSetNormals)
 				{
-					FVertexArray MorphTargetVertices = MorphTargetMeshDescription.Vertices();
-
 					for (FVertexID MorphTargetVertexID : MorphTargetMeshDescription.Vertices().GetElementIDs())
 					{
 						FVertexID DestinationVertexID = MorphTargetVertexID + VertexOffset;
 
 						if ((DestinationMeshMorphPosDeltas.GetNumElements() > DestinationVertexID) &&
-							(DestinationMeshMorphNormals.GetNumElements() > DestinationVertexID) &&
+							(DestinationMeshVertexPositions.GetNumElements() > DestinationVertexID) &&
 							(MorphTargetMeshVertexPositions.GetNumElements() > MorphTargetVertexID))
 						{
 
 							FVector3f PositionDelta = MorphTargetMeshVertexPositions[MorphTargetVertexID] - DestinationMeshVertexPositions[DestinationVertexID];
 
-							VertexIDsUsed.Add(DestinationVertexID);
-
 							DestinationMeshMorphPosDeltas[DestinationVertexID] = PositionDelta;
 
 							TArrayView<const FVertexInstanceID> DestinationVertexInstanceIDs = DestinationMeshDescription.GetVertexVertexInstanceIDs(DestinationVertexID);
 							TArrayView<const FVertexInstanceID> MorphTargetVertexInstanceIDs = MorphTargetMeshDescription.GetVertexVertexInstanceIDs(MorphTargetVertexID);
 
-							VertexInstanceIDsUsed.Append(DestinationVertexInstanceIDs);
-
 							if (DestinationVertexInstanceIDs.Num() == MorphTargetVertexInstanceIDs.Num())
 							{
 								for (size_t VertexInstanceIndex = 0; VertexInstanceIndex < DestinationVertexInstanceIDs.Num(); VertexInstanceIndex++)
@@ -662,10 +660,6 @@ namespace UE::Interchange::Private::MeshHelper
 									DestinationMeshMorphNormals.Set(DestinationVertexInstanceID, NDelta);
 								}
 							}
-							else
-							{
-								UE_LOG(LogInterchangeImport, Error, TEXT("Morph target normal expectation missmatch (DestinationVertexInstanceIDs.Num() != MorphTargetVertexInstanceIDs.Num()) in morph target key [%s]."), *MorphTargetKey);
-							}
 						}
 					}
 				}
@@ -674,15 +668,12 @@ namespace UE::Interchange::Private::MeshHelper
 					for (FVertexID MorphTargetVertexID : MorphTargetMeshDescription.Vertices().GetElementIDs())
 					{
 						FVertexID DestinationVertexID = MorphTargetVertexID + VertexOffset;
-
+						
 						if ((DestinationMeshMorphPosDeltas.GetNumElements() > DestinationVertexID) &&
+							(DestinationMeshVertexPositions.GetNumElements() > DestinationVertexID) &&
 							(MorphTargetMeshVertexPositions.GetNumElements() > MorphTargetVertexID))
 						{
 							FVector3f PositionDelta = MorphTargetMeshVertexPositions[MorphTargetVertexID] - DestinationMeshVertexPositions[DestinationVertexID];
-							if (PositionDelta.IsNearlyZero())
-							{
-								continue;
-							}
 
 							DestinationMeshMorphPosDeltas[DestinationVertexID] = PositionDelta;
 						}
