Interchange : in FBX parser, add different possible texture paths encoded in fbx file like we used to have in Legacy FBX importer (Absolute, Relative to fbx source file)

diff --git a/Engine/Plugins/Interchange/Runtime/Source/Parsers/Fbx/Private/FbxAPI.h b/Engine/Plugins/Interchange/Runtime/Source/Parsers/Fbx/Private/FbxAPI.h
--- a/Engine/Plugins/Interchange/Runtime/Source/Parsers/Fbx/Private/FbxAPI.h
+++ b/Engine/Plugins/Interchange/Runtime/Source/Parsers/Fbx/Private/FbxAPI.h
@@ -107,6 +107,9 @@ namespace UE
 				double GetFrameRate() { return FrameRate; }
 
 				bool IsCreatorBlender() { return bCreatorIsBlender; }
+
+				const FString& GetSourceFilename() const { return SourceFilename; }
+
 				/**
 				 * Critical section to avoid getting multiple payload in same time.
 				 * The FBX evaluator use a cache mechanism for evaluating global transform that is not thread safe.
diff --git a/Engine/Plugins/Interchange/Runtime/Source/Parsers/Fbx/Private/FbxMaterial.cpp b/Engine/Plugins/Interchange/Runtime/Source/Parsers/Fbx/Private/FbxMaterial.cpp
--- a/Engine/Plugins/Interchange/Runtime/Source/Parsers/Fbx/Private/FbxMaterial.cpp
+++ b/Engine/Plugins/Interchange/Runtime/Source/Parsers/Fbx/Private/FbxMaterial.cpp
@@ -73,7 +73,36 @@ namespace UE
 					return nullptr;
 				}
 
-				const FString TextureFilename = FbxTexture ? UTF8_TO_TCHAR(FbxTexture->GetFileName()) : TEXT("");
+				const FString TextureFilename = [this, FbxTexture]
+				{
+					// try opening from absolute path
+					FString AbsolueTextureFilepath = UTF8_TO_TCHAR(FbxTexture->GetFileName());
+					FPaths::NormalizeFilename(AbsolueTextureFilepath);
+					if( FPaths::FileExists(AbsolueTextureFilepath) )
+					{
+						return AbsolueTextureFilepath;
+					}
+
+					FString FileBasePath = FPaths::GetPath(Parser.GetSourceFilename());
+					// try fbx file base path + relative path
+					FString RelativeToFBXTextureFilepath = FileBasePath / UTF8_TO_TCHAR(FbxTexture->GetRelativeFileName());
+					FPaths::NormalizeFilename(RelativeToFBXTextureFilepath);
+					if( FPaths::FileExists(RelativeToFBXTextureFilepath) )
+					{
+						return RelativeToFBXTextureFilepath;
+					}
+
+					// Some fbx files do not store the actual absolute filename as absolute and it is actually relative.  Try to get it relative to the FBX file we are importing
+					FString AbosluteAsRelativeToFBXTextureFilepath = FileBasePath / UTF8_TO_TCHAR(FbxTexture->GetFileName());
+					FPaths::NormalizeFilename(AbosluteAsRelativeToFBXTextureFilepath);
+					if( FPaths::FileExists(AbosluteAsRelativeToFBXTextureFilepath) )
+					{
+						return AbosluteAsRelativeToFBXTextureFilepath;
+					}
+
+					return FString();
+				}();
+
 				const FString NodeName = TEXT("Sampler_") + InputName;
 
 				// Return already created node if applicable.
