Revisited Scenes' pipeline code about LevelInstance and PackedActor to allow use of any targeted world

#rnx

diff --git a/Engine/Plugins/Interchange/Editor/Source/InterchangeEditor/Private/AssetDefinition_InterchangeSceneImportAsset.cpp b/Engine/Plugins/Interchange/Editor/Source/InterchangeEditor/Private/AssetDefinition_InterchangeSceneImportAsset.cpp
--- a/Engine/Plugins/Interchange/Editor/Source/InterchangeEditor/Private/AssetDefinition_InterchangeSceneImportAsset.cpp
+++ b/Engine/Plugins/Interchange/Editor/Source/InterchangeEditor/Private/AssetDefinition_InterchangeSceneImportAsset.cpp
@@ -74,7 +74,7 @@ namespace MenuExtension_InterchangeSceneImportAsset
 			ImportAssetParameters.bIsAutomated = GIsAutomationTesting || FApp::IsUnattended() || IsRunningCommandlet() || GIsRunningUnattendedScript;
 			ImportAssetParameters.ReimportAsset = Asset;
 			ImportAssetParameters.ReimportSourceIndex = INDEX_NONE;
-			ImportAssetParameters.ImportLevel = Asset->GetOriginalLevel();
+			ImportAssetParameters.ImportLevel = nullptr;
 
 			TTuple<FAssetImportResultRef, FSceneImportResultRef> ImportResult = InterchangeManager.ImportSceneAsync(FString(), SourceData, ImportAssetParameters);
 
diff --git a/Engine/Plugins/Interchange/Runtime/Source/Import/Private/Scene/InterchangeLevelFactory.cpp b/Engine/Plugins/Interchange/Runtime/Source/Import/Private/Scene/InterchangeLevelFactory.cpp
--- a/Engine/Plugins/Interchange/Runtime/Source/Import/Private/Scene/InterchangeLevelFactory.cpp
+++ b/Engine/Plugins/Interchange/Runtime/Source/Import/Private/Scene/InterchangeLevelFactory.cpp
@@ -45,18 +45,6 @@ namespace UE::Interchange::Private::InterchangeLevelFactory
 		{
 			TargetAsset = Arguments.ReimportObject;
 		}
-		else
-		{
-			FSoftObjectPath ExistingWorldAssetPath;
-			if (Arguments.AssetNode->GetCustomReferenceObject(ExistingWorldAssetPath))
-			{
-				TargetAsset = ExistingWorldAssetPath.TryLoad();
-			}
-			if (!TargetAsset)
-			{
-				TargetAsset = StaticFindObject(nullptr, Arguments.Parent, *Arguments.AssetName);
-			}
-		}
 		
 		bool bCreateWorldPartition = false;
 		//Verify and ensure if we should create a level or not
diff --git a/Engine/Plugins/Interchange/Runtime/Source/Import/Private/Scene/InterchangeLevelInstanceActorFactory.cpp b/Engine/Plugins/Interchange/Runtime/Source/Import/Private/Scene/InterchangeLevelInstanceActorFactory.cpp
--- a/Engine/Plugins/Interchange/Runtime/Source/Import/Private/Scene/InterchangeLevelInstanceActorFactory.cpp
+++ b/Engine/Plugins/Interchange/Runtime/Source/Import/Private/Scene/InterchangeLevelInstanceActorFactory.cpp
@@ -67,35 +67,17 @@ UObject* UInterchangeLevelInstanceActorFactory::ImportSceneObject_GameThread(con
 		return nullptr;
 	}
 
-	FString TargetLevelFactoryNodeUid;
-	if (!FactoryNode->GetCustomLevelUid(TargetLevelFactoryNodeUid))
-	{
-		UE_LOG(LogInterchangeImport, Warning, TEXT("UInterchangeLevelInstanceActorFactory::ImportSceneObject_GameThread: Actor level instance do not have a valid target world to be created."));
-		return nullptr;
-	}
-
-	UInterchangeLevelFactoryNode* TargetLevelFactoryNode = Cast<UInterchangeLevelFactoryNode>(CreateSceneObjectsParams.NodeContainer->GetFactoryNode(TargetLevelFactoryNodeUid));
-	if (!TargetLevelFactoryNode)
-	{
-		UE_LOG(LogInterchangeImport, Warning, TEXT("UInterchangeLevelInstanceActorFactory::ImportSceneObject_GameThread: Target Level factory node is invalid."));
-		return nullptr;
-	}
-
-	FSoftObjectPath TargetWorldPath;
-	if (!TargetLevelFactoryNode->GetCustomReferenceObject(TargetWorldPath))
+	//Reference world must be saved if we want to update the level instance actor
+	if (UInterchangeEditorUtilitiesBase* EditorUtilities = UInterchangeManager::GetInterchangeManager().GetEditorUtilities())
 	{
-		UE_LOG(LogInterchangeImport, Warning, TEXT("UInterchangeLevelInstanceActorFactory::ImportSceneObject_GameThread: Target Level factory node do not point on a valid UWorld."));
-		return nullptr;
+		if (!EditorUtilities->SaveAsset(ReferenceWorld))
+		{
+			UE_LOG(LogInterchangeImport, Warning, TEXT("UInterchangeLevelInstanceActorFactory::ImportSceneObject_GameThread: Cannot save the level instance actor (%s) referenced world (%s)"), *FactoryNode->GetDisplayLabel(), *ReferenceWorld->GetName());
+		}
 	}
 
-	UWorld* TargetWorld = Cast<UWorld>(TargetWorldPath.TryLoad());
-	if (!TargetWorld)
-	{
-		UE_LOG(LogInterchangeImport, Warning, TEXT("UInterchangeLevelInstanceActorFactory::ImportSceneObject_GameThread: Cannot load a valid target UWorld."));
-		return nullptr;
-	}
-	//ALevelInstanceActor is always created in the target world persistent level
-	ensure(TargetWorld->PersistentLevel == CreateSceneObjectsParams.Level);
+	UWorld* TargetWorld = CreateSceneObjectsParams.Level ? CreateSceneObjectsParams.Level->GetWorld() : GWorld->GetCurrentLevel()->GetWorld();
+	ensure (TargetWorld);
 
 	const bool bIsPackedActor = FactoryNode->GetObjectClass()->IsChildOf<APackedLevelActor>();
 
@@ -118,36 +100,6 @@ UObject* UInterchangeLevelInstanceActorFactory::ImportSceneObject_GameThread(con
 		}
 	}
 
-	auto FindLevelInstance = [TargetWorld, LevelInstanceClass](FName LevelInstanceName)->ALevelInstance*
-		{
-			ALevelInstance* LevelInstanceActor = nullptr;
-			for (ULevel* Level : TargetWorld->GetLevels())
-			{
-				if (LevelInstanceActor)
-				{
-					break;
-				}
-				for (TObjectPtr<AActor> Actor : Level->Actors)
-				{
-					if (!Actor)
-					{
-						continue;
-					}
-					if (FName(*Actor->GetName()) == LevelInstanceName)
-					{
-						if (!Actor->GetClass()->IsChildOf(LevelInstanceClass))
-						{
-							UE_LOG(LogInterchangeImport, Warning, TEXT("UInterchangeLevelInstanceActorFactory::ImportSceneObject_GameThread: There is an existing level instance actor with the same name we want to create/search but with a different class type."));
-							return nullptr;
-						}
-						LevelInstanceActor = Cast<ALevelInstance>(Actor);
-						break;
-					}
-				}
-			}
-			return LevelInstanceActor;
-		};
-
 	ALevelInstance* LevelInstanceActor = nullptr;
 	//Packed level actor need a blueprint between the reference level and the actor
 	if (bIsPackedActor)
@@ -203,7 +155,8 @@ UObject* UInterchangeLevelInstanceActorFactory::ImportSceneObject_GameThread(con
 				AssetInfo.FactoryNode = BpFactoryNode;
 				AssetInfo.bIsReimport = false;
 			}
-			LevelInstanceActor = FindLevelInstance(SpawnParameters.Name);
+
+			LevelInstanceActor = Cast<ALevelInstance>(CreateSceneObjectsParams.ReimportObject);
 			if(!LevelInstanceActor)
 			{
 				LevelInstanceActor = TargetWorld->SpawnActor<APackedLevelActor>(NewBP->GeneratedClass, SpawnParameters);
@@ -218,7 +171,7 @@ UObject* UInterchangeLevelInstanceActorFactory::ImportSceneObject_GameThread(con
 	}
 	else
 	{
-		LevelInstanceActor = FindLevelInstance(SpawnParameters.Name);
+		LevelInstanceActor = Cast<ALevelInstance>(CreateSceneObjectsParams.ReimportObject);
 		if(!LevelInstanceActor)
 		{
 			//simple Level Instance actor simply need to spawn the actor, the reference level is already setup
@@ -232,6 +185,12 @@ UObject* UInterchangeLevelInstanceActorFactory::ImportSceneObject_GameThread(con
 		return nullptr;
 	}
 
+	AActor* ParentActor = UE::Interchange::ActorHelper::GetSpawnedParentActor(CreateSceneObjectsParams.NodeContainer, FactoryNode);
+	if (ParentActor || !LevelInstanceActor->IsAttachedTo(ParentActor))
+	{
+		LevelInstanceActor->AttachToActor(ParentActor, FAttachmentTransformRules::KeepWorldTransform);
+	}
+
 	FactoryNode->SetCustomReferenceObject(LevelInstanceActor);
 
 	return LevelInstanceActor;
diff --git a/Engine/Plugins/Interchange/Runtime/Source/Pipelines/Private/InterchangeGenericScenesPipeline.cpp b/Engine/Plugins/Interchange/Runtime/Source/Pipelines/Private/InterchangeGenericScenesPipeline.cpp
--- a/Engine/Plugins/Interchange/Runtime/Source/Pipelines/Private/InterchangeGenericScenesPipeline.cpp
+++ b/Engine/Plugins/Interchange/Runtime/Source/Pipelines/Private/InterchangeGenericScenesPipeline.cpp
@@ -227,7 +227,7 @@ void UInterchangeGenericLevelPipeline::AdjustSettingsForContext(const FInterchan
 {
 	Super::AdjustSettingsForContext(ContextParams);
 
-	bIsReimportContext = ContextParams.ReimportAsset != nullptr;
+	bIsReimportContext |= ContextParams.ReimportAsset != nullptr;
 }
 
 void UInterchangeGenericLevelPipeline::ExecutePipeline(UInterchangeBaseNodeContainer* InBaseNodeContainer, const TArray<UInterchangeSourceData*>& InSourceDatas, const FString& ContentBasePath)
@@ -282,93 +282,48 @@ void UInterchangeGenericLevelPipeline::ExecutePipeline(UInterchangeBaseNodeConta
 #if WITH_EDITORONLY_DATA
 	const FString FilePath = FPaths::ConvertRelativePathToFull(InSourceDatas[0]->GetFilename());
 
-	UWorld* TargetWorld = nullptr;
-	//Create the preview scene if its a reimport and we have to reimport a ALevelInstance or a APackedLevelActor
-	//In this case we dont want to add a new actor since those import type generate a kind of prefab
-	//that can be instanced multiple time in one or more world.
-	if (bIsReimportContext && SceneHierarchyType != EInterchangeSceneHierarchyType::CreateLevelActors)
-	{
-		PreviewScene = new FPreviewScene(FPreviewScene::ConstructionValues().SetTransactional(false));
-		TargetWorld = PreviewScene->GetWorld();
-	}
+	UInterchangeSourceNode* SourceNode = UInterchangeSourceNode::FindOrCreateUniqueInstance(InBaseNodeContainer);
 
-	if(!TargetWorld)
+	if (SceneHierarchyType != EInterchangeSceneHierarchyType::CreateLevelActors)
 	{
-		TargetWorld = GWorld->GetCurrentLevel()->GetWorld();
-	}
+		ensure(!LevelFactoryNode);
+		const bool bCreatePackedActor = SceneHierarchyType == EInterchangeSceneHierarchyType::CreatePackedActor;
+		const FString DisplayLabel = ("Level_") + FPaths::GetBaseFilename(FilePath);
+		const FString NodeUid = TEXT("Level_") + FilePath;
+		ensure(!BaseNodeContainer->IsNodeUidValid(NodeUid));
 
-	UInterchangeLevelFactoryNode* ParentLevelFactoryNode = nullptr;
-	LevelInstanceActorFactoryNode = nullptr;
-
-	UInterchangeSourceNode* SourceNode = UInterchangeSourceNode::FindOrCreateUniqueInstance(InBaseNodeContainer);
+		LevelFactoryNode = NewObject<UInterchangeLevelFactoryNode>(BaseNodeContainer, NAME_None);
+		LevelFactoryNode->InitializeNode(NodeUid, DisplayLabel, EInterchangeNodeContainerType::FactoryData);
+		
+		LevelFactoryNode->SetCustomCreateWorldPartitionLevel(false);
+		LevelFactoryNode->SetCustomShouldCreateLevel(!bIsReimportContext);
+		UE::Interchange::PipelineHelper::FillSubPathFromSourceNode(LevelFactoryNode, SourceNode);
+			
+		if (bIsReimportContext)
+		{
+			//When we re-import we want the FinalizeObject_GameThread to be call on the level factory for this level node
+			LevelFactoryNode->SetForceNodeReimport();
+		}
 
-	if (TargetWorld)
-	{
-		//create the parent level factory node
-		ParentLevelFactoryNode = NewObject<UInterchangeLevelFactoryNode>(BaseNodeContainer, NAME_None);
-		const FString ParentDisplayLabel = TEXT("ParentLevel_") + TargetWorld->GetName();
-		const FString ParentNodeUid = ParentDisplayLabel;
-		ParentLevelFactoryNode->InitializeNode(ParentNodeUid, ParentDisplayLabel, EInterchangeNodeContainerType::FactoryData);
-		ParentLevelFactoryNode->SetCustomShouldCreateLevel(false);
-		ParentLevelFactoryNode->SetCustomReferenceObject(TargetWorld);
-		UE::Interchange::PipelineHelper::FillSubPathFromSourceNode(ParentLevelFactoryNode, SourceNode);
-		BaseNodeContainer->AddNode(ParentLevelFactoryNode);
-
-		UWorld* ReimportReferenceWorld = nullptr;
-		if (ULevel* Level = Cast<ULevel>(ReimportLevel.TryLoad()))
-		{
-			ReimportReferenceWorld = Level->GetWorld();
-		}
-
-		const bool bReimportReferenceWorld = ReimportReferenceWorld && ReimportReferenceWorld == TargetWorld;
-
-		if (!bReimportReferenceWorld && SceneHierarchyType != EInterchangeSceneHierarchyType::CreateLevelActors)
-		{
-			ensure(!LevelFactoryNode);
-			const bool bCreatePackedActor = SceneHierarchyType == EInterchangeSceneHierarchyType::CreatePackedActor;
-			const FString DisplayLabel = ReimportReferenceWorld ? ReimportReferenceWorld->GetName() : ("Level_") + FPaths::GetBaseFilename(FilePath);
-			const FString NodeUid = TEXT("Level_") + FilePath;
-			ensure(!BaseNodeContainer->IsNodeUidValid(NodeUid));
-			LevelFactoryNode = NewObject<UInterchangeLevelFactoryNode>(BaseNodeContainer, NAME_None);
-			LevelFactoryNode->InitializeNode(NodeUid, DisplayLabel, EInterchangeNodeContainerType::FactoryData);
-			LevelFactoryNode->SetCustomCreateWorldPartitionLevel(false);
-			LevelFactoryNode->SetCustomShouldCreateLevel(true);
-			UE::Interchange::PipelineHelper::FillSubPathFromSourceNode(LevelFactoryNode, SourceNode);
-
-			if (ReimportReferenceWorld)
-			{
-				LevelFactoryNode->SetCustomShouldCreateLevel(false);
-				LevelFactoryNode->SetCustomReferenceObject(ReimportReferenceWorld);
-				//When we re-import we want the FinalizeObject_GameThread to be call on the level factory for this level node
-				LevelFactoryNode->SetForceNodeReimport();
-			}
-			//Look at the parent level factory node before the level factory node
-			LevelFactoryNode->AddFactoryDependencyUid(ParentLevelFactoryNode->GetUniqueID());
+		BaseNodeContainer->AddNode(LevelFactoryNode);
 
-			BaseNodeContainer->AddNode(LevelFactoryNode);
+		//Create a level instance or packed actor
+		{
+			LevelInstanceActorFactoryNode = NewObject<UInterchangeLevelInstanceActorFactoryNode>(BaseNodeContainer, NAME_None);
 
-			//Create a level instance actor
+			if (ensure(LevelInstanceActorFactoryNode))
 			{
-				LevelInstanceActorFactoryNode = NewObject<UInterchangeLevelInstanceActorFactoryNode>(BaseNodeContainer, NAME_None);
-
-				if (ensure(LevelInstanceActorFactoryNode))
-				{
-					UClass* LevelInstanceClass = bCreatePackedActor ? APackedLevelActor::StaticClass() : ALevelInstance::StaticClass();
-					LevelInstanceActorFactoryNode->SetCustomActorClassName(LevelInstanceClass->GetPathName());
-					//This actor has the same name has the level it reference
-					const FString ActorDisplayLabel = DisplayLabel;
-					FString ActorNodeUid = TEXT("LevelInstance_") + FilePath;
-					LevelInstanceActorFactoryNode->InitializeNode(ActorNodeUid, ActorDisplayLabel, EInterchangeNodeContainerType::FactoryData);
-					ActorNodeUid = BaseNodeContainer->AddNode(LevelInstanceActorFactoryNode);
-					//Set the level this actor is referring
-					LevelInstanceActorFactoryNode->SetCustomLevelReference(LevelFactoryNode->GetUniqueID());
-					//We ensure the actor will be created after the parent and reference world are create or ready
-					LevelInstanceActorFactoryNode->AddFactoryDependencyUid(LevelFactoryNode->GetUniqueID());
-					//Add the Level instance actor into the parent level (TargetWorld)
-					LevelInstanceActorFactoryNode->SetCustomLevelUid(ParentLevelFactoryNode->GetUniqueID());
-					//Add the actor to the world that we want to create the actor
-					ParentLevelFactoryNode->AddCustomActorFactoryNodeUid(ActorNodeUid);
-				}
+				const UClass* LevelInstanceClass = bCreatePackedActor ? APackedLevelActor::StaticClass() : ALevelInstance::StaticClass();
+				LevelInstanceActorFactoryNode->SetCustomActorClassName(LevelInstanceClass->GetPathName());
+
+				const FString ActorNodeUid = TEXT("LevelInstance_") + FilePath;
+				LevelInstanceActorFactoryNode->InitializeNode(ActorNodeUid, DisplayLabel, EInterchangeNodeContainerType::FactoryData);
+				BaseNodeContainer->AddNode(LevelInstanceActorFactoryNode);
+				
+				//Set the level this actor is referring
+				LevelInstanceActorFactoryNode->SetCustomLevelReference(LevelFactoryNode->GetUniqueID());
+				//We ensure the actor will be created after the parent and reference world are create or ready
+				LevelInstanceActorFactoryNode->AddFactoryDependencyUid(LevelFactoryNode->GetUniqueID());
 			}
 		}
 	}
@@ -395,7 +350,6 @@ void UInterchangeGenericLevelPipeline::ExecutePipeline(UInterchangeBaseNodeConta
 
 		BaseNodeContainer->AddNode(SceneImportFactoryNode);
 	}
-
 #endif
 
 	/* Find all scene node that are active joint. Non active joint should be convert to actor if they are in a static mesh hierarchy */
@@ -1396,11 +1350,7 @@ void UInterchangeGenericLevelPipeline::ExecutePostImportPipeline(const UIntercha
 		}
 	}
 
-	if (UWorld* World = Cast<UWorld>(CreatedAsset))
-	{
-		PostPipelineImportData.AddWorld(World);
-	}
-	else if (ALevelInstance* LevelInstanceActor = Cast<ALevelInstance>(CreatedAsset))
+	if (ALevelInstance* LevelInstanceActor = Cast<ALevelInstance>(CreatedAsset))
 	{
 		const UInterchangeLevelInstanceActorFactoryNode* LevelInstanceFactoryNode = Cast<UInterchangeLevelInstanceActorFactoryNode>(InBaseNodeContainer->GetFactoryNode(NodeKey));
 		FString LevelFactoryNodeUid;
@@ -1484,85 +1434,47 @@ void UInterchangeGenericLevelPipeline::ExecutePostBroadcastPipeline(const UInter
 
 #if WITH_EDITORONLY_DATA
 
-void UInterchangeGenericLevelPipeline::FPostPipelineImportData::AddWorld(UWorld* World)
+void UInterchangeGenericLevelPipeline::FPostPipelineImportData::AddLevelInstanceActor(ALevelInstance* LevelInstanceActor, UWorld* ReferencedWorld)
 {
-	if (!Worlds.Contains(World))
-	{
-		Worlds.Add(World, false);
-		TArray<ALevelInstance*> CompletedLevelInstances;
-		for (TPair<ALevelInstance*, UWorld*>& LevelInstanceAndReferenceWorld : ReferenceWorldPerLevelInstanceToUpdates)
-		{
-			if (UpdateLevelInstanceInternal(LevelInstanceAndReferenceWorld.Key, LevelInstanceAndReferenceWorld.Value))
-			{
-				CompletedLevelInstances.Add(LevelInstanceAndReferenceWorld.Key);
-			}
-		}
+	bool bReferencedWorldAlreadyProcess = Worlds.Contains(ReferencedWorld);
 
-		//Remove level instance completed
-		for (ALevelInstance* CompletedLevelInstance : CompletedLevelInstances)
-		{
-			ReferenceWorldPerLevelInstanceToUpdates.Remove(CompletedLevelInstance);
-		}
-	}
-}
-
-void UInterchangeGenericLevelPipeline::FPostPipelineImportData::AddLevelInstanceActor(ALevelInstance* LevelInstanceActor, UWorld* ReferenceWorld)
-{
-	if (!UpdateLevelInstanceInternal(LevelInstanceActor, ReferenceWorld))
-	{
-		ReferenceWorldPerLevelInstanceToUpdates.Add(LevelInstanceActor, ReferenceWorld);
-	}
-}
-
-bool UInterchangeGenericLevelPipeline::FPostPipelineImportData::UpdateLevelInstanceInternal(ALevelInstance* LevelInstanceActor, UWorld* ReferencedWorld)
-{
-	if (!Worlds.Contains(ReferencedWorld))
-	{
-		return false;
-	}
 	UWorld* ParentWorld = LevelInstanceActor->GetWorld();
-	if (!Worlds.Contains(ParentWorld))
-	{
-		return false;
-	}
-
-	bool& bWorldUpdate = Worlds.FindChecked(ReferencedWorld);
 
-	if (!bWorldUpdate)
+	if (!bReferencedWorldAlreadyProcess)
 	{
 		if (UInterchangeEditorUtilitiesBase* EditorUtilities = UInterchangeManager::GetInterchangeManager().GetEditorUtilities())
 		{
 			if (!EditorUtilities->SaveAsset(ReferencedWorld))
 			{
-				UE_LOG(LogInterchangePipeline, Warning, TEXT("UInterchangeGenericAssetsPipeline: Cannot save the level instance actor (%s) referenced world (%s)"), *LevelInstanceActor->GetName(), *ReferencedWorld->GetName());
+				UE_LOG(LogInterchangePipeline, Warning, TEXT("UInterchangeGenericScenesPipeline: Cannot save the level instance actor (%s) referenced world (%s)"), *LevelInstanceActor->GetName(), *ReferencedWorld->GetName());
 			}
 		}
 
 		// Make sure newly created level asset gets scanned
 		ULevel::ScanLevelAssets(ReferencedWorld->GetPackage()->GetName());
+	}
 
-		ParentWorld->PreEditChange(nullptr);
-
-		bool bIsPackedLevelActor = false;
-		LevelInstanceActor->SetWorldAsset(ReferencedWorld);
-		if (APackedLevelActor* PackedLevelActor = Cast<APackedLevelActor>(LevelInstanceActor))
-		{
-			bIsPackedLevelActor = true;
-			constexpr bool bCheckoutAndSaveFalse = false;
-			constexpr bool bPromptForSaveFalse = false;
-			//Get the blueprint
-			UBlueprint* Blueprint = PackedLevelActor->GetRootBlueprint();
-			ULevel::ScanLevelAssets(Blueprint->GetPackage()->GetName());
-			//Update the blueprint with the reference level content
-			FPackedLevelActorBuilder::CreateDefaultBuilder()->CreateOrUpdateBlueprint(ReferencedWorld, Blueprint, bCheckoutAndSaveFalse, bPromptForSaveFalse);
-		}
+	ParentWorld->PreEditChange(nullptr);
 
-		if(!bIsPackedLevelActor)
-		{
-			LevelInstanceActor->UpdateLevelInstanceFromWorldAsset();
-			LevelInstanceActor->LoadLevelInstance();
-		}
+	LevelInstanceActor->SetWorldAsset(ReferencedWorld);
+	if (APackedLevelActor* PackedLevelActor = Cast<APackedLevelActor>(LevelInstanceActor))
+	{
+		constexpr bool bCheckoutAndSaveFalse = false;
+		constexpr bool bPromptForSaveFalse = false;
+		//Get the blueprint
+		UBlueprint* Blueprint = PackedLevelActor->GetRootBlueprint();
+		ULevel::ScanLevelAssets(Blueprint->GetPackage()->GetName());
+		//Update the blueprint with the reference level content
+		FPackedLevelActorBuilder::CreateDefaultBuilder()->CreateOrUpdateBlueprint(ReferencedWorld, Blueprint, bCheckoutAndSaveFalse, bPromptForSaveFalse);
+	}
+	else
+	{
+		LevelInstanceActor->UpdateLevelInstanceFromWorldAsset();
+		LevelInstanceActor->LoadLevelInstance();
+	}
 
+	if (!bReferencedWorldAlreadyProcess)
+	{
 		//Reference world must be cleanup since they are not the main world.
 		//This remove all the world managers and prevent GC issue when unloading the main world referencing this world.
 		if (ReferencedWorld->bIsWorldInitialized)
@@ -1570,12 +1482,10 @@ bool UInterchangeGenericLevelPipeline::FPostPipelineImportData::UpdateLevelInsta
 			ReferencedWorld->CleanupWorld();
 		}
 
-		ParentWorld->PostEditChange();
-
-		bWorldUpdate = true;
+		Worlds.Add(ReferencedWorld);
 	}
 
-	return true;
+	ParentWorld->PostEditChange();
 }
 
 #endif //WITH_EDITORONLY_DATA
diff --git a/Engine/Plugins/Interchange/Runtime/Source/Pipelines/Public/InterchangeGenericScenesPipeline.h b/Engine/Plugins/Interchange/Runtime/Source/Pipelines/Public/InterchangeGenericScenesPipeline.h
--- a/Engine/Plugins/Interchange/Runtime/Source/Pipelines/Public/InterchangeGenericScenesPipeline.h
+++ b/Engine/Plugins/Interchange/Runtime/Source/Pipelines/Public/InterchangeGenericScenesPipeline.h
@@ -120,13 +120,10 @@ protected:
 
 	struct FPostPipelineImportData
 	{
-		void AddWorld(UWorld* World);
 		void AddLevelInstanceActor(ALevelInstance* LevelInstanceActor, UWorld* ReferenceWorld);
 	private:
-		bool UpdateLevelInstanceInternal(ALevelInstance* LevelInstanceActor, UWorld* ReferenceWorld);
 
-		//The bool is to know if we have already update this world
-		TMap<UWorld*, bool> Worlds;
+		TSet<UWorld*> Worlds;
 		TMap<ALevelInstance*, UWorld*> ReferenceWorldPerLevelInstanceToUpdates;
 	};
 
diff --git a/Engine/Source/Runtime/Interchange/Core/Public/InterchangePipelineBase.h b/Engine/Source/Runtime/Interchange/Core/Public/InterchangePipelineBase.h
--- a/Engine/Source/Runtime/Interchange/Core/Public/InterchangePipelineBase.h
+++ b/Engine/Source/Runtime/Interchange/Core/Public/InterchangePipelineBase.h
@@ -549,12 +549,6 @@ public:
 	UPROPERTY()
 	FString ContentImportPath;
 
-	/*
-	 * Extra data available for scene pipeline, when we do a re-import we need to know what is the level so we do not create a new one.
-	 */
-	UPROPERTY()
-	FSoftObjectPath ReimportLevel;
-
 protected:
 
 	INTERCHANGECORE_API virtual FString GetPipelineDisplayName() const;
diff --git a/Engine/Source/Runtime/Interchange/Engine/Private/InterchangeManager.cpp b/Engine/Source/Runtime/Interchange/Engine/Private/InterchangeManager.cpp
--- a/Engine/Source/Runtime/Interchange/Engine/Private/InterchangeManager.cpp
+++ b/Engine/Source/Runtime/Interchange/Engine/Private/InterchangeManager.cpp
@@ -1961,7 +1961,6 @@ UInterchangeManager::ImportInternal(const FString& ContentPath, const UInterchan
 							{
 								MutableImportAssetParameters->ReimportAsset = ExistingAsset;
 								MutableImportAssetParameters->ReimportSourceIndex = INDEX_NONE;
-								MutableImportAssetParameters->ImportLevel = (bImportScene && SceneImportAsset) ? SceneImportAsset->GetOriginalLevel() : nullptr;
 								MutableImportAssetParameters->bReplaceExisting |= bImportScene;
 							}
 							UInterchangeAssetImportData* MutableOriginalAssetImportData = UInterchangeAssetImportData::GetFromObject(ImportAssetParameters.ReimportAsset);
@@ -2002,6 +2001,37 @@ UInterchangeManager::ImportInternal(const FString& ContentPath, const UInterchan
 	else
 	{
 		ContentBasePath = FPaths::GetPath(ImportAssetParameters.ReimportAsset->GetPathName());
+
+		// Check whether asset has been imported with a sub-folder
+		FString SubPath;
+		if (UInterchangeAssetImportData* ImportData = UInterchangeAssetImportData::GetFromObject(ImportAssetParameters.ReimportAsset))
+		{
+			UInterchangeBaseNodeContainer* AssetContainer = ImportData->GetNodeContainer();
+			ensure(AssetContainer);
+
+			FSoftObjectPath AssetObjectPath(ImportAssetParameters.ReimportAsset);
+
+			AssetContainer->IterateNodesOfType<UInterchangeFactoryBaseNode>([&SubPath, &AssetObjectPath](const FString& NodeUid, UInterchangeFactoryBaseNode* FactoryNode)
+				{
+					if (SubPath.IsEmpty())
+					{
+						FSoftObjectPath ObjectPath;
+						if (FactoryNode->GetCustomReferenceObject(ObjectPath))
+						{
+							if (ObjectPath == AssetObjectPath)
+							{
+								FactoryNode->GetCustomSubPath(SubPath);
+							}
+						}
+					}
+				});
+		}
+
+		// Remove sub-path from ContentBasePath
+		if (!SubPath.IsEmpty())
+		{
+			ContentBasePath.LeftChopInline(SubPath.EndsWith(TEXT("/")) ? SubPath.Len() : SubPath.Len() + 1, EAllowShrinking::No);
+		}
 	}
 
 	const bool bIsReimport = OriginalAssetImportData && OriginalAssetImportData->GetPipelines().Num() > 0;
@@ -2110,7 +2140,6 @@ UInterchangeManager::ImportInternal(const FString& ContentPath, const UInterchan
 				Context = bImportScene ? EInterchangePipelineContext::SceneImport : EInterchangePipelineContext::AssetImport;
 			}
 			Pipeline->ContentImportPath = ContentBasePath;
-			Pipeline->ReimportLevel = bIsReimport && bImportScene ? TaskData.ImportLevel : nullptr;
 			FInterchangePipelineContextParams ContextParams;
 			ContextParams.ContextType = Context;
 			ContextParams.ReimportAsset = TaskData.ReimportObject;
diff --git a/Engine/Source/Runtime/Interchange/Engine/Private/InterchangeSceneImportAsset.cpp b/Engine/Source/Runtime/Interchange/Engine/Private/InterchangeSceneImportAsset.cpp
--- a/Engine/Source/Runtime/Interchange/Engine/Private/InterchangeSceneImportAsset.cpp
+++ b/Engine/Source/Runtime/Interchange/Engine/Private/InterchangeSceneImportAsset.cpp
@@ -181,23 +181,6 @@ const TArray<UAssetUserData*>* UInterchangeSceneImportAsset::GetAssetUserDataArr
 #endif // #if WITH_EDITORONLY_DATA
 }
 
-ULevel* UInterchangeSceneImportAsset::GetOriginalLevel() const
-{
-	TArray<FSoftObjectPath> SceneSoftObjectPaths;
-	GetSceneSoftObjectPaths(SceneSoftObjectPaths);
-	for (const FSoftObjectPath& SoftObjectPath : SceneSoftObjectPaths)
-	{
-		if (AActor* Actor = Cast<AActor>(SoftObjectPath.TryLoad()))
-		{
-			if (!Actor->IsA<ALevelInstance>())
-			{
-				return Actor->GetLevel();
-			}
-		}
-	}
-	return nullptr;
-}
-
 void UInterchangeSceneImportAsset::UpdateSceneObjects()
 {
 #if WITH_EDITORONLY_DATA
@@ -236,7 +219,7 @@ UObject* UInterchangeSceneImportAsset::GetSceneObject(const FString& PackageName
 			&& SceneObjectPath.GetAssetFName() == ObjectPath.GetAssetFName())
 		{
 			//World partition actor have a guid that we need to remove before comparison
-			FString SceneObjectSubPathString = SceneObjectPath.GetSubPathString();
+			const FString SceneObjectSubPathString = SceneObjectPath.GetSubPathString();
 			if (SceneObjectSubPathString.Contains(SubPathString)
 				&& FActorSpawnUtils::GetBaseName(*SceneObjectSubPathString) == ObjectSubPathBaseName)
 			{
@@ -250,7 +233,8 @@ UObject* UInterchangeSceneImportAsset::GetSceneObject(const FString& PackageName
 	{
 		if (IsValid(SceneObject))
 		{
-			if (SceneObject->IsA<AActor>())
+			// No AssetImportData on AActors nor UWorlds nor ULevels. Return found SceneObject
+			if (SceneObject->IsA<AActor>() || SceneObject->IsA<UWorld>() || SceneObject->IsA<ULevel>())
 			{
 				return SceneObject;
 			}
diff --git a/Engine/Source/Runtime/Interchange/Engine/Private/Tasks/InterchangeTaskCreateSceneObjects.cpp b/Engine/Source/Runtime/Interchange/Engine/Private/Tasks/InterchangeTaskCreateSceneObjects.cpp
--- a/Engine/Source/Runtime/Interchange/Engine/Private/Tasks/InterchangeTaskCreateSceneObjects.cpp
+++ b/Engine/Source/Runtime/Interchange/Engine/Private/Tasks/InterchangeTaskCreateSceneObjects.cpp
@@ -86,6 +86,7 @@ void UE::Interchange::FTaskCreateSceneObjects_GameThread::Execute()
 		ULevel* FactoryImportLevel = ImportLevel;
 		UWorld* FactoryImportWorld = ImportWorld;
 
+		// Check whether the factory node is part of a specific sub-level. If so, set the latter as the target level.
 		FString FactoryNodeLevelUid;
 		if (FactoryNode->GetCustomLevelUid(FactoryNodeLevelUid))
 		{
diff --git a/Engine/Source/Runtime/Interchange/Engine/Private/Tasks/InterchangeTaskPipeline.cpp b/Engine/Source/Runtime/Interchange/Engine/Private/Tasks/InterchangeTaskPipeline.cpp
--- a/Engine/Source/Runtime/Interchange/Engine/Private/Tasks/InterchangeTaskPipeline.cpp
+++ b/Engine/Source/Runtime/Interchange/Engine/Private/Tasks/InterchangeTaskPipeline.cpp
@@ -156,8 +156,9 @@ void UE::Interchange::FTaskPostImport_GameThread::Execute()
 			for (const UE::Interchange::FImportAsyncHelper::FImportedObjectInfo& ObjectInfo : ImportedInfos)
 			{
 				UObject* ImportedObject = ObjectInfo.ImportedObject.TryLoad();
-				if (!ObjectInfo.bPostEditChangeCalled)
+				if (ImportedObject && !ObjectInfo.bPostEditChangeCalled)
 				{
+					ObjectInfo.bPostEditChangeCalled = true;
 					ImportedObject->PostEditChange();
 				}
 			}
diff --git a/Engine/Source/Runtime/Interchange/Engine/Public/InterchangeSceneImportAsset.h b/Engine/Source/Runtime/Interchange/Engine/Public/InterchangeSceneImportAsset.h
--- a/Engine/Source/Runtime/Interchange/Engine/Public/InterchangeSceneImportAsset.h
+++ b/Engine/Source/Runtime/Interchange/Engine/Public/InterchangeSceneImportAsset.h
@@ -51,9 +51,6 @@ public:
 	INTERCHANGEENGINE_API void RegisterWorldRenameCallbacks();
 #endif
 
-	/** Return the original level we import actor into. Return nullptr if no level was found */
-	INTERCHANGEENGINE_API ULevel* GetOriginalLevel() const;
-
 	/** Updates the SceneObjects cache based on the node container stored in AssetImportData */
 	INTERCHANGEENGINE_API void UpdateSceneObjects();
 
