Interchange MaterialX:
- Fix wrong shader graph for material mtlx_cyc_wall
- Move code from the MaterialX Translator to the MaterialX Manager in order for users call to be consistent when the API is not called directly from Interchange
- Fix MaterialExpressions Minus and Difference where the operation was inverted
- Use default input values for some nodes in case there's no input
- Added material functions MX_Checkerboard, MX_Circle, MX_TiledCircles, MX_Line, MX_UnifiedNoise3d in order to have smaller shader graphs
- Added a material expression Mod based on glsl spec, because it's what MaterialX is based on, whereas hlsl fmod gives different results if 'y' input is negative
- Allow textures with the same name to be imported, but avoid duplicates if several materials in a same file use the same texture (e.g: chess_set)

#virtualized

diff --git a/Engine/Plugins/Interchange/Runtime/Source/Common/Public/MaterialX/InterchangeMaterialXDefinitions.h b/Engine/Plugins/Interchange/Runtime/Source/Common/Public/MaterialX/InterchangeMaterialXDefinitions.h
--- a/Engine/Plugins/Interchange/Runtime/Source/Common/Public/MaterialX/InterchangeMaterialXDefinitions.h
+++ b/Engine/Plugins/Interchange/Runtime/Source/Common/Public/MaterialX/InterchangeMaterialXDefinitions.h
@@ -158,7 +158,6 @@ namespace UE
 			static constexpr const TCHAR* Dotproduct = TEXT("Dotproduct");
 			static constexpr const TCHAR* Exponential = TEXT("Exponential");
 			static constexpr const TCHAR* Floor = TEXT("Floor");
-			static constexpr const TCHAR* Fmod = TEXT("Fmod");
 			static constexpr const TCHAR* HsvToRgb = TEXT("HsvToRgb");
 			static constexpr const TCHAR* If = TEXT("If");
 			static constexpr const TCHAR* Logarithm = TEXT("Logarithm");
@@ -177,6 +176,7 @@ namespace UE
 			static constexpr const TCHAR* MaterialXMask = TEXT("MaterialXMask");
 			static constexpr const TCHAR* MaterialXMatte = TEXT("MaterialXMatte");
 			static constexpr const TCHAR* MaterialXMinus = TEXT("MaterialXMinus");
+			static constexpr const TCHAR* MaterialXMod = TEXT("MaterialXMod");
 			static constexpr const TCHAR* MaterialXOut = TEXT("MaterialXOut");
 			static constexpr const TCHAR* MaterialXOver = TEXT("MaterialXOver");
 			static constexpr const TCHAR* MaterialXOverlay = TEXT("MaterialXOverlay");
@@ -292,9 +292,12 @@ namespace UE
 			constexpr const TCHAR* MxACEScgToRec709 = TEXT("/InterchangeAssets/Functions/MX_ACEScgToRec709.MX_ACEScgToRec709");
 			constexpr const TCHAR* MxAdobeRGBToRec709 = TEXT("/InterchangeAssets/Functions/MX_AdobeRGBToRec709.MX_AdobeRGBToRec709");
 			constexpr const TCHAR* MxArtisticIOR = TEXT("/InterchangeAssets/Functions/MX_Artistic_IOR.MX_Artistic_IOR");
+			constexpr const TCHAR* MxCircle = TEXT("/InterchangeAssets/Functions/MX_Circle.MX_Circle");
+			constexpr const TCHAR* MxCheckerboard = TEXT("/InterchangeAssets/Functions/MX_Checkerboard.MX_Checkerboard");
 			constexpr const TCHAR* MxColorCorrect = TEXT("/InterchangeAssets/Functions/MX_ColorCorrect.MX_ColorCorrect");
 			constexpr const TCHAR* MxGoochShade = TEXT("/InterchangeAssets/Functions/MX_GoochShade.MX_GoochShade");
 			constexpr const TCHAR* MxHsvAdjust = TEXT("/InterchangeAssets/Functions/MX_HSVAdjust.MX_HSVAdjust");
+			constexpr const TCHAR* MxLine = TEXT("/InterchangeAssets/Functions/MX_Line.MX_Line");
 			constexpr const TCHAR* MxLinearAdobeRGBToRec709 = TEXT("/InterchangeAssets/Functions/MX_LinearAdobeRGBToRec709.MX_LinearAdobeRGBToRec709");
 			constexpr const TCHAR* MxLinearDisplayP3ToRec709 = TEXT("/InterchangeAssets/Functions/MX_LinearDisplayP3ToRec709.MX_LinearDisplayP3ToRec709");
 			constexpr const TCHAR* MxPlace2D = TEXT("/InterchangeAssets/Functions/MX_Place2D.MX_Place2D");
@@ -304,6 +307,8 @@ namespace UE
 			constexpr const TCHAR* MxRoughnessDual = TEXT("/InterchangeAssets/Functions/MX_Roughness_Dual.MX_Roughness_Dual");
 			constexpr const TCHAR* MxSrgbDisplayP3ToRec709 = TEXT("/InterchangeAssets/Functions/MX_SrgbDisplayP3ToRec709.MX_SrgbDisplayP3ToRec709");
 			constexpr const TCHAR* MxSrgbToRec709 = TEXT("/InterchangeAssets/Functions/MX_SrgbToRec709.MX_SrgbToRec709");
+			constexpr const TCHAR* MxUnifiedNoise3D = TEXT("/InterchangeAssets/Functions/MX_UnifiedNoise3d.MX_UnifiedNoise3d");
+			constexpr const TCHAR* MxTiledCircles = TEXT("/InterchangeAssets/Functions/MX_TiledCircles.MX_TiledCircles");
 			constexpr const TCHAR* NormalFromHeightMap = TEXT("/Engine/Functions/Engine_MaterialFunctions03/Procedurals/NormalFromHeightmap.NormalFromHeightmap");
 			constexpr const TCHAR* Refract = TEXT("/Engine/Functions/Engine_MaterialFunctions01/Vectors/Refract.Refract");
 		}
@@ -725,6 +730,9 @@ MATERIALX_NAMESPACE_BEGIN
 		static constexpr const char* Constant = "constant";
 		// Procedural2D nodes
 		static constexpr const char* CellNoise2D = "cellnoise2d";
+		static constexpr const char* Checkerboard = "checkerboard";
+		static constexpr const char* Circle = "circle";
+		static constexpr const char* Line = "line";
 		static constexpr const char* Noise2D = "noise2d";
 		static constexpr const char* Ramp4 = "ramp4";
 		static constexpr const char* RampLR = "ramplr";
@@ -732,12 +740,14 @@ MATERIALX_NAMESPACE_BEGIN
 		static constexpr const char* RandomFloat = "randomfloat";
 		static constexpr const char* SplitLR = "splitlr";
 		static constexpr const char* SplitTB = "splittb";
+		static constexpr const char* TiledCircles = "tiledcircles";
 		static constexpr const char* WorleyNoise2D = "worleynoise2d";
 		// Procedural3D nodes 
 		static constexpr const char* CellNoise3D = "cellnoise3d";
 		static constexpr const char* Fractal3D = "fractal3d";
 		static constexpr const char* Noise3D = "noise3d";
 		static constexpr const char* RandomColor = "randomcolor";
+		static constexpr const char* UnifiedNoise3D = "unifiednoise3d";
 		static constexpr const char* WorleyNoise3D = "worleynoise3d";
 		// Organization nodes 
 		static constexpr const char* Dot = "dot";
diff --git a/Engine/Plugins/Interchange/Runtime/Source/Import/Private/MaterialX/InterchangeMaterialXTranslator.cpp b/Engine/Plugins/Interchange/Runtime/Source/Import/Private/MaterialX/InterchangeMaterialXTranslator.cpp
--- a/Engine/Plugins/Interchange/Runtime/Source/Import/Private/MaterialX/InterchangeMaterialXTranslator.cpp
+++ b/Engine/Plugins/Interchange/Runtime/Source/Import/Private/MaterialX/InterchangeMaterialXTranslator.cpp
@@ -6,16 +6,8 @@
 #include "MaterialX/MaterialXUtils/MaterialXManager.h"
 #include "Nodes/InterchangeSourceNode.h"
 
-#if WITH_EDITOR
-#include "MaterialXFormat/Util.h"
-#include "MaterialX/InterchangeMaterialXDefinitions.h"
-#include "MaterialX/MaterialXUtils/MaterialXBase.h"
-#endif
-
 #include UE_INLINE_GENERATED_CPP_BY_NAME(InterchangeMaterialXTranslator)
 
-#define LOCTEXT_NAMESPACE "InterchangeMaterialXTranslator"
-
 static bool GInterchangeEnableMaterialXImport = true;
 static FAutoConsoleVariableRef CCvarInterchangeEnableMaterialXImport(
 	TEXT("Interchange.FeatureFlags.Import.MTLX"),
@@ -45,58 +37,6 @@ TArray<FString> UInterchangeMaterialXTranslator::GetSupportedFormats() const
 	return UE::Interchange::MaterialX::AreMaterialFunctionPackagesLoaded() ? TArray<FString>{ TEXT("mtlx;MaterialX File Format") } : TArray<FString>{};
 }
 
-namespace
-{
-#if WITH_EDITOR
-	bool ValidateDocument(MaterialX::DocumentPtr Document, const UInterchangeTranslatorBase* Translator)
-	{
-		namespace mx = MaterialX;
-		
-		if(std::string MaterialXMessage; !Document->validate(&MaterialXMessage))
-		{
-			UInterchangeResultError_Generic* Message = Translator->AddMessage<UInterchangeResultError_Generic>();
-			Message->Text = FText::Format(LOCTEXT("MaterialXDocumentInvalid", "{0}"),
-										  FText::FromString(MaterialXMessage.c_str()));
-			return false;
-		}
-
-		for(mx::ElementPtr Elem : Document->traverseTree())
-		{
-			//make sure to read only the current file otherwise we'll process the entire library
-			if(Elem->getActiveSourceUri() != Document->getActiveSourceUri())
-			{
-				continue;
-			}
-
-			mx::NodePtr Node = Elem->asA<mx::Node>();
-
-			if(Node)
-			{
-				// Validate that all nodes in the file are strictly respecting their node definition
-				if(!Node->getNodeDef())
-				{
-					UInterchangeResultError_Generic* Message = Translator->AddMessage<UInterchangeResultError_Generic>();
-					Message->Text = FText::Format(LOCTEXT("NodeDefNotFound", "<{0}> has no matching NodeDef, aborting import..."),
-												  FText::FromString(Node->getName().c_str()));
-					return false;
-				}
-
-				//verify first if it's not a multioutput node (weird that there's no typedef for it)
-				if (!Node->getTypeDef() && Node->getType() != "multioutput")
-				{
-					UInterchangeResultError_Generic* Message = Translator->AddMessage<UInterchangeResultError_Generic>();
-					Message->Text = FText::Format(LOCTEXT("TypeDefNotFound", "<{0}> has no matching TypeDef, aborting import..."),
-												  FText::FromString(Node->getName().c_str()));
-					return false;
-				}
-			}
-		}
-
-		return true;
-	}
-#endif // WITH_EDITOR
-}
-
 bool UInterchangeMaterialXTranslator::Translate(UInterchangeBaseNodeContainer& BaseNodeContainer) const
 {
 	bool bIsDocumentValid = false;
@@ -106,78 +46,7 @@ bool UInterchangeMaterialXTranslator::Translate(UInterchangeBaseNodeContainer& B
 
 	FString Filename = GetSourceData()->GetFilename();
 
-	if(!FPaths::FileExists(Filename))
-	{
-		return false;
-	}
-	try
-	{
-		mx::FileSearchPath MaterialXFolder{ TCHAR_TO_UTF8(*FPaths::Combine(
-			FPaths::EngineDir(),
-			TEXT("Binaries"),
-			TEXT("ThirdParty"),
-			TEXT("MaterialX"))) };
-
-		mx::DocumentPtr MaterialXLibrary = mx::createDocument();
-
-		mx::StringSet LoadedLibs = mx::loadLibraries({ mx::Library::Libraries }, MaterialXFolder, MaterialXLibrary);
-		if(LoadedLibs.empty())
-		{
-			UInterchangeResultError_Generic* Message = AddMessage<UInterchangeResultError_Generic>();
-			Message->Text = FText::Format(LOCTEXT("MaterialXLibrariesNotFound", "Couldn't load MaterialX libraries from {0}"),
-										  FText::FromString(MaterialXFolder.asString().c_str()));
-			return false;
-		}
-
-		mx::DocumentPtr Document = mx::createDocument();
-		mx::readFromXmlFile(Document, TCHAR_TO_UTF8(*Filename));
-		Document->importLibrary(MaterialXLibrary);
-
-		// Read the document to be sure that the file is valid (meaning all nodes have their nodedef and typedef well-defined)
-		bIsDocumentValid = ValidateDocument(Document, this);
-		if(!bIsDocumentValid)
-		{
-			return false;
-		}
-
-		//Update the document by initializing and reorganizing the different nodes and subgraphs
-		FMaterialXBase::UpdateDocumentRecursively(Document);
-
-		// coming to this point we know for sure that the document is valid
-		for(mx::ElementPtr Elem : Document->traverseTree())
-		{
-			//make sure to read only the current file otherwise we'll process the entire library
-			if(Elem->getActiveSourceUri() != Document->getActiveSourceUri())
-			{
-				continue;
-			}
-
-			mx::NodePtr Node = Elem->asA<mx::Node>();
-
-			if(Node)
-			{
-				bool bIsMaterialShader = Node->getType() == mx::Type::Material;
-				bool bIsLightShader = Node->getType() == mx::Type::LightShader;
-
-				//The entry point is only surfacematerial or lightshader
-				if(bIsMaterialShader || bIsLightShader)
-				{
-					TSharedPtr<FMaterialXBase> ShaderTranslator = FMaterialXManager::GetInstance().GetShaderTranslator(Node->getCategory().c_str(), BaseNodeContainer);
-					if(ShaderTranslator)
-					{
-						ShaderTranslator->Translate(Node);
-					}
-				}
-			}
-		}
-	}
-	catch(std::exception& Exception)
-	{
-		bIsDocumentValid = false;
-		UInterchangeResultError_Generic* Message = AddMessage<UInterchangeResultError_Generic>();
-		Message->Text = FText::Format(LOCTEXT("MaterialXException", "{0}"),
-									  FText::FromString(Exception.what()));
-	}
+	bIsDocumentValid = FMaterialXManager::GetInstance().Translate(Filename, BaseNodeContainer, this);
 
 #endif // WITH_EDITOR
 
diff --git a/Engine/Plugins/Interchange/Runtime/Source/Import/Private/MaterialX/MaterialExpressions/MaterialExpressionDifference.cpp b/Engine/Plugins/Interchange/Runtime/Source/Import/Private/MaterialX/MaterialExpressions/MaterialExpressionDifference.cpp
--- a/Engine/Plugins/Interchange/Runtime/Source/Import/Private/MaterialX/MaterialExpressions/MaterialExpressionDifference.cpp
+++ b/Engine/Plugins/Interchange/Runtime/Source/Import/Private/MaterialX/MaterialExpressions/MaterialExpressionDifference.cpp
@@ -40,7 +40,7 @@ int32 UMaterialExpressionMaterialXDifference::Compile(FMaterialCompiler* Compile
 	int32 IndexAlpha = Alpha.GetTracedInput().Expression ? Alpha.Compile(Compiler) : Compiler->Constant(ConstAlpha);
 
 	int32 IndexB = B.Compile(Compiler);
-	int32 Diff = Compiler->Abs(Compiler->Sub(A.Compile(Compiler), IndexB));
+	int32 Diff = Compiler->Abs(Compiler->Sub(IndexB, A.Compile(Compiler)));
 	return Compiler->Lerp(IndexB, Diff, IndexAlpha);
 }
 
diff --git a/Engine/Plugins/Interchange/Runtime/Source/Import/Private/MaterialX/MaterialExpressions/MaterialExpressionDifference.h b/Engine/Plugins/Interchange/Runtime/Source/Import/Private/MaterialX/MaterialExpressions/MaterialExpressionDifference.h
--- a/Engine/Plugins/Interchange/Runtime/Source/Import/Private/MaterialX/MaterialExpressions/MaterialExpressionDifference.h
+++ b/Engine/Plugins/Interchange/Runtime/Source/Import/Private/MaterialX/MaterialExpressions/MaterialExpressionDifference.h
@@ -10,8 +10,8 @@
  * Blend nodes take two 1-4 channel inputs and apply the same operator to all channels.
  * Blend nodes support an optional float input mix , which can be used
  * to mix the original B value with the result of the blend operation.
- * Operation: abs(A - B)
- * Result: Lerp(B, abs(A - B), Alpha)
+ * Operation: abs(B - A)
+ * Result: Lerp(B, abs(B - A), Alpha)
  */
 UCLASS(MinimalAPI, collapsecategories, hidecategories = Object, meta = (Private))
 class UMaterialExpressionMaterialXDifference : public UMaterialExpression
diff --git a/Engine/Plugins/Interchange/Runtime/Source/Import/Private/MaterialX/MaterialExpressions/MaterialExpressionMinus.cpp b/Engine/Plugins/Interchange/Runtime/Source/Import/Private/MaterialX/MaterialExpressions/MaterialExpressionMinus.cpp
--- a/Engine/Plugins/Interchange/Runtime/Source/Import/Private/MaterialX/MaterialExpressions/MaterialExpressionMinus.cpp
+++ b/Engine/Plugins/Interchange/Runtime/Source/Import/Private/MaterialX/MaterialExpressions/MaterialExpressionMinus.cpp
@@ -40,7 +40,7 @@ int32 UMaterialExpressionMaterialXMinus::Compile(FMaterialCompiler* Compiler, in
 	int32 IndexAlpha = Alpha.GetTracedInput().Expression ? Alpha.Compile(Compiler) : Compiler->Constant(ConstAlpha);
 
 	int32 IndexB = B.Compile(Compiler);
-	int32 Sub = Compiler->Sub(A.Compile(Compiler), IndexB);
+	int32 Sub = Compiler->Sub(IndexB, A.Compile(Compiler));
 	return Compiler->Lerp(IndexB, Sub, IndexAlpha);
 }
 
diff --git a/Engine/Plugins/Interchange/Runtime/Source/Import/Private/MaterialX/MaterialExpressions/MaterialExpressionMinus.h b/Engine/Plugins/Interchange/Runtime/Source/Import/Private/MaterialX/MaterialExpressions/MaterialExpressionMinus.h
--- a/Engine/Plugins/Interchange/Runtime/Source/Import/Private/MaterialX/MaterialExpressions/MaterialExpressionMinus.h
+++ b/Engine/Plugins/Interchange/Runtime/Source/Import/Private/MaterialX/MaterialExpressions/MaterialExpressionMinus.h
@@ -10,8 +10,8 @@
  * Blend nodes take two 1-4 channel inputs and apply the same operator to all channels.
  * Blend nodes support an optional float input mix , which can be used
  * to mix the original B value with the result of the blend operation.
- * Operation: A - B
- * Result: Lerp(B, A - B, Alpha)
+ * Operation: B - A
+ * Result: Lerp(B, B - A, Alpha)
  */
 UCLASS(MinimalAPI, collapsecategories, hidecategories = Object, meta = (Private))
 class UMaterialExpressionMaterialXMinus : public UMaterialExpression
diff --git a/Engine/Plugins/Interchange/Runtime/Source/Import/Private/MaterialX/MaterialExpressions/MaterialExpressionMod.cpp b/Engine/Plugins/Interchange/Runtime/Source/Import/Private/MaterialX/MaterialExpressions/MaterialExpressionMod.cpp
new file mode 100644
--- /dev/null
+++ b/Engine/Plugins/Interchange/Runtime/Source/Import/Private/MaterialX/MaterialExpressions/MaterialExpressionMod.cpp
@@ -0,0 +1,55 @@
+// Copyright Epic Games, Inc. All Rights Reserved.
+#include "MaterialExpressionMod.h"
+#include "MaterialCompiler.h"
+
+#include UE_INLINE_GENERATED_CPP_BY_NAME(MaterialExpressionMod)
+
+#define LOCTEXT_NAMESPACE "MaterialExpressionMaterialXMod"
+
+UMaterialExpressionMaterialXMod::UMaterialExpressionMaterialXMod(const FObjectInitializer& ObjectInitializer)
+	: Super(ObjectInitializer)
+{
+	// Structure to hold one-time initialization
+	struct FConstructorStatics
+	{
+		FText NAME_MaterialX;
+		FConstructorStatics()
+			: NAME_MaterialX(LOCTEXT("MaterialX", "MaterialX"))
+		{}
+	};
+	static FConstructorStatics ConstructorStatics;
+
+#if WITH_EDITORONLY_DATA
+	MenuCategories.Add(ConstructorStatics.NAME_MaterialX);
+#endif
+}
+
+#if WITH_EDITOR
+int32 UMaterialExpressionMaterialXMod::Compile(FMaterialCompiler* Compiler, int32 OutputIndex)
+{
+	if(!A.GetTracedInput().Expression)
+	{
+		return Compiler->Errorf(TEXT("Missing MaterialX Mod input A"));
+	}
+
+	int32 X = A.Compile(Compiler);
+	int32 Y = B.GetTracedInput().Expression ? B.Compile(Compiler) : Compiler->Constant(ConstB);
+
+	return Compiler->Sub(X, Compiler->Mul(Y, Compiler->Floor(Compiler->Div(X, Y))));
+}
+
+void UMaterialExpressionMaterialXMod::GetCaption(TArray<FString>& OutCaptions) const
+{
+	OutCaptions.Add(TEXT("MaterialX Mod"));
+}
+
+void UMaterialExpressionMaterialXMod::GetExpressionToolTip(TArray<FString>& OutToolTip)
+{
+	ConvertToMultilineToolTip(TEXT(R"(The remaining fraction after dividing an incoming input by a value and subtracting the integer portion.
+Unlike UE FMod or Modulo expressions, Mod always returns a non - negative result, matching the interpretation of the GLSL and OSL mod() function(not fmod()).
+This is computed as x - y * floor(x / y). )"), 40, OutToolTip);
+}
+
+#endif
+
+#undef LOCTEXT_NAMESPACE 
\ No newline at end of file
diff --git a/Engine/Plugins/Interchange/Runtime/Source/Import/Private/MaterialX/MaterialExpressions/MaterialExpressionMod.h b/Engine/Plugins/Interchange/Runtime/Source/Import/Private/MaterialX/MaterialExpressions/MaterialExpressionMod.h
new file mode 100644
--- /dev/null
+++ b/Engine/Plugins/Interchange/Runtime/Source/Import/Private/MaterialX/MaterialExpressions/MaterialExpressionMod.h
@@ -0,0 +1,37 @@
+// Copyright Epic Games, Inc. All Rights Reserved.
+#pragma once
+
+#include "CoreMinimal.h"
+#include "Materials/MaterialExpression.h"
+
+#include "MaterialExpressionMod.generated.h"
+
+/**
+ * The remaining fraction after dividing an incoming input by a value and subtracting the integer portion.
+ * Unlike UE FMod or Modulo expressions, Mod always returns a non-negative result, matching the interpretation of the GLSL and OSL mod() function (not fmod()).
+ * This is computed as x - y * floor(x/y). 
+ */
+UCLASS(MinimalAPI, collapsecategories, hidecategories = Object, meta = (Private))
+class UMaterialExpressionMaterialXMod : public UMaterialExpression
+{
+	GENERATED_UCLASS_BODY()
+
+	UPROPERTY()
+	FExpressionInput A;
+
+	UPROPERTY(meta = (RequiredInput = "false", ToolTip = "Defaults to 'ConstB' if not specified"))
+	FExpressionInput B;
+
+	/** only used if B is not hooked up */
+	UPROPERTY(EditAnywhere, Category = MaterialExpressionLinearInterpolate, meta = (OverridingInputProperty = "B"))
+	float ConstB = 1.f;
+
+	//~ Begin UMaterialExpressionMaterialX Interface
+#if WITH_EDITOR
+	virtual int32 Compile(class FMaterialCompiler* Compiler, int32 OutputIndex) override;
+	virtual void GetCaption(TArray<FString>& OutCaptions) const override;
+	virtual void GetExpressionToolTip(TArray<FString>& OutToolTip) override;
+#endif
+	//~ End UMaterialExpressionMaterialX Interface
+};
+
diff --git a/Engine/Plugins/Interchange/Runtime/Source/Import/Private/MaterialX/MaterialXUtils/MaterialXManager.cpp b/Engine/Plugins/Interchange/Runtime/Source/Import/Private/MaterialX/MaterialXUtils/MaterialXManager.cpp
--- a/Engine/Plugins/Interchange/Runtime/Source/Import/Private/MaterialX/MaterialXUtils/MaterialXManager.cpp
+++ b/Engine/Plugins/Interchange/Runtime/Source/Import/Private/MaterialX/MaterialXUtils/MaterialXManager.cpp
@@ -6,6 +6,7 @@
 #include "UObject/SoftObjectPath.h"
 
 #if WITH_EDITOR
+#include "MaterialXFormat/Util.h"
 #include "MaterialX/InterchangeMaterialXDefinitions.h"
 #include "MaterialX/MaterialXUtils/MaterialXOpenPBRSurfaceShader.h"
 #include "MaterialX/MaterialXUtils/MaterialXSurfaceShader.h"
@@ -20,7 +21,9 @@
 
 #include "InterchangeMaterialDefinitions.h"
 #include "InterchangeImportLog.h"
+#include "InterchangeHelper.h"
 
+#define LOCTEXT_NAMESPACE "InterchangeMaterialXManager"
 
 namespace mx = MaterialX;
 
@@ -132,7 +135,7 @@ FMaterialXManager::FMaterialXManager()
 		{mx::Category::Magnitude,    UE::Expressions::Names::Length},
 		{mx::Category::Max,          UE::Expressions::Names::Max},
 		{mx::Category::Min,          UE::Expressions::Names::Min},
-		{mx::Category::Modulo,       UE::Expressions::Names::Fmod},
+		{mx::Category::Modulo,       UE::Expressions::Names::MaterialXMod},
 		{mx::Category::Multiply,     UE::Expressions::Names::Multiply},
 		{mx::Category::Normalize,    UE::Expressions::Names::Normalize},
 		{mx::Category::Power,        UE::Expressions::Names::Power},
@@ -235,8 +238,13 @@ FMaterialXManager::FMaterialXManager()
 		{mx::Category::SrgbDisplayP3ToLinRec709,FMaterialXMaterialFunction{TInPlaceType<FString>{}, UE::MaterialFunctions::Path::MxSrgbDisplayP3ToRec709}},
 		{mx::Category::SrgbTextureToLinRec709,	FMaterialXMaterialFunction{TInPlaceType<FString>{}, UE::MaterialFunctions::Path::MxSrgbToRec709}},
 		// Procedural
+		{mx::Category::Checkerboard,			FMaterialXMaterialFunction{TInPlaceType<FString>{}, UE::MaterialFunctions::Path::MxCheckerboard}},
+		{mx::Category::Circle,					FMaterialXMaterialFunction{TInPlaceType<FString>{}, UE::MaterialFunctions::Path::MxCircle}},
+		{mx::Category::Line,					FMaterialXMaterialFunction{TInPlaceType<FString>{}, UE::MaterialFunctions::Path::MxLine}},
 		{mx::Category::RandomFloat,	            FMaterialXMaterialFunction{TInPlaceType<FString>{}, UE::MaterialFunctions::Path::MxRandomFloat}},
 		{mx::Category::RandomColor,	            FMaterialXMaterialFunction{TInPlaceType<FString>{}, UE::MaterialFunctions::Path::MxRandomColor}},
+		{mx::Category::TiledCircles,			FMaterialXMaterialFunction{TInPlaceType<FString>{}, UE::MaterialFunctions::Path::MxTiledCircles}},
+		{mx::Category::UnifiedNoise3D,			FMaterialXMaterialFunction{TInPlaceType<FString>{}, UE::MaterialFunctions::Path::MxUnifiedNoise3D}},
 		// Adjustment
 		{mx::Category::ColorCorrect,	        FMaterialXMaterialFunction{TInPlaceType<FString>{}, UE::MaterialFunctions::Path::MxColorCorrect}},
 		{mx::Category::HsvAdjust,	            FMaterialXMaterialFunction{TInPlaceType<FString>{}, UE::MaterialFunctions::Path::MxHsvAdjust}},
@@ -246,6 +254,8 @@ FMaterialXManager::FMaterialXManager()
 	, CategoriesToSkip {
 		mx::Category::ACEScgToLinRec709,
 		mx::Category::AdobeRgbToLinRec709,
+		mx::Category::Checkerboard,
+		mx::Category::Circle,
 		mx::Category::ColorCorrect,
 		mx::Category::Contrast,
 		mx::Category::Distance,
@@ -253,6 +263,7 @@ FMaterialXManager::FMaterialXManager()
 		mx::Category::Fractal3D,
 		mx::Category::GoochShade,
 		mx::Category::HsvAdjust,
+		mx::Category::Line,
 		mx::Category::LinAdobeRgbToLinRec709,
 		mx::Category::LinDisplayP3ToLinRec709,
 		mx::Category::Noise2D,
@@ -272,15 +283,24 @@ FMaterialXManager::FMaterialXManager()
 		mx::Category::SrgbDisplayP3ToLinRec709,
 		mx::Category::SrgbTextureToLinRec709,
 		mx::Category::StandardSurface,
+		mx::Category::TiledCircles,
+		mx::Category::UnifiedNoise3D,
 		mx::Category::UsdPreviewSurface
 	}
 	, NodeDefsCategories{
+		mx::Category::Add,
+		mx::Category::Combine2,
+		mx::Category::Combine3,
+		mx::Category::Combine4,
+		mx::Category::Constant,
+		mx::Category::Divide,
 		mx::Category::IfEqual,
 		mx::Category::IfGreater,
 		mx::Category::IfGreaterEq,
 		mx::Category::Max,
 		mx::Category::Min,
-		mx::Category::Modulo,
+		mx::Category::Rotate3D,
+		mx::Category::Sub,
 	}
 	, NodeInputsToRemove{
 		{mx::Category::GeomColor, {"index"}},                          // There's only one set of VertexColor
@@ -316,6 +336,201 @@ FMaterialXManager& FMaterialXManager::GetInstance()
 	return Instance;
 }
 
+
+
+namespace
+{
+	bool ValidateDocument(MaterialX::DocumentPtr Document, const UInterchangeTranslatorBase* Translator = nullptr)
+	{
+		namespace mx = MaterialX;
+
+		if (std::string MaterialXMessage; !Document->validate(&MaterialXMessage))
+		{
+			if(Translator)
+			{
+				UInterchangeResultError_Generic* Message = Translator->AddMessage<UInterchangeResultError_Generic>();
+				Message->Text = FText::Format(LOCTEXT("MaterialXDocumentInvalid", "{0}"),
+											  FText::FromString(MaterialXMessage.c_str()));
+			}
+			else
+			{
+				UE_LOG(LogInterchangeImport, Error, TEXT("%s"), UTF8_TO_TCHAR(MaterialXMessage.c_str()));
+			}
+			return false;
+		}
+
+		for (mx::ElementPtr Elem : Document->traverseTree())
+		{
+			//make sure to read only the current file otherwise we'll process the entire library
+			if (Elem->getActiveSourceUri() != Document->getActiveSourceUri())
+			{
+				continue;
+			}
+
+			mx::NodePtr Node = Elem->asA<mx::Node>();
+
+			if (Node)
+			{
+				// Validate that all nodes in the file are strictly respecting their node definition
+				if (!Node->getNodeDef())
+				{
+					if(Translator)
+					{
+						UInterchangeResultError_Generic* Message = Translator->AddMessage<UInterchangeResultError_Generic>();
+						Message->Text = FText::Format(LOCTEXT("NodeDefNotFound", "<{0}> has no matching NodeDef, aborting import..."),
+													  FText::FromString(Node->getName().c_str()));
+					}
+					else
+					{
+						UE_LOG(LogInterchangeImport, Error, TEXT("<%s> has no matching NodeDef, aborting import..."), UTF8_TO_TCHAR(Node->getName().c_str()));
+					}
+					return false;
+				}
+
+				//verify first if it's not a multioutput node (weird that there's no typedef for it)
+				if (!Node->getTypeDef() && Node->getType() != "multioutput")
+				{
+					if(Translator)
+					{
+						UInterchangeResultError_Generic* Message = Translator->AddMessage<UInterchangeResultError_Generic>();
+						Message->Text = FText::Format(LOCTEXT("TypeDefNotFound", "<{0}> has no matching TypeDef, aborting import..."),
+													  FText::FromString(Node->getName().c_str()));
+					}
+					else
+					{
+						UE_LOG(LogInterchangeImport, Error, TEXT("<%s> has no matching TypeDef, aborting import..."), UTF8_TO_TCHAR(Node->getName().c_str()));
+					}
+					return false;
+				}
+			}
+		}
+
+		return true;
+	}
+}
+
+bool FMaterialXManager::Translate(const FString& Filename, UInterchangeBaseNodeContainer& BaseNodeContainer, const UInterchangeTranslatorBase* Translator)
+{
+	bool bIsDocumentValid = false;
+
+	if (!FPaths::FileExists(Filename))
+	{
+		return false;
+	}
+
+	try
+	{
+		mx::FileSearchPath MaterialXFolder{ TCHAR_TO_UTF8(*FPaths::Combine(
+			FPaths::EngineDir(),
+			TEXT("Binaries"),
+			TEXT("ThirdParty"),
+			TEXT("MaterialX"))) };
+
+		mx::DocumentPtr MaterialXLibrary = mx::createDocument();
+
+		mx::StringSet LoadedLibs = mx::loadLibraries({ mx::Library::Libraries }, MaterialXFolder, MaterialXLibrary);
+		if (LoadedLibs.empty())
+		{
+			if (Translator)
+			{
+				UInterchangeResultError_Generic* Message = Translator->AddMessage<UInterchangeResultError_Generic>();
+				Message->Text = FText::Format(LOCTEXT("MaterialXLibrariesNotFound", "Couldn't load MaterialX libraries from {0}"),
+											  FText::FromString(MaterialXFolder.asString().c_str()));
+			}
+			else
+			{
+				UE_LOG(LogInterchangeImport, Error, TEXT("Couldn't load MaterialX libraries from %s"), UTF8_TO_TCHAR(MaterialXFolder.asString().c_str()));
+			}
+			return false;
+		}
+
+		mx::DocumentPtr Document = mx::createDocument();
+		mx::readFromXmlFile(Document, TCHAR_TO_UTF8(*Filename));
+		Document->importLibrary(MaterialXLibrary);
+
+		bIsDocumentValid = Translate(Document, BaseNodeContainer, Translator);
+	}
+	catch (std::exception& Exception)
+	{
+		bIsDocumentValid = false;
+		if (Translator)
+		{
+			UInterchangeResultError_Generic* Message = Translator->AddMessage<UInterchangeResultError_Generic>();
+			Message->Text = FText::Format(LOCTEXT("MaterialXException", "{0}"),
+										  FText::FromString(Exception.what()));
+		}
+		else
+		{
+			UE_LOG(LogInterchangeImport, Error, TEXT("%s"), UTF8_TO_TCHAR(Exception.what()));
+		}
+	}
+
+	return bIsDocumentValid;
+}
+
+bool FMaterialXManager::Translate(MaterialX::DocumentPtr Document, UInterchangeBaseNodeContainer& BaseNodeContainer, const UInterchangeTranslatorBase* Translator)
+{
+	bool bIsDocumentValid;
+	UE::Interchange::FScopedLambda ScopedLambda([this](){TextureNodeUids.Empty();});
+
+	try
+	{
+		// Read the document to be sure that the file is valid (meaning all nodes have their nodedef and typedef well-defined)
+		bIsDocumentValid = ValidateDocument(Document, Translator);
+		if (!bIsDocumentValid)
+		{
+			return false;
+		}
+
+		//Update the document by initializing and reorganizing the different nodes and subgraphs
+		FMaterialXBase::UpdateDocumentRecursively(Document);
+
+		// coming to this point we know for sure that the document is valid
+		for (mx::ElementPtr Elem : Document->traverseTree())
+		{
+			//make sure to read only the current file otherwise we'll process the entire library
+			if (Elem->getActiveSourceUri() != Document->getActiveSourceUri())
+			{
+				continue;
+			}
+
+			mx::NodePtr Node = Elem->asA<mx::Node>();
+
+			if (Node)
+			{
+				bool bIsMaterialShader = Node->getType() == mx::Type::Material;
+				bool bIsLightShader = Node->getType() == mx::Type::LightShader;
+
+				//The entry point is only surfacematerial or lightshader
+				if (bIsMaterialShader || bIsLightShader)
+				{
+					TSharedPtr<FMaterialXBase> ShaderTranslator = FMaterialXManager::GetInstance().GetShaderTranslator(Node->getCategory().c_str(), BaseNodeContainer);
+					if (ShaderTranslator)
+					{
+						ShaderTranslator->Translate(Node);
+					}
+				}
+			}
+		}
+	}
+	catch (std::exception& Exception)
+	{
+		bIsDocumentValid = false;
+		if (Translator)
+		{
+			UInterchangeResultError_Generic* Message = Translator->AddMessage<UInterchangeResultError_Generic>();
+			Message->Text = FText::Format(LOCTEXT("MaterialXException", "{0}"),
+										  FText::FromString(Exception.what()));
+		}
+		else
+		{
+			UE_LOG(LogInterchangeImport, Error, TEXT("%s"), UTF8_TO_TCHAR(Exception.what()));
+		}
+	}
+
+	return bIsDocumentValid;
+}
+
 const FString* FMaterialXManager::FindMatchingInput(const FString& CategoryKey, const FString& InputKey, const FString& NodeGroup) const
 {
 	return MatchingInputNames.Find({ FKeyCategoryNodegroup{ CategoryKey, NodeGroup}, InputKey });
@@ -405,6 +620,25 @@ void FMaterialXManager::RemoveInputs(MaterialX::NodePtr Node) const
 	}
 }
 
+FString FMaterialXManager::FindOrAddTextureNodeUid(const FString& TexturePath)
+{
+	FString TextureNodeUid;
+	if (const FString* TextureUidPtr = TextureNodeUids.Find(TexturePath))
+	{
+		TextureNodeUid =*TextureUidPtr;
+	}
+	else
+	{
+		FSHA1 SHA1;
+		SHA1.UpdateWithString(*TexturePath, TexturePath.Len());
+		FSHAHash Hash = SHA1.Finalize();
+		TextureNodeUid = Hash.ToString();
+		TextureNodeUid = TEXT("\\Texture\\") + TextureNodeUid + TEXT("\\") + FPaths::GetCleanFilename(TexturePath);
+		TextureNodeUids.Emplace(TexturePath, TextureNodeUid);
+	}
+	return TextureNodeUid;
+}
+
 void FMaterialXManager::AddInputsFromNodeDef(MaterialX::NodePtr Node) const
 {
 	if (NodeDefsCategories.find(Node->getCategory()) != NodeDefsCategories.cend())
@@ -468,6 +702,11 @@ namespace UE::Interchange::MaterialX
 																UE::MaterialFunctions::Path::MxColorCorrect,
 																UE::MaterialFunctions::Path::MxHsvAdjust,
 																UE::MaterialFunctions::Path::MxGoochShade,
+																UE::MaterialFunctions::Path::MxCircle,
+																UE::MaterialFunctions::Path::MxCheckerboard,
+																UE::MaterialFunctions::Path::MxTiledCircles,
+																UE::MaterialFunctions::Path::MxUnifiedNoise3D,
+																UE::MaterialFunctions::Path::MxLine,
 															  });
 
 		return bPackagesLoaded;
@@ -476,3 +715,5 @@ namespace UE::Interchange::MaterialX
 #endif
 	}
 }
+
+#undef LOCTEXT_NAMESPACE 
\ No newline at end of file
diff --git a/Engine/Plugins/Interchange/Runtime/Source/Import/Private/MaterialX/MaterialXUtils/MaterialXSurfaceShaderAbstract.cpp b/Engine/Plugins/Interchange/Runtime/Source/Import/Private/MaterialX/MaterialXUtils/MaterialXSurfaceShaderAbstract.cpp
--- a/Engine/Plugins/Interchange/Runtime/Source/Import/Private/MaterialX/MaterialXUtils/MaterialXSurfaceShaderAbstract.cpp
+++ b/Engine/Plugins/Interchange/Runtime/Source/Import/Private/MaterialX/MaterialXUtils/MaterialXSurfaceShaderAbstract.cpp
@@ -54,8 +54,7 @@ bool FMaterialXSurfaceShaderAbstract::AddAttribute(MaterialX::InputPtr Input, co
 		}
 		else if(Input->getType() == mx::Type::Vector2)
 		{
-			FLinearColor Vector = GetVector(Input);
-			return ShaderNode->AddVector2Attribute(UInterchangeShaderPortsAPI::MakeInputValueKey(InputChannelName), FVector2f{ Vector.R, Vector.G });
+			return AddVector2Attribute(Input, InputChannelName, ShaderNode, FVector2f{ std::numeric_limits<float>::max() ,std::numeric_limits<float>::max() }, OutputIndex);
 		}
 		else if(Input->getType() == mx::Type::Vector3 || Input->getType() == mx::Type::Vector4)
 		{
@@ -217,6 +216,31 @@ bool FMaterialXSurfaceShaderAbstract::AddVectorAttribute(MaterialX::InputPtr Inp
 	return false;
 }
 
+bool FMaterialXSurfaceShaderAbstract::AddVector2Attribute(MaterialX::InputPtr Input, const FString& InputChannelName, UInterchangeShaderNode* ShaderNode, const FVector2f& DefaultValue, int32 OutputIndex)
+{
+	using namespace UE::Interchange::Materials::Standard::Nodes;
+	if (Input && Input->hasValue())
+	{
+		const FLinearColor Value = GetVector(Input);
+		
+		if (!FVector2f{ Value.R,Value.B }.Equals(DefaultValue))
+		{
+			// The parent is either a node, or it's an interfacename and we just take the name of the input
+			mx::NodePtr Node = Input->getParent()->asA<mx::Node>();
+			FString NodeName = Node ? Node->getName().c_str() + FString{ TEXT("_") } : FString{};
+			NodeName += Input->getName().c_str();
+
+			UInterchangeShaderNode* VectorParameterNode = CreateShaderNode(Input, NodeName, VectorParameter::Name.ToString());
+			VectorParameterNode->AddLinearColorAttribute(UInterchangeShaderPortsAPI::MakeInputParameterKey(VectorParameter::Attributes::DefaultValue.ToString()), Value);
+			UInterchangeShaderNode* VectorMaskNode = CreateMaskShaderNode(0b1100, Input, NodeName + TEXT("_Vector2"));
+			UInterchangeShaderPortsAPI::ConnectOuputToInputByIndex(VectorMaskNode, Mask::Inputs::Input.ToString(), VectorParameterNode->GetUniqueID(), OutputIndex);
+			return UInterchangeShaderPortsAPI::ConnectOuputToInputByIndex(ShaderNode, InputChannelName, VectorMaskNode->GetUniqueID(), OutputIndex);
+		}
+	}
+
+	return false;
+}
+
 bool FMaterialXSurfaceShaderAbstract::ConnectNodeGraphOutputToInput(MaterialX::InputPtr InputToNodeGraph, UInterchangeShaderNode* ShaderNode, const FString& ParentInputName)
 {
 	using namespace UE::Interchange::Materials::Standard::Nodes;
diff --git a/Engine/Plugins/Interchange/Runtime/Source/Import/Public/MaterialX/MaterialXUtils/MaterialXManager.h b/Engine/Plugins/Interchange/Runtime/Source/Import/Public/MaterialX/MaterialXUtils/MaterialXManager.h
--- a/Engine/Plugins/Interchange/Runtime/Source/Import/Public/MaterialX/MaterialXUtils/MaterialXManager.h
+++ b/Engine/Plugins/Interchange/Runtime/Source/Import/Public/MaterialX/MaterialXUtils/MaterialXManager.h
@@ -15,6 +15,7 @@ class FMaterialXBase;
 class UInterchangeBaseNodeContainer;
 class UInterchangeShaderNode;
 class FMaterialXSurfaceShaderAbstract;
+class UInterchangeTranslatorBase;
 
 class INTERCHANGEIMPORT_API FMaterialXManager
 {
@@ -24,6 +25,19 @@ public:
 
 	static FMaterialXManager& GetInstance();
 
+	/** 
+	 * Translate a MaterialX file
+	 * @param Filename - Name of the MaterialX file to translate
+	 * @param BaseNodeContainer - Node Container
+	 * @param Translator - The Translator if this function is called from a Translaotr, used to log errors, otherwise the errors are logged in the default output
+	 * 
+	 * @return true if the file has been successfully translated
+	 */
+	bool Translate(const FString& Filename, UInterchangeBaseNodeContainer& BaseNodeContainer, const UInterchangeTranslatorBase* Translator = nullptr);
+
+	/** Translate a MaterialX Document */
+	bool Translate(MaterialX::DocumentPtr Document, UInterchangeBaseNodeContainer& BaseNodeContainer, const UInterchangeTranslatorBase* Translator = nullptr);
+
 	/** Find a corresponding Material Expression Input given a (Category [NodeGroup], Input) pair*/
 	const FString* FindMatchingInput(const FString& CategoryKey, const FString& InputKey, const FString& NodeGroup = {}) const;
 
@@ -48,6 +62,9 @@ public:
 	/** Remove the inputs of a node with no match in UE Material Expressions*/
 	void RemoveInputs(MaterialX::NodePtr Node) const;
 
+	/** Find or Add a Texture Node UID, if not found it creates a hash of the path*/
+	FString FindOrAddTextureNodeUid(const FString& TexturePath);
+
 	/**
 	 * Add all the inputs from the nodedef in case the node doesn't have them all (we don't have a 1:1 match on the default values otherwise)
 	 * For exemple 'min' defaults to (0, 0) whereas UE to (0, 1)
@@ -115,6 +132,14 @@ private:
 	/** Given a MaterialX node category, return the UE material function, used for BSDF nodes*/
 	TMap<FString, FMaterialXMaterialFunction> MatchingMaterialFunctions;
 
+	/**
+	 * Mapping between texture fullpath and their UID
+	 * used to allow same name textures with different path,
+	 * but avoid duplicating textures across different materials in the same file (for example the chess_set)
+	 * The map is reset after each translate
+	 */
+	TMap<FString, FString> TextureNodeUids;
+
 	/** Categories of nodes to skip during the phase of flattening the subgraphs. Basically if a node has a nodegraph it will be processed as is instead of its nodegraph*/
 	std::unordered_set<std::string> CategoriesToSkip;
 
diff --git a/Engine/Plugins/Interchange/Runtime/Source/Import/Public/MaterialX/MaterialXUtils/MaterialXSurfaceShaderAbstract.h b/Engine/Plugins/Interchange/Runtime/Source/Import/Public/MaterialX/MaterialXUtils/MaterialXSurfaceShaderAbstract.h
--- a/Engine/Plugins/Interchange/Runtime/Source/Import/Public/MaterialX/MaterialXUtils/MaterialXSurfaceShaderAbstract.h
+++ b/Engine/Plugins/Interchange/Runtime/Source/Import/Public/MaterialX/MaterialXUtils/MaterialXSurfaceShaderAbstract.h
@@ -132,7 +132,7 @@ protected:
 	 * @return true if the attribute was successfully added.
 	 */
 	bool AddVectorAttribute(MaterialX::InputPtr Input, const FString& InputChannelName, UInterchangeShaderNode* ShaderNode, const FVector4f& DefaultValue = FVector4f{ std::numeric_limits<float>::max(), std::numeric_limits<float>::max(), std::numeric_limits<float>::max() }, int32 OutputIndex = 0);
-
+	bool AddVector2Attribute(MaterialX::InputPtr Input, const FString& InputChannelName, UInterchangeShaderNode* ShaderNode, const FVector2f& DefaultValue = FVector2f{ std::numeric_limits<float>::max(), std::numeric_limits<float>::max()}, int32 OutputIndex = 0);
 	/**
 	 * Connect an output either from a node name or a node graph from a MaterialX input to the shader node.
 	 * 
@@ -399,11 +399,13 @@ protected:
 				}
 
 				const FString Filename = FPaths::GetCleanFilename(Filepath);
-				FString TextureNodeUID = TEXT("\\Texture\\") + Filename;
+				
+				FString TextureNodeUID = FMaterialXManager::GetInstance().FindOrAddTextureNodeUid(Filepath);				
+				
 				// We need to duplicate the texture for a geomprop, at the end the texture will be the baking texture
 				if (bGeomPropImage)
 				{
-					TextureNodeUID += TEXT("_") + SurfaceMaterialName + Node->getName().c_str();
+					TextureNodeUID = FPaths::GetPath(TextureNodeUID) + TEXT("\\") + SurfaceMaterialName + TEXT("_") + Node->getName().c_str();
 				}
 
 				//Only add the TextureNode once
