Removed check on GLTF material assets on each creation of a UInterchangeGLTFTranslator object
- Moved  check on GLTF material assets in GetSupportedFormats method
- Made check on GLTF material assets to happen only once

#rnx

diff --git a/Engine/Plugins/Interchange/Runtime/Source/Import/Private/Gltf/InterchangeGltfTranslator.cpp b/Engine/Plugins/Interchange/Runtime/Source/Import/Private/Gltf/InterchangeGltfTranslator.cpp
--- a/Engine/Plugins/Interchange/Runtime/Source/Import/Private/Gltf/InterchangeGltfTranslator.cpp
+++ b/Engine/Plugins/Interchange/Runtime/Source/Import/Private/Gltf/InterchangeGltfTranslator.cpp
@@ -472,13 +472,23 @@ EInterchangeTranslatorAssetType UInterchangeGLTFTranslator::GetSupportedAssetTyp
 
 TArray<FString> UInterchangeGLTFTranslator::GetSupportedFormats() const
 {
-	TArray<FString> GltfExtensions;
+	static bool bRequiredPackageLoaded = []() -> bool
+		{
+			if (!UE::Interchange::GLTFMaterials::AreRequiredPackagesLoaded())
+			{
+				UE_LOG(LogInterchangeImport, Warning, TEXT("UInterchangeGLTFPipeline: Some required packages are missing. Material import might be wrong."));
+				return false;
+			}
+
+			return true;
+		}();
 
-	GltfExtensions.Reserve(2);
-	GltfExtensions.Add(TEXT("gltf;GL Transmission Format"));
-	GltfExtensions.Add(TEXT("glb;GL Transmission Format (Binary)"));
+	static TArray<FString> GltfExtensions{
+		TEXT("gltf;GL Transmission Format"),
+		TEXT("glb;GL Transmission Format (Binary)")
+	};
 
-	return GltfExtensions;
+	return bRequiredPackageLoaded ? GltfExtensions : TArray<FString>();
 }
 
 bool UInterchangeGLTFTranslator::Translate( UInterchangeBaseNodeContainer& NodeContainer ) const
@@ -1607,11 +1617,6 @@ UInterchangeGLTFTranslator::UInterchangeGLTFTranslator()
 	{
 		bRenderSettingsClearCoatEnableSecondNormal = GetDefault<URendererSettings>()->bClearCoatEnableSecondNormal != 0;
 	}
-
-	if (!UE::Interchange::GLTFMaterials::AreRequiredPackagesLoaded())
-	{
-		UE_LOG(LogInterchangeImport, Warning, TEXT("UInterchangeGLTFPipeline: Some required packages are missing. Material import might be wrong."));
-	}
 }
 
 TOptional<UE::Interchange::FImportLightProfile> UInterchangeGLTFTranslator::GetLightProfilePayloadData(const FString& PayloadKey, TOptional<FString>& AlternateTexturePath) const
