Interchange FBX:
- change random color to default color in case no suitable material was found

diff --git a/Engine/Plugins/Interchange/Runtime/Source/Parsers/Fbx/Private/FbxMaterial.cpp b/Engine/Plugins/Interchange/Runtime/Source/Parsers/Fbx/Private/FbxMaterial.cpp
--- a/Engine/Plugins/Interchange/Runtime/Source/Parsers/Fbx/Private/FbxMaterial.cpp
+++ b/Engine/Plugins/Interchange/Runtime/Source/Parsers/Fbx/Private/FbxMaterial.cpp
@@ -477,9 +477,9 @@ namespace UE
 				if (InputNames.Num() == 0)
 				{
 					FLinearColor BaseColor;
-					BaseColor.R = 0.5f + 0.5f * FMath::FRand();
-					BaseColor.G = 0.5f + 0.5f * FMath::FRand();
-					BaseColor.B = 0.5f + 0.5f * FMath::FRand();
+					BaseColor.R = 0.7f;
+					BaseColor.G = 0.7f;
+					BaseColor.B = 0.7f;
 
 					const FString InputValueKey = UInterchangeShaderPortsAPI::MakeInputValueKey(Phong::Parameters::DiffuseColor.ToString());
 					ShaderGraphNode->AddLinearColorAttribute(InputValueKey, BaseColor);
