Fixed potential crashes when predefined assets are not loaded in the Game thread
- Solution: avoid calling the method to load predefined assets when not in the Game thread.
- Note: This has currently no side-effects since the translator's  GetSupportedFormats()method  is only called from another thread when laoding texture's payload. However, a dedicated mechanism should be implemented to allow translors to initialize themselves in the Game thread. Opening a task for that.

#rnx

diff --git a/Engine/Plugins/Interchange/Extensions/OpenUSD/Source/Import/Private/InterchangeUsdTranslator.cpp b/Engine/Plugins/Interchange/Extensions/OpenUSD/Source/Import/Private/InterchangeUsdTranslator.cpp
--- a/Engine/Plugins/Interchange/Extensions/OpenUSD/Source/Import/Private/InterchangeUsdTranslator.cpp
+++ b/Engine/Plugins/Interchange/Extensions/OpenUSD/Source/Import/Private/InterchangeUsdTranslator.cpp
@@ -4293,8 +4293,11 @@ TArray<FString> UInterchangeUSDTranslator::GetSupportedFormats() const
 	TArray<FString> Extensions;
 	if (GInterchangeEnableUSDImport)
 	{
-		// ensure that MaterialX material functions are loaded in the Game Thread
-		UE::Interchange::MaterialX::AreMaterialFunctionPackagesLoaded();
+		if (IsInGameThread())
+		{
+			// ensure that MaterialX material functions are loaded in the Game Thread
+			UE::Interchange::MaterialX::AreMaterialFunctionPackagesLoaded();
+		}
 		FModuleManager::Get().LoadModuleChecked(TEXT("UnrealUSDWrapper"));
 		UnrealUSDWrapper::AddUsdImportFileFormatDescriptions(Extensions);
 	}
diff --git a/Engine/Plugins/Interchange/Runtime/Source/Import/Private/Gltf/InterchangeGltfTranslator.cpp b/Engine/Plugins/Interchange/Runtime/Source/Import/Private/Gltf/InterchangeGltfTranslator.cpp
--- a/Engine/Plugins/Interchange/Runtime/Source/Import/Private/Gltf/InterchangeGltfTranslator.cpp
+++ b/Engine/Plugins/Interchange/Runtime/Source/Import/Private/Gltf/InterchangeGltfTranslator.cpp
@@ -472,23 +472,26 @@ EInterchangeTranslatorAssetType UInterchangeGLTFTranslator::GetSupportedAssetTyp
 
 TArray<FString> UInterchangeGLTFTranslator::GetSupportedFormats() const
 {
-	static bool bRequiredPackageLoaded = []() -> bool
-		{
-			if (!UE::Interchange::GLTFMaterials::AreRequiredPackagesLoaded())
+	if (IsInGameThread())
+	{
+		static bool bRequiredPackageLoaded = []() -> bool
 			{
-				UE_LOG(LogInterchangeImport, Warning, TEXT("UInterchangeGLTFPipeline: Some required packages are missing. Material import might be wrong."));
-				return false;
-			}
+				if (!UE::Interchange::GLTFMaterials::AreRequiredPackagesLoaded())
+				{
+					UE_LOG(LogInterchangeImport, Warning, TEXT("UInterchangeGLTFPipeline: Some required packages are missing. Material import might be wrong."));
+					return false;
+				}
 
-			return true;
-		}();
+				return true;
+			}();
+	}
 
 	static TArray<FString> GltfExtensions{
 		TEXT("gltf;GL Transmission Format"),
 		TEXT("glb;GL Transmission Format (Binary)")
 	};
 
-	return bRequiredPackageLoaded ? GltfExtensions : TArray<FString>();
+	return GltfExtensions;
 }
 
 bool UInterchangeGLTFTranslator::Translate( UInterchangeBaseNodeContainer& NodeContainer ) const
