Fixed reordering issue of material slots on reimport

#rnx

diff --git a/Engine/Plugins/Interchange/Runtime/Source/Import/Private/Mesh/InterchangeStaticMeshFactory.cpp b/Engine/Plugins/Interchange/Runtime/Source/Import/Private/Mesh/InterchangeStaticMeshFactory.cpp
--- a/Engine/Plugins/Interchange/Runtime/Source/Import/Private/Mesh/InterchangeStaticMeshFactory.cpp
+++ b/Engine/Plugins/Interchange/Runtime/Source/Import/Private/Mesh/InterchangeStaticMeshFactory.cpp
@@ -53,103 +53,6 @@ UClass* UInterchangeStaticMeshFactory::GetFactoryClass() const
 	return UStaticMesh::StaticClass();
 }
 
-namespace UE::Interchange::Private::StaticMesh
-{
-	void ReorderMaterialSlotToBaseLod(UStaticMesh* StaticMesh)
-	{
-#if WITH_EDITOR
-		if (!StaticMesh || !StaticMesh->IsMeshDescriptionValid(0))
-		{
-			return;
-		}
-
-		TArray<FStaticMaterial>& Materials = StaticMesh->GetStaticMaterials();
-		if (Materials.Num() < 2)
-		{
-			return;
-		}
-
-		TArray<int32> RemapMaterialIndexes;
-		RemapMaterialIndexes.Reserve(Materials.Num());
-		for (int32 MaterialIndex = 0; MaterialIndex < Materials.Num(); ++MaterialIndex)
-		{
-			RemapMaterialIndexes.Add(INDEX_NONE);
-		}
-		TArray<FStaticMaterial> ReorderMaterialArray;
-		ReorderMaterialArray.Reserve(Materials.Num());
-
-		for (int32 LodIndex = 0; LodIndex < StaticMesh->GetNumSourceModels(); ++LodIndex)
-		{
-			const FMeshDescription* LodMeshDescription = StaticMesh->GetMeshDescription(LodIndex);
-			if (!LodMeshDescription || !ensure(!LodMeshDescription->NeedsCompact()))
-			{
-				if (LodIndex == 0)
-				{
-					return; //Lod 0 must always participate in the re-order, return if we can't use it
-				}
-				continue;
-			}
-
-			FStaticMeshConstAttributes StaticMeshAttributes(*LodMeshDescription);
-			TPolygonGroupAttributesConstRef<FName> SlotNames = StaticMeshAttributes.GetPolygonGroupMaterialSlotNames();
-			for (FPolygonGroupID PolygonGroupID : LodMeshDescription->PolygonGroups().GetElementIDs())
-			{
-				FName ImportMaterialName = SlotNames[PolygonGroupID];
-				for (int32 MaterialIndex = 0; MaterialIndex < Materials.Num(); ++MaterialIndex)
-				{
-					FName MaterialName = Materials[MaterialIndex].ImportedMaterialSlotName;
-					if (RemapMaterialIndexes[MaterialIndex] != INDEX_NONE)
-					{
-						//This material was already matched
-
-						//If the name match say we found the match)
-						if (MaterialName == ImportMaterialName)
-						{
-							break;
-						}
-						continue;
-					}
-					int32& RemapIndex = RemapMaterialIndexes[MaterialIndex];
-					
-					if (MaterialName == ImportMaterialName)
-					{
-						RemapIndex = ReorderMaterialArray.Add(Materials[MaterialIndex]);
-						break;
-					}
-				}
-			}
-		}
-		//Custom LOD can add materials, so we add them at the end of the material slots
-		for (int32 MaterialIndex = 0; MaterialIndex < Materials.Num(); ++MaterialIndex)
-		{
-			if (RemapMaterialIndexes[MaterialIndex] == INDEX_NONE)
-			{
-				RemapMaterialIndexes[MaterialIndex] = ReorderMaterialArray.Add(Materials[MaterialIndex]);
-			}
-		}
-
-		check(ReorderMaterialArray.Num() == Materials.Num());
-
-		//Reorder the static mesh material slot array
-		Materials = ReorderMaterialArray;
-
-		//Fix all the section info map MaterialIndex with the RemapMaterialIndexes
-		FMeshSectionInfoMap& SectionInfoMap = StaticMesh->GetSectionInfoMap();
-		const int32 LodCount = StaticMesh->GetNumSourceModels();
-		for (int32 LodIndex = 0; LodIndex < LodCount; ++LodIndex)
-		{
-			const int32 SectionCount = SectionInfoMap.GetSectionNumber(LodIndex);
-			for (int32 SectionIndex = 0; SectionIndex < SectionCount; ++SectionIndex)
-			{
-				FMeshSectionInfo SectionInfo = SectionInfoMap.Get(LodIndex, SectionIndex);
-				SectionInfo.MaterialIndex = RemapMaterialIndexes[SectionInfo.MaterialIndex];
-				SectionInfoMap.Set(LodIndex, SectionIndex, SectionInfo);
-			}
-		}
-#endif // WITH_EDITOR
-	}
-} //ns UE::Interchange::Private::StaticMesh
-
 void UInterchangeStaticMeshFactory::CreatePayloadTasks(const FImportAssetObjectParams& Arguments, bool bAsync, TArray<TSharedPtr<UE::Interchange::FInterchangeTaskBase>>& PayloadTasks)
 {
 	TRACE_CPUPROFILER_EVENT_SCOPE(UInterchangeStaticMeshFactory::CreatePayloadTasks);
@@ -993,13 +896,12 @@ UInterchangeFactoryBase::FImportAssetResult UInterchangeStaticMeshFactory::EndIm
 	FMeshBudgetProjectSettingsUtils::SetLodGroupForStaticMesh(StaticMesh);
 #endif
 
+#if WITH_EDITOR
 	if (bReimport)
 	{
-		UE::Interchange::Private::StaticMesh::ReorderMaterialSlotToBaseLod(StaticMesh);
-#if WITH_EDITOR
 		UStaticMesh::RemoveUnusedMaterialSlots(StaticMesh);
-#endif
 	}
+#endif
 
 	ImportAssetResult.ImportedObject = StaticMesh;
 	return ImportAssetResult;
