[Interchange-glTF] Editor crashes upon importing a corrupt file.

#rnx

diff --git a/Engine/Plugins/Interchange/Runtime/Source/Parsers/GLTFCore/Private/GLTF/GLTFReader.cpp b/Engine/Plugins/Interchange/Runtime/Source/Parsers/GLTFCore/Private/GLTF/GLTFReader.cpp
--- a/Engine/Plugins/Interchange/Runtime/Source/Parsers/GLTFCore/Private/GLTF/GLTFReader.cpp
+++ b/Engine/Plugins/Interchange/Runtime/Source/Parsers/GLTFCore/Private/GLTF/GLTFReader.cpp
@@ -1548,7 +1548,7 @@ namespace GLTF
 					continue;
 				}
 
-				//0. Group RootJointNodes by ParentIndices
+				//0. Group RootJointNodes by ParentIndices or by Index if ParentIndex is -1
 				TMap<int32, FRootJoints> ParentToRootJointIndices;
 				for (size_t JointIndex = 0; JointIndex < Skin.Joints.Num(); JointIndex++)
 				{
@@ -1566,7 +1566,7 @@ namespace GLTF
 						}
 						else
 						{
-							RootJoints = &ParentToRootJointIndices.Add(JointNode.ParentIndex);
+							RootJoints = &ParentToRootJointIndices.Add(JointNode.ParentIndex != INDEX_NONE ? JointNode.ParentIndex : JointNode.Index);
 							RootJoints->Indices.Add(JointNode.Index);
 							RootJoints->ParentDistanceFromRoot = GetRootDistance(JointNode);
 						}
