[Interchange]
* Moved dll storage to methods/variables instead of on types. This is needed to get merged modules working

diff --git a/Engine/Plugins/Interchange/Runtime/Source/Dispatcher/Public/InterchangeDispatcher.h b/Engine/Plugins/Interchange/Runtime/Source/Dispatcher/Public/InterchangeDispatcher.h
--- a/Engine/Plugins/Interchange/Runtime/Source/Dispatcher/Public/InterchangeDispatcher.h
+++ b/Engine/Plugins/Interchange/Runtime/Source/Dispatcher/Public/InterchangeDispatcher.h
@@ -7,30 +7,32 @@
 #include "InterchangeDispatcherTask.h"
 #include "InterchangeWorkerHandler.h"
 
+#define UE_API INTERCHANGEDISPATCHER_API
+
 namespace UE
 {
 	namespace Interchange
 	{
 
 		// Handle a list of tasks, and a set of external workers to consume them.
-		class INTERCHANGEDISPATCHER_API FInterchangeDispatcher
+		class  FInterchangeDispatcher
 		{
 		public:
-			FInterchangeDispatcher(const FString& InResultFolder);
+			UE_API FInterchangeDispatcher(const FString& InResultFolder);
 			~FInterchangeDispatcher() { TerminateProcess(); }
 
-			int32 AddTask(const FString& JsonDescription);
-			int32 AddTask(const FString& JsonDescription, FInterchangeDispatcherTaskCompleted TaskCompledDelegate);
-			TOptional<FTask> GetNextTask();
-			void SetTaskState(int32 TaskIndex, ETaskState TaskState, const FString& JsonResult, const TArray<FString>& JSonMessages);
-			void GetTaskState(int32 TaskIndex, ETaskState& TaskState, double& TaskRunningStateStartTime);
-			void GetTaskState(int32 TaskIndex, ETaskState& TaskState, FString& JsonResult, TArray<FString>& JSonMessages);
+			UE_API int32 AddTask(const FString& JsonDescription);
+			UE_API int32 AddTask(const FString& JsonDescription, FInterchangeDispatcherTaskCompleted TaskCompledDelegate);
+			UE_API TOptional<FTask> GetNextTask();
+			UE_API void SetTaskState(int32 TaskIndex, ETaskState TaskState, const FString& JsonResult, const TArray<FString>& JSonMessages);
+			UE_API void GetTaskState(int32 TaskIndex, ETaskState& TaskState, double& TaskRunningStateStartTime);
+			UE_API void GetTaskState(int32 TaskIndex, ETaskState& TaskState, FString& JsonResult, TArray<FString>& JSonMessages);
 
-			void StartProcess();
-			void StopProcess(bool bBlockUntilTerminated);
-			void TerminateProcess();
-			void WaitAllTaskToCompleteExecution();
-			bool IsOver();
+			UE_API void StartProcess();
+			UE_API void StopProcess(bool bBlockUntilTerminated);
+			UE_API void TerminateProcess();
+			UE_API void WaitAllTaskToCompleteExecution();
+			UE_API bool IsOver();
 
 			void SetInterchangeWorkerFatalError(FString& ErrorMessage)
 			{
@@ -42,8 +44,8 @@ namespace UE
 				return InterchangeWorkerFatalError;
 			}
 
-			bool IsInterchangeWorkerRunning();
-			static bool IsInterchangeWorkerAvailable();
+			UE_API bool IsInterchangeWorkerRunning();
+			static UE_API bool IsInterchangeWorkerAvailable();
 		private:
 			void SpawnHandler();
 			bool IsHandlerAlive();
@@ -67,3 +69,5 @@ namespace UE
 
 	} //ns Interchange
 }//ns UE
+
+#undef UE_API
diff --git a/Engine/Plugins/Interchange/Runtime/Source/Dispatcher/Public/InterchangeDispatcherTask.h b/Engine/Plugins/Interchange/Runtime/Source/Dispatcher/Public/InterchangeDispatcherTask.h
--- a/Engine/Plugins/Interchange/Runtime/Source/Dispatcher/Public/InterchangeDispatcherTask.h
+++ b/Engine/Plugins/Interchange/Runtime/Source/Dispatcher/Public/InterchangeDispatcherTask.h
@@ -49,7 +49,7 @@ namespace UE
 		/**
 		 * Json cmd helper to be able to read and write a FTask::JsonDescription
 		 */
-		class INTERCHANGEDISPATCHER_API IJsonCmdBase
+		class IJsonCmdBase
 		{
 		public:
 			virtual ~IJsonCmdBase() = default;
@@ -84,7 +84,7 @@ namespace UE
 			bool bIsDataInitialize = false;
 		};
 
-		class INTERCHANGEDISPATCHER_API FJsonLoadSourceCmd : public IJsonCmdBase
+		class FJsonLoadSourceCmd : public IJsonCmdBase
 		{
 		public:
 			FJsonLoadSourceCmd()
@@ -121,8 +121,8 @@ namespace UE
 				return TranslatorID;
 			}
 
-			virtual FString ToJson() const;
-			virtual bool FromJson(const FString& JsonString);
+			INTERCHANGEDISPATCHER_API virtual FString ToJson() const;
+			INTERCHANGEDISPATCHER_API virtual bool FromJson(const FString& JsonString);
 
 			FString GetSourceFilename() const
 			{
@@ -192,7 +192,7 @@ namespace UE
 			/**
 			 * Use this class helper to create the cmd result json string and to read it
 			 */
-			class INTERCHANGEDISPATCHER_API JsonResultParser
+			class JsonResultParser
 			{
 			public:
 				FString GetResultFilename() const
@@ -203,8 +203,8 @@ namespace UE
 				{
 					ResultFilename = InResultFilename;
 				}
-				FString ToJson() const;
-				bool FromJson(const FString& JsonString);
+				INTERCHANGEDISPATCHER_API FString ToJson() const;
+				INTERCHANGEDISPATCHER_API bool FromJson(const FString& JsonString);
 
 				static FString GetResultFilenameJsonKey()
 				{
@@ -224,7 +224,7 @@ namespace UE
 			bool bKeepFbxNamespace = false;
 		};
 
-		class INTERCHANGEDISPATCHER_API FJsonFetchPayloadCmd : public IJsonCmdBase
+		class FJsonFetchPayloadCmd : public IJsonCmdBase
 		{
 		public:
 			FJsonFetchPayloadCmd()
@@ -253,8 +253,8 @@ namespace UE
 				return TranslatorID;
 			}
 
-			virtual FString ToJson() const;
-			virtual bool FromJson(const FString& JsonString);
+			INTERCHANGEDISPATCHER_API virtual FString ToJson() const;
+			INTERCHANGEDISPATCHER_API virtual bool FromJson(const FString& JsonString);
 
 			FString GetPayloadKey() const
 			{
@@ -272,7 +272,7 @@ namespace UE
 			/**
 			 * Use this class helper to create the cmd result json string and to read it
 			 */
-			class INTERCHANGEDISPATCHER_API JsonResultParser
+			class JsonResultParser
 			{
 			public:
 				FString GetResultFilename() const
@@ -283,8 +283,8 @@ namespace UE
 				{
 					ResultFilename = InResultFilename;
 				}
-				FString ToJson() const;
-				bool FromJson(const FString& JsonString);
+				INTERCHANGEDISPATCHER_API FString ToJson() const;
+				INTERCHANGEDISPATCHER_API bool FromJson(const FString& JsonString);
 
 				static FString GetResultFilenameJsonKey()
 				{
@@ -301,7 +301,7 @@ namespace UE
 		};
 
 		//Mesh payload require a transform to bake the mesh to avoid degenerate triangle when importing a small mesh scale by a scene node.
-		class INTERCHANGEDISPATCHER_API FJsonFetchMeshPayloadCmd : public FJsonFetchPayloadCmd
+		class FJsonFetchMeshPayloadCmd : public FJsonFetchPayloadCmd
 		{
 		public:
 			FJsonFetchMeshPayloadCmd()
@@ -316,8 +316,8 @@ namespace UE
 				, MeshGlobalTransform(InMeshGlobalTransform)
 			{}
 
-			virtual FString ToJson() const override;
-			virtual bool FromJson(const FString& JsonString) override;
+			INTERCHANGEDISPATCHER_API virtual FString ToJson() const override;
+			INTERCHANGEDISPATCHER_API virtual bool FromJson(const FString& JsonString) override;
 
 			FTransform GetMeshGlobalTransform() const
 			{
@@ -334,7 +334,7 @@ namespace UE
 			FTransform MeshGlobalTransform = FTransform::Identity;
 		};
 
-		class INTERCHANGEDISPATCHER_API FJsonFetchAnimationQueriesCmd : public FJsonFetchPayloadCmd
+		class FJsonFetchAnimationQueriesCmd : public FJsonFetchPayloadCmd
 		{
 		public:
 			FJsonFetchAnimationQueriesCmd()
@@ -348,8 +348,8 @@ namespace UE
 			{
 			}
 
-			virtual FString ToJson() const override;
-			virtual bool FromJson(const FString& JsonString) override;
+			INTERCHANGEDISPATCHER_API virtual FString ToJson() const override;
+			INTERCHANGEDISPATCHER_API virtual bool FromJson(const FString& JsonString) override;
 
 			static FString GetQueriesJsonStringKey()
 			{
@@ -365,7 +365,7 @@ namespace UE
 			/**
 			 * Use this class helper to create the cmd result json string and to read it
 			 */
-			class INTERCHANGEDISPATCHER_API JsonAnimationQueriesResultParser
+			class JsonAnimationQueriesResultParser
 			{
 			public:
 				TMap<FString, FString> GetResultFilename() const
@@ -377,8 +377,8 @@ namespace UE
 					HashToFilenames = InHashToFilenames;
 				}
 
-				FString ToJson() const;
-				bool FromJson(const FString& JsonString);
+				INTERCHANGEDISPATCHER_API FString ToJson() const;
+				INTERCHANGEDISPATCHER_API bool FromJson(const FString& JsonString);
 
 				static FString GetHashToFilenamesKey()
 				{
diff --git a/Engine/Plugins/Interchange/Runtime/Source/FactoryNodes/Public/InterchangePhysicsAssetFactoryNode.h b/Engine/Plugins/Interchange/Runtime/Source/FactoryNodes/Public/InterchangePhysicsAssetFactoryNode.h
--- a/Engine/Plugins/Interchange/Runtime/Source/FactoryNodes/Public/InterchangePhysicsAssetFactoryNode.h
+++ b/Engine/Plugins/Interchange/Runtime/Source/FactoryNodes/Public/InterchangePhysicsAssetFactoryNode.h
@@ -12,8 +12,8 @@
 
 #include "InterchangePhysicsAssetFactoryNode.generated.h"
 
-UCLASS(BlueprintType)
-class INTERCHANGEFACTORYNODES_API UInterchangePhysicsAssetFactoryNode : public UInterchangeFactoryBaseNode
+UCLASS(BlueprintType, MinimalAPI)
+class UInterchangePhysicsAssetFactoryNode : public UInterchangeFactoryBaseNode
 {
 	GENERATED_BODY()
 
diff --git a/Engine/Plugins/Interchange/Runtime/Source/FactoryNodes/Public/InterchangeSkeletonFactoryNode.h b/Engine/Plugins/Interchange/Runtime/Source/FactoryNodes/Public/InterchangeSkeletonFactoryNode.h
--- a/Engine/Plugins/Interchange/Runtime/Source/FactoryNodes/Public/InterchangeSkeletonFactoryNode.h
+++ b/Engine/Plugins/Interchange/Runtime/Source/FactoryNodes/Public/InterchangeSkeletonFactoryNode.h
@@ -11,8 +11,8 @@
 
 #include "InterchangeSkeletonFactoryNode.generated.h"
 
-UCLASS(BlueprintType)
-class INTERCHANGEFACTORYNODES_API UInterchangeSkeletonFactoryNode : public UInterchangeFactoryBaseNode
+UCLASS(BlueprintType, MinimalAPI)
+class UInterchangeSkeletonFactoryNode : public UInterchangeFactoryBaseNode
 {
 	GENERATED_BODY()
 
diff --git a/Engine/Plugins/Interchange/Runtime/Source/FactoryNodes/Public/InterchangeTexture2DFactoryNode.h b/Engine/Plugins/Interchange/Runtime/Source/FactoryNodes/Public/InterchangeTexture2DFactoryNode.h
--- a/Engine/Plugins/Interchange/Runtime/Source/FactoryNodes/Public/InterchangeTexture2DFactoryNode.h
+++ b/Engine/Plugins/Interchange/Runtime/Source/FactoryNodes/Public/InterchangeTexture2DFactoryNode.h
@@ -22,8 +22,8 @@ namespace UE::Interchange
 	};
 }//ns UE::Interchange
 
-UCLASS(BlueprintType)
-class INTERCHANGEFACTORYNODES_API UInterchangeTexture2DFactoryNode : public UInterchangeTextureFactoryNode
+UCLASS(BlueprintType, MinimalAPI)
+class UInterchangeTexture2DFactoryNode : public UInterchangeTextureFactoryNode
 {
 	GENERATED_BODY()
 
diff --git a/Engine/Plugins/Interchange/Runtime/Source/FactoryNodes/Public/InterchangeTextureFactoryNode.h b/Engine/Plugins/Interchange/Runtime/Source/FactoryNodes/Public/InterchangeTextureFactoryNode.h
--- a/Engine/Plugins/Interchange/Runtime/Source/FactoryNodes/Public/InterchangeTextureFactoryNode.h
+++ b/Engine/Plugins/Interchange/Runtime/Source/FactoryNodes/Public/InterchangeTextureFactoryNode.h
@@ -10,8 +10,8 @@
 #include "InterchangeTextureFactoryNode.generated.h"
 
 
-UCLASS(BlueprintType, Abstract)
-class INTERCHANGEFACTORYNODES_API UInterchangeTextureFactoryNode : public UInterchangeFactoryBaseNode
+UCLASS(BlueprintType, Abstract, MinimalAPI)
+class UInterchangeTextureFactoryNode : public UInterchangeFactoryBaseNode
 {
 	GENERATED_BODY()
 
@@ -529,10 +529,10 @@ public:
 
 	/** Return false if the Attribute was not set previously.*/
 	UFUNCTION(BlueprintCallable, Category = "Interchange | Node | Texture")
-	bool GetCustomColorSpace(ETextureColorSpace& AttributeValue) const;
+	INTERCHANGEFACTORYNODES_API bool GetCustomColorSpace(ETextureColorSpace& AttributeValue) const;
 
 	UFUNCTION(BlueprintCallable, Category = "Interchange | Node | Texture")
-	bool SetCustomColorSpace(ETextureColorSpace AttributeValue, bool bAddApplyDelegate = true);
+	INTERCHANGEFACTORYNODES_API bool SetCustomColorSpace(ETextureColorSpace AttributeValue, bool bAddApplyDelegate = true);
 
 	//////////////////////////////////////////////////////////////////////////
 	//Level of Detail
diff --git a/Engine/Plugins/Interchange/Runtime/Source/FactoryNodes/Public/InterchangeTextureLightProfileFactoryNode.h b/Engine/Plugins/Interchange/Runtime/Source/FactoryNodes/Public/InterchangeTextureLightProfileFactoryNode.h
--- a/Engine/Plugins/Interchange/Runtime/Source/FactoryNodes/Public/InterchangeTextureLightProfileFactoryNode.h
+++ b/Engine/Plugins/Interchange/Runtime/Source/FactoryNodes/Public/InterchangeTextureLightProfileFactoryNode.h
@@ -15,8 +15,8 @@
 
 #include "InterchangeTextureLightProfileFactoryNode.generated.h"
 
-UCLASS(BlueprintType)
-class INTERCHANGEFACTORYNODES_API UInterchangeTextureLightProfileFactoryNode : public UInterchangeTexture2DFactoryNode
+UCLASS(BlueprintType, MinimalAPI)
+class UInterchangeTextureLightProfileFactoryNode : public UInterchangeTexture2DFactoryNode
 {
 	GENERATED_BODY()
 
diff --git a/Engine/Plugins/Interchange/Runtime/Source/Nodes/Public/InterchangeAnimationTrackSetNode.h b/Engine/Plugins/Interchange/Runtime/Source/Nodes/Public/InterchangeAnimationTrackSetNode.h
--- a/Engine/Plugins/Interchange/Runtime/Source/Nodes/Public/InterchangeAnimationTrackSetNode.h
+++ b/Engine/Plugins/Interchange/Runtime/Source/Nodes/Public/InterchangeAnimationTrackSetNode.h
@@ -8,6 +8,8 @@
 
 #include "InterchangeAnimationTrackSetNode.generated.h"
 
+#define UE_API INTERCHANGENODES_API
+
 UENUM(BlueprintType)
 enum class EInterchangeAnimationPayLoadType : uint8
 {
@@ -27,7 +29,7 @@ enum class EInterchangeAnimationPayLoadType : uint8
 
 
 USTRUCT(BlueprintType)
-struct INTERCHANGENODES_API FInterchangeAnimationPayLoadKey
+struct FInterchangeAnimationPayLoadKey
 {
 	GENERATED_BODY()
 
@@ -64,12 +66,12 @@ namespace UE
 	namespace Interchange
 	{
 
-		struct INTERCHANGENODES_API FAnimationStaticData : public FBaseNodeStaticData
+		struct  FAnimationStaticData : public FBaseNodeStaticData
 		{
-			static const FAttributeKey& AnimationPayLoadUidKey();
-			static const FAttributeKey& AnimationPayLoadTypeKey();
-			static const FAttributeKey& MorphTargetAnimationPayLoadUidKey();
-			static const FAttributeKey& MorphTargetAnimationPayLoadTypeKey();
+			static UE_API const FAttributeKey& AnimationPayLoadUidKey();
+			static UE_API const FAttributeKey& AnimationPayLoadTypeKey();
+			static UE_API const FAttributeKey& MorphTargetAnimationPayLoadUidKey();
+			static UE_API const FAttributeKey& MorphTargetAnimationPayLoadTypeKey();
 		};
 	}
 }
@@ -97,12 +99,12 @@ enum class EInterchangeAimationCompletionMode : uint8
 /**
  * Class to represent a set of animation track nodes that share the same frame rate.
  */
-UCLASS(BlueprintType)
-class INTERCHANGENODES_API UInterchangeAnimationTrackSetNode : public UInterchangeBaseNode
+UCLASS(BlueprintType, MinimalAPI)
+class UInterchangeAnimationTrackSetNode : public UInterchangeBaseNode
 {
 	GENERATED_BODY()
 
-	UInterchangeAnimationTrackSetNode();
+	UE_API UInterchangeAnimationTrackSetNode();
 
 public:
 	static FStringView StaticAssetTypeName()
@@ -123,43 +125,43 @@ public:
 	 * Retrieve the number of track dependencies for this object.
 	 */
 	UFUNCTION(BlueprintCallable, Category = "Interchange | Node | AnimationTrackSet")
-	int32 GetCustomAnimationTrackUidCount() const;
+	UE_API int32 GetCustomAnimationTrackUidCount() const;
 
 	/**
 	 * Retrieve the track dependency for this object.
 	 */
 	UFUNCTION(BlueprintCallable, Category = "Interchange | Node | AnimationTrackSet")
-	void GetCustomAnimationTrackUids(TArray<FString>& OutAnimationTrackUids) const;
+	UE_API void GetCustomAnimationTrackUids(TArray<FString>& OutAnimationTrackUids) const;
 
 	/**
 	 * Retrieve one track dependency for this object.
 	 */
 	UFUNCTION(BlueprintCallable, Category = "Interchange | Node | AnimationTrackSet")
-	void GetCustomAnimationTrackUid(const int32 Index, FString& OutAnimationTrackUid) const;
+	UE_API void GetCustomAnimationTrackUid(const int32 Index, FString& OutAnimationTrackUid) const;
 
 	/**
 	 * Add one track dependency to this object.
 	 */
 	UFUNCTION(BlueprintCallable, Category = "Interchange | Node | AnimationTrackSet")
-	bool AddCustomAnimationTrackUid(const FString& AnimationTrackUid);
+	UE_API bool AddCustomAnimationTrackUid(const FString& AnimationTrackUid);
 
 	/**
 	 * Remove one track dependency from this object.
 	 */
 	UFUNCTION(BlueprintCallable, Category = "Interchange | Node | AnimationTrackSet")
-	bool RemoveCustomAnimationTrackUid(const FString& AnimationTrackUid);
+	UE_API bool RemoveCustomAnimationTrackUid(const FString& AnimationTrackUid);
 
 	/**
 	 * Set the frame rate for the animations in the level sequence.
 	 */
 	UFUNCTION(BlueprintCallable, Category = "Interchange | Node | AnimationTrackSet")
-	bool SetCustomFrameRate(const float& AttributeValue);
+	UE_API bool SetCustomFrameRate(const float& AttributeValue);
 
 	/**
 	 * Get the frame rate for the animations in the level sequence.
 	 */
 	UFUNCTION(BlueprintCallable, Category = "Interchange | Node | AnimationTrackSet")
-	bool GetCustomFrameRate(float& AttributeValue) const;
+	UE_API bool GetCustomFrameRate(float& AttributeValue) const;
 
 private:
 	const UE::Interchange::FAttributeKey Macro_CustomFrameRateKey = UE::Interchange::FAttributeKey(TEXT("FrameRate"));
@@ -170,8 +172,8 @@ private:
 /**
  * Abstract class providing the minimal services required for an animation track node.
  */
-UCLASS(Abstract)
-class INTERCHANGENODES_API UInterchangeAnimationTrackBaseNode : public UInterchangeBaseNode
+UCLASS(Abstract, MinimalAPI)
+class UInterchangeAnimationTrackBaseNode : public UInterchangeBaseNode
 {
 	GENERATED_BODY()
 
@@ -194,14 +196,14 @@ public:
 	 * Set how the actor's animated property should behave once its animation completes.
 	 */
 	UFUNCTION(BlueprintCallable, Category = "Interchange | Node | AnimationTrack")
-	bool SetCustomCompletionMode(const int32& AttributeValue);
+	UE_API bool SetCustomCompletionMode(const int32& AttributeValue);
 
 	/**
 	 * Get how the actor's animated property behaves once this animation is complete.
 	 * The output value will be clamped to the range of values defined in EInterchangeAimationCompletionMode.
 	 */
 	UFUNCTION(BlueprintCallable, Category = "Interchange | Node | AnimationTrack")
-	bool GetCustomCompletionMode(int32& AttributeValue) const;
+	UE_API bool GetCustomCompletionMode(int32& AttributeValue) const;
 
 private:
 	const UE::Interchange::FAttributeKey Macro_CustomCompletionModeKey = UE::Interchange::FAttributeKey(TEXT("CompletionMode"));
@@ -210,8 +212,8 @@ private:
 /**
  * Class to represent an animation that instances another animation track set node.
  */
-UCLASS(BlueprintType)
-class INTERCHANGENODES_API UInterchangeAnimationTrackSetInstanceNode : public UInterchangeAnimationTrackBaseNode
+UCLASS(BlueprintType, MinimalAPI)
+class UInterchangeAnimationTrackSetInstanceNode : public UInterchangeAnimationTrackBaseNode
 {
 	GENERATED_BODY()
 
@@ -234,49 +236,49 @@ public:
 	 * Set the time scale used for the level sequence instance.
 	 */
 	UFUNCTION(BlueprintCallable, Category = "Interchange | Node | AnimationTrackSetInstance")
-	bool SetCustomTimeScale(const float& AttributeValue);
+	UE_API bool SetCustomTimeScale(const float& AttributeValue);
 
 	/**
 	 * Get the time scale used for the level sequence instance.
 	 */
 	UFUNCTION(BlueprintCallable, Category = "Interchange | Node | AnimationTrackSetInstance")
-	bool GetCustomTimeScale(float& AttributeValue) const;
+	UE_API bool GetCustomTimeScale(float& AttributeValue) const;
 
 	/**
 	 * Set the level sequence instance duration in number of frames.
 	 */
 	UFUNCTION(BlueprintCallable, Category = "Interchange | Node | AnimationTrackSetInstance")
-	bool SetCustomDuration(const int32& AttributeValue);
+	UE_API bool SetCustomDuration(const int32& AttributeValue);
 
 	/**
 	 * Get the level sequence instance duration in number of frames.
 	 */
 	UFUNCTION(BlueprintCallable, Category = "Interchange | Node | AnimationTrackSetInstance")
-	bool GetCustomDuration(int32& AttributeValue) const;
+	UE_API bool GetCustomDuration(int32& AttributeValue) const;
 
 	/**
 	 * Set the frame where the level sequence instance starts.
 	 */
 	UFUNCTION(BlueprintCallable, Category = "Interchange | Node | AnimationTrackSetInstance")
-	bool SetCustomStartFrame(const int32& AttributeValue);
+	UE_API bool SetCustomStartFrame(const int32& AttributeValue);
 
 	/**
 	 * Get the frame where the level sequence instance starts.
 	 */
 	UFUNCTION(BlueprintCallable, Category = "Interchange | Node | AnimationTrackSetInstance")
-	bool GetCustomStartFrame(int32& AttributeValue) const;
+	UE_API bool GetCustomStartFrame(int32& AttributeValue) const;
 
 	/**
 	 * Set the unique id of the level sequence this instance references.
 	 */
 	UFUNCTION(BlueprintCallable, Category = "Interchange | Node | AnimationTrackSetInstance")
-	bool SetCustomTrackSetDependencyUid(const FString& AttributeValue);
+	UE_API bool SetCustomTrackSetDependencyUid(const FString& AttributeValue);
 
 	/**
 	 * Get the unique id of the level sequence this instance references.
 	 */
 	UFUNCTION(BlueprintCallable, Category = "Interchange | Node | AnimationTrackSetInstance")
-	bool GetCustomTrackSetDependencyUid(FString& AttributeValue) const;
+	UE_API bool GetCustomTrackSetDependencyUid(FString& AttributeValue) const;
 
 private:
 	const UE::Interchange::FAttributeKey Macro_CustomStartFrameKey = UE::Interchange::FAttributeKey(TEXT("StartFrame"));
@@ -289,8 +291,8 @@ private:
  * Class to represent an animation on the property of a camera, light, or scene node
  * The list of supported properties is enumerated in EInterchangeAnimatedProperty.
  */
-UCLASS(BlueprintType)
-class INTERCHANGENODES_API UInterchangeAnimationTrackNode : public UInterchangeAnimationTrackBaseNode
+UCLASS(BlueprintType, MinimalAPI)
+class UInterchangeAnimationTrackNode : public UInterchangeAnimationTrackBaseNode
 {
 	GENERATED_BODY()
 
@@ -313,55 +315,55 @@ public:
 	 * Set the actor dependency to this object.
 	 */
 	UFUNCTION(BlueprintCallable, Category = "Interchange | Node | AnimationTrack")
-	bool SetCustomActorDependencyUid(const FString& DependencyUid);
+	UE_API bool SetCustomActorDependencyUid(const FString& DependencyUid);
 
 	/**
 	 * Get the actor dependency to this object.
 	 */
 	UFUNCTION(BlueprintCallable, Category = "Interchange | Node | AnimationTrack")
-	bool GetCustomActorDependencyUid(FString& DependencyUid) const;
+	UE_API bool GetCustomActorDependencyUid(FString& DependencyUid) const;
 
 	UE_DEPRECATED(5.5, "SetCustomPropertyTrack using an FName has been deprecated, please use the function with the enum instead.")
-	bool SetCustomPropertyTrack(const FName& PropertyTrack);
+	UE_API bool SetCustomPropertyTrack(const FName& PropertyTrack);
 
 	UE_DEPRECATED(5.5, "SetCustomPropertyTrack using an FName has been deprecated, please use the function with the enum instead.")
-	bool GetCustomPropertyTrack(FName& PropertyTrack) const;
+	UE_API bool GetCustomPropertyTrack(FName& PropertyTrack) const;
 
 	/**
 	 * Set the property animated by this track.
 	 */
 	UFUNCTION(BlueprintCallable, Category = "Interchange | Node | AnimationTrack")
-	bool SetCustomPropertyTrack(EInterchangePropertyTracks PropertyTrack);
+	UE_API bool SetCustomPropertyTrack(EInterchangePropertyTracks PropertyTrack);
 
 	/**
 	 * Get the property animated by this track.
 	 */
 	UFUNCTION(BlueprintCallable, Category = "Interchange | Node | AnimationTrack")
-	bool GetCustomPropertyTrack(EInterchangePropertyTracks& PropertyTrack) const;
+	UE_API bool GetCustomPropertyTrack(EInterchangePropertyTracks& PropertyTrack) const;
 
 	/**
 	 * Set the payload key needed to retrieve the animation for this track.
 	 */
 	UFUNCTION(BlueprintCallable, Category = "Interchange | Node | AnimationTrack")
-	bool SetCustomAnimationPayloadKey(const FString& InUniqueId, const EInterchangeAnimationPayLoadType& InType);
+	UE_API bool SetCustomAnimationPayloadKey(const FString& InUniqueId, const EInterchangeAnimationPayLoadType& InType);
 
 	/**
 	 * Get the payload key needed to retrieve the animation for this track.
 	 */
 	UFUNCTION(BlueprintCallable, Category = "Interchange | Node | AnimationTrack")
-	bool GetCustomAnimationPayloadKey(FInterchangeAnimationPayLoadKey& AnimationPayLoadKey) const;
+	UE_API bool GetCustomAnimationPayloadKey(FInterchangeAnimationPayLoadKey& AnimationPayLoadKey) const;
 
 	/**
 	 * Set the number of frames for the animation of this track.
 	 */
 	UFUNCTION(BlueprintCallable, Category = "Interchange | Node | AnimationTrack")
-	bool SetCustomFrameCount(const int32& AttributeValue);
+	UE_API bool SetCustomFrameCount(const int32& AttributeValue);
 
 	/**
 	 * Get the number of frames for the animation of this track.
 	 */
 	UFUNCTION(BlueprintCallable, Category = "Interchange | Node | AnimationTrack")
-	bool GetCustomFrameCount(int32& AttributeValue) const;
+	UE_API bool GetCustomFrameCount(int32& AttributeValue) const;
 
 private:
 	const UE::Interchange::FAttributeKey Macro_CustomActorDependencyKey = UE::Interchange::FAttributeKey(TEXT("ActorDependency"));
@@ -375,8 +377,8 @@ private:
 /**
  * Class to represent an animation on the transform of a camera, light, or scene node.
  */
-UCLASS(BlueprintType)
-class INTERCHANGENODES_API UInterchangeTransformAnimationTrackNode : public UInterchangeAnimationTrackNode
+UCLASS(BlueprintType, MinimalAPI)
+class UInterchangeTransformAnimationTrackNode : public UInterchangeAnimationTrackNode
 {
 	GENERATED_BODY()
 
@@ -416,14 +418,14 @@ public:
 	 *    All           = Translation | Rotation | Scale | Weight,
 	 */
 	UFUNCTION(BlueprintCallable, Category = "Interchange | Node | TransformAnimationTrack")
-	bool SetCustomUsedChannels(const int32& AttributeValue);
+	UE_API bool SetCustomUsedChannels(const int32& AttributeValue);
 
 	/**
 	 * Get which channels of this animation should be used. This is a bitmask.
 	 * See SetCustomUsedChannels for description of bitmask
 	 */
 	UFUNCTION(BlueprintCallable, Category = "Interchange | Node | TransformAnimationTrack")
-	bool GetCustomUsedChannels(int32& AttributeValue) const;
+	UE_API bool GetCustomUsedChannels(int32& AttributeValue) const;
 
 private:
 	const UE::Interchange::FAttributeKey Macro_CustomUsedChannelsKey = UE::Interchange::FAttributeKey(TEXT("UsedChannels"));
@@ -432,13 +434,13 @@ private:
 /*
 * Class to hold onto the relationships between a set of animation tracks and the bones, morph targets of a skeleton.
 */
-UCLASS(BlueprintType)
-class INTERCHANGENODES_API UInterchangeSkeletalAnimationTrackNode : public UInterchangeAnimationTrackBaseNode
+UCLASS(BlueprintType, MinimalAPI)
+class UInterchangeSkeletalAnimationTrackNode : public UInterchangeAnimationTrackBaseNode
 {
 	GENERATED_BODY()
 
 public:
-	UInterchangeSkeletalAnimationTrackNode();
+	UE_API UInterchangeSkeletalAnimationTrackNode();
 
 	/**
 	 * Override Serialize() to restore SlotMaterialDependencies on load.
@@ -475,47 +477,47 @@ public:
 public:
 	/** Get the unique ID of the skeleton factory node. Return false if the attribute is not set. */
 	UFUNCTION(BlueprintCallable, Category = "Interchange | Node | SkeletalAnimationTrack")
-	bool GetCustomSkeletonNodeUid(FString& AttributeValue) const;
+	UE_API bool GetCustomSkeletonNodeUid(FString& AttributeValue) const;
 
 	/** Set the unique ID of the skeleton factory node. Return false if the attribute could not be set. */
 	UFUNCTION(BlueprintCallable, Category = "Interchange | Node | SkeletalAnimationTrack")
-	bool SetCustomSkeletonNodeUid(const FString& AttributeValue);
+	UE_API bool SetCustomSkeletonNodeUid(const FString& AttributeValue);
 
 	/**
 	 * Set the animation sample rate. Return false if the attribute could not be set.
 	 */
 	UFUNCTION(BlueprintCallable, Category = "Interchange | Node | SkeletalAnimationTrack")
-	bool SetCustomAnimationSampleRate(const double& SampleRate);
+	UE_API bool SetCustomAnimationSampleRate(const double& SampleRate);
 
 	/**
 	 * Get the animation sample rate. Return false if the attribute is not set.
 	 */
 	UFUNCTION(BlueprintCallable, Category = "Interchange | Node | SkeletalAnimationTrack")
-	bool GetCustomAnimationSampleRate(double& SampleRate) const;
+	UE_API bool GetCustomAnimationSampleRate(double& SampleRate) const;
 
 	/**
 	 * Set the animation start time. Return false if the attribute could not be set.
 	 */
 	UFUNCTION(BlueprintCallable, Category = "Interchange | Node | SkeletalAnimationTrack")
-	bool SetCustomAnimationStartTime(const double& StartTime);
+	UE_API bool SetCustomAnimationStartTime(const double& StartTime);
 
 	/**
 	 * Get the animation start time. Return false if the attribute is not set.
 	 */
 	UFUNCTION(BlueprintCallable, Category = "Interchange | Node | SkeletalAnimationTrack")
-	bool GetCustomAnimationStartTime(double& StartTime) const;
+	UE_API bool GetCustomAnimationStartTime(double& StartTime) const;
 	
 	/**
 	 * Set the animation stop time. Return false if the attribute could not be set.
 	 */
 	UFUNCTION(BlueprintCallable, Category = "Interchange | Node | SkeletalAnimationTrack")
-	bool SetCustomAnimationStopTime(const double& StopTime);
+	UE_API bool SetCustomAnimationStopTime(const double& StopTime);
 
 	/**
 	 * Get the animation stop time. Return false if the attribute is not set.
 	 */
 	UFUNCTION(BlueprintCallable, Category = "Interchange | Node | SkeletalAnimationTrack")
-	bool GetCustomAnimationStopTime(double& StopTime) const;
+	UE_API bool GetCustomAnimationStopTime(double& StopTime) const;
 
 	/**
 	 * Set the source timeline animation start time. Return false if the attribute could not be set.
@@ -523,7 +525,7 @@ public:
 	 * Optional: if not set CustomAnimationStartTime will be used
 	 */
 	UFUNCTION(BlueprintCallable, Category = "Interchange | Node | SkeletalAnimationTrack")
-	bool SetCustomSourceTimelineAnimationStartTime(const double& StartTime);
+	UE_API bool SetCustomSourceTimelineAnimationStartTime(const double& StartTime);
 
 	/**
 	 * Get the source timeline animation start time. Return false if the attribute is not set.
@@ -531,7 +533,7 @@ public:
 	 * Optional: if not set CustomAnimationStartTime will be used
 	 */
 	UFUNCTION(BlueprintCallable, Category = "Interchange | Node | SkeletalAnimationTrack")
-	bool GetCustomSourceTimelineAnimationStartTime(double& StartTime) const;
+	UE_API bool GetCustomSourceTimelineAnimationStartTime(double& StartTime) const;
 	
 	/**
 	 * Set the source timeline animation stop time. Return false if the attribute could not be set.
@@ -539,7 +541,7 @@ public:
 	 * Optional: if not set CustomAnimationStopTime will be used
 	 */
 	UFUNCTION(BlueprintCallable, Category = "Interchange | Node | SkeletalAnimationTrack")
-	bool SetCustomSourceTimelineAnimationStopTime(const double& StopTime);
+	UE_API bool SetCustomSourceTimelineAnimationStopTime(const double& StopTime);
 
 	/**
 	 * Get the source timeline animation stop time. Return false if the attribute is not set.
@@ -547,23 +549,23 @@ public:
 	 * Optional: if not set CustomAnimationStopTime will be used
 	 */
 	UFUNCTION(BlueprintCallable, Category = "Interchange | Node | SkeletalAnimationTrack")
-	bool GetCustomSourceTimelineAnimationStopTime(double& StopTime) const;
+	UE_API bool GetCustomSourceTimelineAnimationStopTime(double& StopTime) const;
 
 
 	UFUNCTION(BlueprintCallable, Category = "Interchange | Node | SkeletalAnimationTrack")
-	void GetSceneNodeAnimationPayloadKeys(TMap<FString, FString>& OutSceneNodeAnimationPayloadKeyUids, TMap<FString, uint8>& OutSceneNodeAnimationPayloadKeyTypes) const;
+	UE_API void GetSceneNodeAnimationPayloadKeys(TMap<FString, FString>& OutSceneNodeAnimationPayloadKeyUids, TMap<FString, uint8>& OutSceneNodeAnimationPayloadKeyTypes) const;
 
 	UFUNCTION(BlueprintCallable, Category = "Interchange | Node | SkeletalAnimationTrack")
-	bool SetAnimationPayloadKeyForSceneNodeUid(const FString& SceneNodeUid, const FString& InUniqueId, const EInterchangeAnimationPayLoadType& InType);
+	UE_API bool SetAnimationPayloadKeyForSceneNodeUid(const FString& SceneNodeUid, const FString& InUniqueId, const EInterchangeAnimationPayLoadType& InType);
 
 	UFUNCTION(BlueprintCallable, Category = "Interchange | Node | SkeletalAnimationTrack")
-	void GetMorphTargetNodeAnimationPayloadKeys(TMap<FString, FString>& OutMorphTargetNodeAnimationPayloadKeyUids, TMap<FString, uint8>& OutMorphTargetNodeAnimationPayloadKeyTypes) const;
+	UE_API void GetMorphTargetNodeAnimationPayloadKeys(TMap<FString, FString>& OutMorphTargetNodeAnimationPayloadKeyUids, TMap<FString, uint8>& OutMorphTargetNodeAnimationPayloadKeyTypes) const;
 
 	UFUNCTION(BlueprintCallable, Category = "Interchange | Node | SkeletalAnimationTrack")
-	bool SetAnimationPayloadKeyForMorphTargetNodeUid(const FString& MorphTargetNodeUid, const FString& InUniqueId, const EInterchangeAnimationPayLoadType& InType);
+	UE_API bool SetAnimationPayloadKeyForMorphTargetNodeUid(const FString& MorphTargetNodeUid, const FString& InUniqueId, const EInterchangeAnimationPayLoadType& InType);
 
 	UFUNCTION(BlueprintCallable, Category = "Interchange | Node | SkeletalAnimationTrack")
-	bool IsNodeAnimatedWithBakedCurve(const FString& SceneNodeUid) const;
+	UE_API bool IsNodeAnimatedWithBakedCurve(const FString& SceneNodeUid) const;
 
 private:
 	const UE::Interchange::FAttributeKey Macro_CustomSkeletonNodeUidKey = UE::Interchange::FAttributeKey(TEXT("SkeletonNodeUid"));
@@ -579,3 +581,5 @@ private:
 	UE::Interchange::TMapAttributeHelper<FString, FString>	MorphTargetPayloadKeyUidMap;
 	UE::Interchange::TMapAttributeHelper<FString, uint8>	MorphTargetPayloadKeyTypeMap;
 };
+
+#undef UE_API
diff --git a/Engine/Plugins/Interchange/Runtime/Source/Nodes/Public/InterchangeCameraNode.h b/Engine/Plugins/Interchange/Runtime/Source/Nodes/Public/InterchangeCameraNode.h
--- a/Engine/Plugins/Interchange/Runtime/Source/Nodes/Public/InterchangeCameraNode.h
+++ b/Engine/Plugins/Interchange/Runtime/Source/Nodes/Public/InterchangeCameraNode.h
@@ -7,8 +7,8 @@
 
 #include "InterchangeCameraNode.generated.h"
 
-UCLASS(BlueprintType)
-class INTERCHANGENODES_API UInterchangePhysicalCameraNode : public UInterchangeBaseNode
+UCLASS(BlueprintType, MinimalAPI)
+class UInterchangePhysicalCameraNode : public UInterchangeBaseNode
 {
 	GENERATED_BODY()
 
@@ -90,8 +90,8 @@ enum class EInterchangeCameraProjectionType : uint8
 };
 
 // Primarily used for Ortho Camera
-UCLASS(BlueprintType)
-class INTERCHANGENODES_API UInterchangeStandardCameraNode : public UInterchangeBaseNode
+UCLASS(BlueprintType, MinimalAPI)
+class UInterchangeStandardCameraNode : public UInterchangeBaseNode
 {
 	GENERATED_BODY()
 
diff --git a/Engine/Plugins/Interchange/Runtime/Source/Nodes/Public/InterchangeDecalNode.h b/Engine/Plugins/Interchange/Runtime/Source/Nodes/Public/InterchangeDecalNode.h
--- a/Engine/Plugins/Interchange/Runtime/Source/Nodes/Public/InterchangeDecalNode.h
+++ b/Engine/Plugins/Interchange/Runtime/Source/Nodes/Public/InterchangeDecalNode.h
@@ -7,8 +7,8 @@
 
 #include "InterchangeDecalNode.generated.h"
 
-UCLASS(BlueprintType)
-class INTERCHANGENODES_API UInterchangeDecalNode : public UInterchangeBaseNode
+UCLASS(BlueprintType, MinimalAPI)
+class UInterchangeDecalNode : public UInterchangeBaseNode
 {
 	GENERATED_BODY()
 
diff --git a/Engine/Plugins/Interchange/Runtime/Source/Nodes/Public/InterchangeMeshNode.h b/Engine/Plugins/Interchange/Runtime/Source/Nodes/Public/InterchangeMeshNode.h
--- a/Engine/Plugins/Interchange/Runtime/Source/Nodes/Public/InterchangeMeshNode.h
+++ b/Engine/Plugins/Interchange/Runtime/Source/Nodes/Public/InterchangeMeshNode.h
@@ -8,23 +8,25 @@
 
 #include "InterchangeMeshNode.generated.h"
 
+#define UE_API INTERCHANGENODES_API
+
 //Interchange namespace
 namespace UE
 {
 	namespace Interchange
 	{
 
-		struct INTERCHANGENODES_API FMeshNodeStaticData : public FBaseNodeStaticData
+		struct FMeshNodeStaticData : public FBaseNodeStaticData
 		{
-			static const FAttributeKey& PayLoadKey();
-			static const FAttributeKey& PayLoadTypeKey();
-			static const FAttributeKey& IsSkinnedMeshKey();
-			static const FAttributeKey& IsMorphTargetKey();
-			static const FAttributeKey& MorphTargetNameKey();
-			static const FAttributeKey& GetSkeletonDependenciesKey();
-			static const FAttributeKey& GetMorphTargetDependenciesKey();
-			static const FAttributeKey& GetSceneInstancesUidsKey();
-			static const FAttributeKey& GetSlotMaterialDependenciesKey();
+			static UE_API const FAttributeKey& PayLoadKey();
+			static UE_API const FAttributeKey& PayLoadTypeKey();
+			static UE_API const FAttributeKey& IsSkinnedMeshKey();
+			static UE_API const FAttributeKey& IsMorphTargetKey();
+			static UE_API const FAttributeKey& MorphTargetNameKey();
+			static UE_API const FAttributeKey& GetSkeletonDependenciesKey();
+			static UE_API const FAttributeKey& GetMorphTargetDependenciesKey();
+			static UE_API const FAttributeKey& GetSceneInstancesUidsKey();
+			static UE_API const FAttributeKey& GetSlotMaterialDependenciesKey();
 		};
 
 	}//ns Interchange
@@ -42,7 +44,7 @@ enum class EInterchangeMeshPayLoadType : uint8
 
 
 USTRUCT(BlueprintType)
-struct INTERCHANGENODES_API FInterchangeMeshPayLoadKey
+struct FInterchangeMeshPayLoadKey
 {
 	GENERATED_BODY()
 
@@ -93,29 +95,29 @@ public:
 };
 
 
-UCLASS(BlueprintType)
-class INTERCHANGENODES_API UInterchangeMeshNode : public UInterchangeBaseNode
+UCLASS(BlueprintType, MinimalAPI)
+class UInterchangeMeshNode : public UInterchangeBaseNode
 {
 	GENERATED_BODY()
 
 public:
-	UInterchangeMeshNode();
+	UE_API UInterchangeMeshNode();
 
 #if WITH_EDITOR
-	virtual FString GetKeyDisplayName(const UE::Interchange::FAttributeKey& NodeAttributeKey) const override;
+	UE_API virtual FString GetKeyDisplayName(const UE::Interchange::FAttributeKey& NodeAttributeKey) const override;
 
-	virtual FString GetAttributeCategory(const UE::Interchange::FAttributeKey& NodeAttributeKey) const override;
+	UE_API virtual FString GetAttributeCategory(const UE::Interchange::FAttributeKey& NodeAttributeKey) const override;
 #endif //WITH_EDITOR
 
 	/**
 	 * Return the node type name of the class. This is used when reporting errors.
 	 */
-	virtual FString GetTypeName() const override;
+	UE_API virtual FString GetTypeName() const override;
 
 	/**
 	 * Icon names are created by adding "InterchangeIcon_" in front of the specialized type. If there is no special type, the function will return NAME_None, which will use the default icon.
 	 */
-	virtual FName GetIconName() const override;
+	UE_API virtual FName GetIconName() const override;
 
 	/**
 	 * Override Serialize() to restore SlotMaterialDependencies on load.
@@ -134,229 +136,229 @@ public:
 	 * Return true if this node represents a skinned mesh.
 	 */
 	UFUNCTION(BlueprintCallable, Category = "Interchange | Node | Mesh")
-	bool IsSkinnedMesh() const;
+	UE_API bool IsSkinnedMesh() const;
 
 	/**
 	 * Set the IsSkinnedMesh attribute to determine whether this node represents a skinned mesh.
 	 */
 	UFUNCTION(BlueprintCallable, Category = "Interchange | Node | Mesh")
-	bool SetSkinnedMesh(const bool bIsSkinnedMesh);
+	UE_API bool SetSkinnedMesh(const bool bIsSkinnedMesh);
 
 	/**
 	 * Return true if this node represents a morph target.
 	 */
 	UFUNCTION(BlueprintCallable, Category = "Interchange | Node | Mesh")
-	bool IsMorphTarget() const;
+	UE_API bool IsMorphTarget() const;
 
 	/**
 	 * Set the IsMorphTarget attribute to determine whether this node represents a morph target.
 	 */
 	UFUNCTION(BlueprintCallable, Category = "Interchange | Node | Mesh")
-	bool SetMorphTarget(const bool bIsMorphTarget);
+	UE_API bool SetMorphTarget(const bool bIsMorphTarget);
 
 	/**
 	 * Get the morph target name.
 	 * Return true if we successfully retrieved the MorphTargetName attribute.
 	 */
 	UFUNCTION(BlueprintCallable, Category = "Interchange | Node | Mesh")
-	bool GetMorphTargetName(FString& OutMorphTargetName) const;
+	UE_API bool GetMorphTargetName(FString& OutMorphTargetName) const;
 
 	/**
 	 * Set the MorphTargetName attribute to determine the name of the morph target.
 	 */
 	UFUNCTION(BlueprintCallable, Category = "Interchange | Node | Mesh")
-	bool SetMorphTargetName(const FString& MorphTargetName);
+	UE_API bool SetMorphTargetName(const FString& MorphTargetName);
 
-	virtual const TOptional<FInterchangeMeshPayLoadKey> GetPayLoadKey() const;
+	UE_API virtual const TOptional<FInterchangeMeshPayLoadKey> GetPayLoadKey() const;
 
 	UFUNCTION(BlueprintCallable, Category = "Interchange | Node | Mesh")
-	virtual void SetPayLoadKey(const FString& PayLoadKey, const EInterchangeMeshPayLoadType& PayLoadType);
+	UE_API virtual void SetPayLoadKey(const FString& PayLoadKey, const EInterchangeMeshPayLoadType& PayLoadType);
 
 	/** Query the vertex count of this mesh. Return false if the attribute was not set.*/
 	UFUNCTION(BlueprintCallable, Category = "Interchange | Node | Mesh")
-	bool GetCustomVertexCount(int32& AttributeValue) const;
+	UE_API bool GetCustomVertexCount(int32& AttributeValue) const;
 
 	/** Set the vertex count of this mesh. Return false if the attribute could not be set.*/
 	UFUNCTION(BlueprintCallable, Category = "Interchange | Node | Mesh")
-	bool SetCustomVertexCount(const int32& AttributeValue);
+	UE_API bool SetCustomVertexCount(const int32& AttributeValue);
 
 	/** Query the polygon count of this mesh. Return false if the attribute was not set.*/
 	UFUNCTION(BlueprintCallable, Category = "Interchange | Node | Mesh")
-	bool GetCustomPolygonCount(int32& AttributeValue) const;
+	UE_API bool GetCustomPolygonCount(int32& AttributeValue) const;
 
 	/** Set the polygon count of this mesh. Return false if the attribute could not be set.*/
 	UFUNCTION(BlueprintCallable, Category = "Interchange | Node | Mesh")
-	bool SetCustomPolygonCount(const int32& AttributeValue);
+	UE_API bool SetCustomPolygonCount(const int32& AttributeValue);
 
 	/** Query the bounding box of this mesh. Return false if the attribute was not set.*/
 	UFUNCTION(BlueprintCallable, Category = "Interchange | Node | Mesh")
-	bool GetCustomBoundingBox(FBox& AttributeValue) const;
+	UE_API bool GetCustomBoundingBox(FBox& AttributeValue) const;
 
 	/** Set the bounding box of this mesh. Return false if the attribute could not be set.*/
 	UFUNCTION(BlueprintCallable, Category = "Interchange | Node | Mesh")
-	bool SetCustomBoundingBox(const FBox& AttributeValue);
+	UE_API bool SetCustomBoundingBox(const FBox& AttributeValue);
 
 	/** Query whether this mesh has vertex normals. Return false if the attribute was not set.*/
 	UFUNCTION(BlueprintCallable, Category = "Interchange | Node | Mesh")
-	bool GetCustomHasVertexNormal(bool& AttributeValue) const;
+	UE_API bool GetCustomHasVertexNormal(bool& AttributeValue) const;
 
 	/** Set the vertex normal attribute of this mesh. Return false if the attribute could not be set.*/
 	UFUNCTION(BlueprintCallable, Category = "Interchange | Node | Mesh")
-	bool SetCustomHasVertexNormal(const bool& AttributeValue);
+	UE_API bool SetCustomHasVertexNormal(const bool& AttributeValue);
 
 	/** Query whether this mesh has vertex bi-normals. Return false if the attribute was not set.*/
 	UFUNCTION(BlueprintCallable, Category = "Interchange | Node | Mesh")
-	bool GetCustomHasVertexBinormal(bool& AttributeValue) const;
+	UE_API bool GetCustomHasVertexBinormal(bool& AttributeValue) const;
 
 	/** Set the vertex bi-normal attribute of this mesh. Return false if the attribute could not be set.*/
 	UFUNCTION(BlueprintCallable, Category = "Interchange | Node | Mesh")
-	bool SetCustomHasVertexBinormal(const bool& AttributeValue);
+	UE_API bool SetCustomHasVertexBinormal(const bool& AttributeValue);
 
 	/** Query whether this mesh has vertex tangents. Return false if the attribute was not set.*/
 	UFUNCTION(BlueprintCallable, Category = "Interchange | Node | Mesh")
-	bool GetCustomHasVertexTangent(bool& AttributeValue) const;
+	UE_API bool GetCustomHasVertexTangent(bool& AttributeValue) const;
 
 	/** Set the vertex tangent attribute of this mesh. Return false if the attribute could not be set.*/
 	UFUNCTION(BlueprintCallable, Category = "Interchange | Node | Mesh")
-	bool SetCustomHasVertexTangent(const bool& AttributeValue);
+	UE_API bool SetCustomHasVertexTangent(const bool& AttributeValue);
 
 	/** Query whether this mesh has smoothing groups. Return false if the attribute was not set.*/
 	UFUNCTION(BlueprintCallable, Category = "Interchange | Node | Mesh")
-	bool GetCustomHasSmoothGroup(bool& AttributeValue) const;
+	UE_API bool GetCustomHasSmoothGroup(bool& AttributeValue) const;
 
 	/** Set the smoothing group attribute of this mesh. Return false if the attribute could not be set.*/
 	UFUNCTION(BlueprintCallable, Category = "Interchange | Node | Mesh")
-	bool SetCustomHasSmoothGroup(const bool& AttributeValue);
+	UE_API bool SetCustomHasSmoothGroup(const bool& AttributeValue);
 
 	/** Query whether this mesh has vertex colors. Return false if the attribute was not set.*/
 	UFUNCTION(BlueprintCallable, Category = "Interchange | Node | Mesh")
-	bool GetCustomHasVertexColor(bool& AttributeValue) const;
+	UE_API bool GetCustomHasVertexColor(bool& AttributeValue) const;
 
 	/** Set the vertex color attribute of this mesh. Return false if the attribute could not be set.*/
 	UFUNCTION(BlueprintCallable, Category = "Interchange | Node | Mesh")
-	bool SetCustomHasVertexColor(const bool& AttributeValue);
+	UE_API bool SetCustomHasVertexColor(const bool& AttributeValue);
 
 	/** Query the UV count of this mesh. Return false if the attribute was not set.*/
 	UFUNCTION(BlueprintCallable, Category = "Interchange | Node | Mesh")
-	bool GetCustomUVCount(int32& AttributeValue) const;
+	UE_API bool GetCustomUVCount(int32& AttributeValue) const;
 
 	/** Set the UV count attribute of this mesh. Return false if the attribute could not be set.*/
 	UFUNCTION(BlueprintCallable, Category = "Interchange | Node | Mesh")
-	bool SetCustomUVCount(const int32& AttributeValue);
+	UE_API bool SetCustomUVCount(const int32& AttributeValue);
 
 	/**
 	 * Retrieve the number of skeleton dependencies for this object.
 	 */
 	UFUNCTION(BlueprintCallable, Category = "Interchange | Node | Mesh")
-	int32 GetSkeletonDependeciesCount() const;
+	UE_API int32 GetSkeletonDependeciesCount() const;
 
 	/**
 	 * Retrieve the skeleton dependency for this object.
 	 */
 	UFUNCTION(BlueprintCallable, Category = "Interchange | Node | Mesh")
-	void GetSkeletonDependencies(TArray<FString>& OutDependencies) const;
+	UE_API void GetSkeletonDependencies(TArray<FString>& OutDependencies) const;
 
 	/**
 	 * Retrieve the specified skeleton dependency for this object.
 	 */
 	UFUNCTION(BlueprintCallable, Category = "Interchange | Node | Mesh")
-	void GetSkeletonDependency(const int32 Index, FString& OutDependency) const;
+	UE_API void GetSkeletonDependency(const int32 Index, FString& OutDependency) const;
 
 	/**
 	 * Add the specified skeleton dependency to this object.
 	 */
 	UFUNCTION(BlueprintCallable, Category = "Interchange | Node | Mesh")
-	bool SetSkeletonDependencyUid(const FString& DependencyUid);
+	UE_API bool SetSkeletonDependencyUid(const FString& DependencyUid);
 
 	/**
 	 * Remove the specified skeleton dependency from this object.
 	 */
 	UFUNCTION(BlueprintCallable, Category = "Interchange | Node | Mesh")
-	bool RemoveSkeletonDependencyUid(const FString& DependencyUid);
+	UE_API bool RemoveSkeletonDependencyUid(const FString& DependencyUid);
 
 	/**
 	 * Retrieve the number of morph target dependencies for this object.
 	 */
 	UFUNCTION(BlueprintCallable, Category = "Interchange | Node | Mesh")
-	int32 GetMorphTargetDependeciesCount() const;
+	UE_API int32 GetMorphTargetDependeciesCount() const;
 
 	/**
 	 * Retrieve all morph target dependencies for this object.
 	 */
 	UFUNCTION(BlueprintCallable, Category = "Interchange | Node | Mesh")
-	void GetMorphTargetDependencies(TArray<FString>& OutDependencies) const;
+	UE_API void GetMorphTargetDependencies(TArray<FString>& OutDependencies) const;
 
 	/**
 	 * Retrieve the specified morph target dependency for this object.
 	 */
 	UFUNCTION(BlueprintCallable, Category = "Interchange | Node | Mesh")
-	void GetMorphTargetDependency(const int32 Index, FString& OutDependency) const;
+	UE_API void GetMorphTargetDependency(const int32 Index, FString& OutDependency) const;
 
 	/**
 	 * Add the specified morph target dependency to this object.
 	 */
 	UFUNCTION(BlueprintCallable, Category = "Interchange | Node | Mesh")
-	bool SetMorphTargetDependencyUid(const FString& DependencyUid);
+	UE_API bool SetMorphTargetDependencyUid(const FString& DependencyUid);
 
 	/**
 	 * Remove the specified morph target dependency from this object.
 	 */
 	UFUNCTION(BlueprintCallable, Category = "Interchange | Node | Mesh")
-	bool RemoveMorphTargetDependencyUid(const FString& DependencyUid);
+	UE_API bool RemoveMorphTargetDependencyUid(const FString& DependencyUid);
 
 	/**
 	 * Retrieve the number of scene nodes instancing this mesh.
 	 */
 	UFUNCTION(BlueprintCallable, Category = "Interchange | Node | Mesh")
-	int32 GetSceneInstanceUidsCount() const;
+	UE_API int32 GetSceneInstanceUidsCount() const;
 
 	/**
 	 * Retrieve the asset instances this scene node refers to.
 	 */
 	UFUNCTION(BlueprintCallable, Category = "Interchange | Node | Mesh")
-	void GetSceneInstanceUids(TArray<FString>& OutDependencies) const;
+	UE_API void GetSceneInstanceUids(TArray<FString>& OutDependencies) const;
 
 	/**
 	 * Retrieve the asset instance this scene node refers to.
 	 */
 	UFUNCTION(BlueprintCallable, Category = "Interchange | Node | Mesh")
-	void GetSceneInstanceUid(const int32 Index, FString& OutDependency) const;
+	UE_API void GetSceneInstanceUid(const int32 Index, FString& OutDependency) const;
 
 	/**
 	 * Add the specified asset instance this scene node refers to.
 	 */
 	UFUNCTION(BlueprintCallable, Category = "Interchange | Node | Mesh")
-	bool SetSceneInstanceUid(const FString& DependencyUid);
+	UE_API bool SetSceneInstanceUid(const FString& DependencyUid);
 
 	/**
 	 * Remove the specified asset instance this scene node refers to.
 	 */
 	UFUNCTION(BlueprintCallable, Category = "Interchange | Node | Mesh")
-	bool RemoveSceneInstanceUid(const FString& DependencyUid);
+	UE_API bool RemoveSceneInstanceUid(const FString& DependencyUid);
 
 	/**
 	 * Retrieve the correspondence table between slot names and assigned materials for this object.
 	 */
 	UFUNCTION(BlueprintCallable, Category = "Interchange | Node | StaticMesh")
-	void GetSlotMaterialDependencies(TMap<FString, FString>& OutMaterialDependencies) const;
+	UE_API void GetSlotMaterialDependencies(TMap<FString, FString>& OutMaterialDependencies) const;
 
 	/**
 	 * Retrieve the specified Material dependency for a given slot of this object.
 	 */
 	UFUNCTION(BlueprintCallable, Category = "Interchange | Node | StaticMesh")
-	bool GetSlotMaterialDependencyUid(const FString& SlotName, FString& OutMaterialDependency) const;
+	UE_API bool GetSlotMaterialDependencyUid(const FString& SlotName, FString& OutMaterialDependency) const;
 
 	/**
 	 * Add the specified Material dependency to a specific slot name of this object.
 	 */
 	UFUNCTION(BlueprintCallable, Category = "Interchange | Node | StaticMesh")
-	bool SetSlotMaterialDependencyUid(const FString& SlotName, const FString& MaterialDependencyUid);
+	UE_API bool SetSlotMaterialDependencyUid(const FString& SlotName, const FString& MaterialDependencyUid);
 
 	/**
 	 * Remove the Material dependency associated with the given slot name from this object.
 	 */
 	UFUNCTION(BlueprintCallable, Category = "Interchange | Node | StaticMesh")
-	bool RemoveSlotMaterialDependencyUid(const FString& SlotName);
+	UE_API bool RemoveSlotMaterialDependencyUid(const FString& SlotName);
 
 	/**
 	 * Get the type of collision shapes we should generate from this mesh.
@@ -364,11 +366,11 @@ public:
 	 * collision shapes will only be used for the static mesh generated from this very same Mesh node
 	 */
 	UFUNCTION(BlueprintCallable, Category = "Interchange | Node | Mesh")
-	bool GetCustomCollisionType(EInterchangeMeshCollision& AttributeValue) const;
+	UE_API bool GetCustomCollisionType(EInterchangeMeshCollision& AttributeValue) const;
 
 	/** Set the type of collision shapes we should generate from this mesh */
 	UFUNCTION(BlueprintCallable, Category = "Interchange | Node | Mesh")
-	bool SetCustomCollisionType(EInterchangeMeshCollision AttributeValue);
+	UE_API bool SetCustomCollisionType(EInterchangeMeshCollision AttributeValue);
 
 private:
 	IMPLEMENT_NODE_ATTRIBUTE_KEY(VertexCount);
@@ -390,45 +392,45 @@ private:
 	UE::Interchange::TMapAttributeHelper<FString, FString> SlotMaterialDependencies;
 };
 
-UCLASS(BlueprintType)
-class INTERCHANGENODES_API UInterchangeGeometryCacheNode : public UInterchangeMeshNode
+UCLASS(BlueprintType, MinimalAPI)
+class  UInterchangeGeometryCacheNode : public UInterchangeMeshNode
 {
 	GENERATED_BODY()
 
 public:
-	UInterchangeGeometryCacheNode();
+	UE_API UInterchangeGeometryCacheNode();
 
 	/** Query the start frame of the animated mesh. Return false if the attribute was not set.*/
 	UFUNCTION(BlueprintCallable, Category = "Interchange | Node | GeometryCache")
-	bool GetCustomStartFrame(int32& AttributeValue) const;
+	UE_API bool GetCustomStartFrame(int32& AttributeValue) const;
 
 	/** Set the start frame of the animated mesh. Return false if the attribute could not be set.*/
 	UFUNCTION(BlueprintCallable, Category = "Interchange | Node | GeometryCache")
-	bool SetCustomStartFrame(const int32& AttributeValue);
+	UE_API bool SetCustomStartFrame(const int32& AttributeValue);
 
 	/** Query the end frame of the animated mesh. Return false if the attribute was not set.*/
 	UFUNCTION(BlueprintCallable, Category = "Interchange | Node | GeometryCache")
-	bool GetCustomEndFrame(int32& AttributeValue) const;
+	UE_API bool GetCustomEndFrame(int32& AttributeValue) const;
 
 	/** Set the end frame of the animated mesh. Return false if the attribute could not be set.*/
 	UFUNCTION(BlueprintCallable, Category = "Interchange | Node | GeometryCache")
-	bool SetCustomEndFrame(const int32& AttributeValue);
+	UE_API bool SetCustomEndFrame(const int32& AttributeValue);
 
 	/** Query the frame rate of the animated mesh. Return false if the attribute was not set.*/
 	UFUNCTION(BlueprintCallable, Category = "Interchange | Node | GeometryCache")
-	bool GetCustomFrameRate(double& AttributeValue) const;
+	UE_API bool GetCustomFrameRate(double& AttributeValue) const;
 
 	/** Set the frame rate of the animated mesh. Return false if the attribute could not be set.*/
 	UFUNCTION(BlueprintCallable, Category = "Interchange | Node | GeometryCache")
-	bool SetCustomFrameRate(const double& AttributeValue);
+	UE_API bool SetCustomFrameRate(const double& AttributeValue);
 
 	/** Query whether this animated mesh has constant topoplogy. Return false if the attribute was not set.*/
 	UFUNCTION(BlueprintCallable, Category = "Interchange | Node | GeometryCache")
-	bool GetCustomHasConstantTopology(bool& AttributeValue) const;
+	UE_API bool GetCustomHasConstantTopology(bool& AttributeValue) const;
 
 	/** Set the constant topology attribute of this animated mesh. Return false if the attribute could not be set.*/
 	UFUNCTION(BlueprintCallable, Category = "Interchange | Node | GeometryCache")
-	bool SetCustomHasConstantTopology(const bool& AttributeValue);
+	UE_API bool SetCustomHasConstantTopology(const bool& AttributeValue);
 
 private:
 	IMPLEMENT_NODE_ATTRIBUTE_KEY(StartFrame);
@@ -436,3 +438,5 @@ private:
 	IMPLEMENT_NODE_ATTRIBUTE_KEY(FrameRate);
 	IMPLEMENT_NODE_ATTRIBUTE_KEY(HasConstantTopology);
 };
+
+#undef UE_API
\ No newline at end of file
diff --git a/Engine/Plugins/Interchange/Runtime/Source/Nodes/Public/InterchangeTexture2DNode.h b/Engine/Plugins/Interchange/Runtime/Source/Nodes/Public/InterchangeTexture2DNode.h
--- a/Engine/Plugins/Interchange/Runtime/Source/Nodes/Public/InterchangeTexture2DNode.h
+++ b/Engine/Plugins/Interchange/Runtime/Source/Nodes/Public/InterchangeTexture2DNode.h
@@ -22,8 +22,8 @@ namespace UE::Interchange
 	};
 }//ns UE::Interchange
 
-UCLASS(BlueprintType)
-class INTERCHANGENODES_API UInterchangeTexture2DNode : public UInterchangeTextureNode
+UCLASS(BlueprintType, MinimalAPI)
+class UInterchangeTexture2DNode : public UInterchangeTextureNode
 {
 	GENERATED_BODY()
 
@@ -32,12 +32,12 @@ public:
 	/**
 	 * Build and return a UID name for a texture 2D node.
 	 */
-	static FString MakeNodeUid(const FStringView NodeName);
+	static INTERCHANGENODES_API FString MakeNodeUid(const FStringView NodeName);
 
 	/**
 	 * Creates a new UInterchangeTexture2DNode and adds it to NodeContainer as a translated node.
 	 */
-	static UInterchangeTexture2DNode* Create(UInterchangeBaseNodeContainer* NodeContainer, const FStringView TextureNodeName);
+	static INTERCHANGENODES_API UInterchangeTexture2DNode* Create(UInterchangeBaseNodeContainer* NodeContainer, const FStringView TextureNodeName);
 
 	virtual void PostInitProperties()
 	{
diff --git a/Engine/Plugins/Interchange/Runtime/Source/Nodes/Public/InterchangeTextureNode.h b/Engine/Plugins/Interchange/Runtime/Source/Nodes/Public/InterchangeTextureNode.h
--- a/Engine/Plugins/Interchange/Runtime/Source/Nodes/Public/InterchangeTextureNode.h
+++ b/Engine/Plugins/Interchange/Runtime/Source/Nodes/Public/InterchangeTextureNode.h
@@ -62,8 +62,8 @@ enum class EInterchangeTextureColorSpace : uint8
 	TCS_MAX,
 };
 
-UCLASS(BlueprintType, Abstract)
-class INTERCHANGENODES_API UInterchangeTextureNode : public UInterchangeBaseNode
+UCLASS(BlueprintType, Abstract, MinimalAPI)
+class UInterchangeTextureNode : public UInterchangeBaseNode
 {
 	GENERATED_BODY()
 
diff --git a/Engine/Plugins/Interchange/Runtime/Source/Parsers/CommonParser/Public/InterchangeCommonAnimationPayload.h b/Engine/Plugins/Interchange/Runtime/Source/Parsers/CommonParser/Public/InterchangeCommonAnimationPayload.h
--- a/Engine/Plugins/Interchange/Runtime/Source/Parsers/CommonParser/Public/InterchangeCommonAnimationPayload.h
+++ b/Engine/Plugins/Interchange/Runtime/Source/Parsers/CommonParser/Public/InterchangeCommonAnimationPayload.h
@@ -22,7 +22,7 @@ namespace UE::Interchange
 		INTERCHANGECOMMONPARSER_API FString HashString(const FString& String);
 	};
 
-	struct INTERCHANGECOMMONPARSER_API FAnimationPayloadData
+	struct FAnimationPayloadData
 	{
 		FString SceneNodeUniqueID;
 
@@ -58,13 +58,13 @@ namespace UE::Interchange
 		{
 		}
 
-		void SerializeBaked(FArchive& Ar);
+		INTERCHANGECOMMONPARSER_API void SerializeBaked(FArchive& Ar);
 
 		//Conversions:
-		void CalculateDataFor(const EInterchangeAnimationPayLoadType& ToType, const FTransform& DefaultTransform = FTransform());
+		INTERCHANGECOMMONPARSER_API void CalculateDataFor(const EInterchangeAnimationPayLoadType& ToType, const FTransform& DefaultTransform = FTransform());
 	};
 	
-	struct INTERCHANGECOMMONPARSER_API FAnimationTimeDescription
+	struct FAnimationTimeDescription
 	{
 		double BakeFrequency;
 		double RangeStartSecond;
@@ -91,7 +91,7 @@ namespace UE::Interchange
 		}
 	};
 
-	struct INTERCHANGECOMMONPARSER_API FAnimationPayloadQuery
+	struct FAnimationPayloadQuery
 	{
 		FString SceneNodeUniqueID;
 		FInterchangeAnimationPayLoadKey PayloadKey;
@@ -111,13 +111,13 @@ namespace UE::Interchange
 
 		}
 
-		FString GetHashString() const;
+		INTERCHANGECOMMONPARSER_API FString GetHashString() const;
 
-		FString ToJson() const;
-		void FromJson(const FString& JsonString);
+		INTERCHANGECOMMONPARSER_API FString ToJson() const;
+		INTERCHANGECOMMONPARSER_API void FromJson(const FString& JsonString);
 
-		static FString ToJson(const TArray<FAnimationPayloadQuery>& Queries);
-		static void FromJson(const FString& JsonString, TArray<FAnimationPayloadQuery>& Queries);
+		static INTERCHANGECOMMONPARSER_API FString ToJson(const TArray<FAnimationPayloadQuery>& Queries);
+		static INTERCHANGECOMMONPARSER_API void FromJson(const FString& JsonString, TArray<FAnimationPayloadQuery>& Queries);
 	private:
 		mutable TOptional<FString> HashStringCache;
 	};
@@ -170,7 +170,7 @@ enum class EInterchangeCurveTangentWeightMode : uint8
 * This struct contains only the key data, this is only used to pass animation data from translators to factories
 */
 USTRUCT()
-struct INTERCHANGECOMMONPARSER_API FInterchangeCurveKey
+struct FInterchangeCurveKey
 {
 	GENERATED_BODY()
 
@@ -212,10 +212,10 @@ struct INTERCHANGECOMMONPARSER_API FInterchangeCurveKey
 
 #if WITH_ENGINE
 	/** Conversion to FRichCurve */
-	void ToRichCurveKey(FRichCurveKey& OutKey) const;
+	INTERCHANGECOMMONPARSER_API void ToRichCurveKey(FRichCurveKey& OutKey) const;
 #endif
 
-	void Serialize(FArchive& Ar);
+	INTERCHANGECOMMONPARSER_API void Serialize(FArchive& Ar);
 
 	friend FArchive& operator<<(FArchive& Ar, FInterchangeCurveKey& InterchangeCurveKey)
 	{
@@ -228,7 +228,7 @@ struct INTERCHANGECOMMONPARSER_API FInterchangeCurveKey
 * This struct contains only the key data, this is only used to pass animation data from interchange worker process translators to factories.
 */
 USTRUCT()
-struct INTERCHANGECOMMONPARSER_API FInterchangeCurve
+struct FInterchangeCurve
 {
 	GENERATED_BODY()
 
@@ -237,10 +237,10 @@ struct INTERCHANGECOMMONPARSER_API FInterchangeCurve
 
 #if WITH_ENGINE
 	/** Conversion to FRichCurve */
-	void ToRichCurve(FRichCurve& OutKey) const;
+	INTERCHANGECOMMONPARSER_API void ToRichCurve(FRichCurve& OutKey) const;
 #endif
 
-	void Serialize(FArchive& Ar);
+	INTERCHANGECOMMONPARSER_API void Serialize(FArchive& Ar);
 
 	friend FArchive& operator<<(FArchive& Ar, FInterchangeCurve& InterchangeCurve)
 	{
@@ -253,7 +253,7 @@ struct INTERCHANGECOMMONPARSER_API FInterchangeCurve
 * This struct contains only the key data, this is only used to pass animation data from translators to factories.
 */
 USTRUCT()
-struct INTERCHANGECOMMONPARSER_API FInterchangeStepCurve
+struct FInterchangeStepCurve
 {
 	GENERATED_BODY()
 
@@ -265,9 +265,9 @@ struct INTERCHANGECOMMONPARSER_API FInterchangeStepCurve
 	TOptional<TArray<int32>> IntegerKeyValues;
 	TOptional<TArray<FString>> StringKeyValues;
 
-	void RemoveRedundantKeys(float Threshold);
+	INTERCHANGECOMMONPARSER_API void RemoveRedundantKeys(float Threshold);
 
-	void Serialize(FArchive& Ar);
+	INTERCHANGECOMMONPARSER_API void Serialize(FArchive& Ar);
 
 	friend FArchive& operator<<(FArchive& Ar, FInterchangeStepCurve& InterchangeStepCurve)
 	{
diff --git a/Engine/Plugins/Interchange/Runtime/Source/Parsers/GLTFCore/Public/GLTFAnimation.h b/Engine/Plugins/Interchange/Runtime/Source/Parsers/GLTFCore/Public/GLTFAnimation.h
--- a/Engine/Plugins/Interchange/Runtime/Source/Parsers/GLTFCore/Public/GLTFAnimation.h
+++ b/Engine/Plugins/Interchange/Runtime/Source/Parsers/GLTFCore/Public/GLTFAnimation.h
@@ -11,7 +11,7 @@ namespace GLTF
 	struct FNode;
 	struct FAccessor;
 
-	struct GLTFCORE_API FAnimation
+	struct FAnimation
 	{
 		enum class EInterpolation
 		{
diff --git a/Engine/Plugins/Interchange/Runtime/Source/Parsers/GLTFCore/Public/GLTFAsset.h b/Engine/Plugins/Interchange/Runtime/Source/Parsers/GLTFCore/Public/GLTFAsset.h
--- a/Engine/Plugins/Interchange/Runtime/Source/Parsers/GLTFCore/Public/GLTFAsset.h
+++ b/Engine/Plugins/Interchange/Runtime/Source/Parsers/GLTFCore/Public/GLTFAsset.h
@@ -52,7 +52,7 @@ namespace GLTF
 		Count
 	};
 
-	struct GLTFCORE_API FScene
+	struct FScene
 	{
 		FString       Name;
 		TArray<int32> Nodes;
@@ -60,7 +60,7 @@ namespace GLTF
 		FString	      UniqueId; //will be generated in FAsset::GenerateNames
 	};
 
-	struct GLTFCORE_API FMetadata
+	struct FMetadata
 	{
 		struct FExtraData
 		{
@@ -71,12 +71,12 @@ namespace GLTF
 		float              Version;
 		TArray<FExtraData> Extras;
 
-		const FExtraData* GetExtraData(const FString& Name) const;
+		GLTFCORE_API const FExtraData* GetExtraData(const FString& Name) const;
 	};
 
-	struct GLTFCORE_API FAsset : FNoncopyable
+	struct FAsset : FNoncopyable
 	{
-		static const TSet<EExtension> SupportedExtensions;
+		static GLTFCORE_API const TSet<EExtension> SupportedExtensions;
 
 		enum EValidationCheck
 		{
@@ -122,7 +122,7 @@ namespace GLTF
 		 * @param ExtraBinBufferKBytes -  bytes to reserve for the extra binary buffer(eg. image, mime data, etc.)
 		 * @note Only reserves buffers if they had any existing data.
 		 */
-		void Clear(uint32 BinBufferKBytes, uint32 ExtraBinBufferKBytes);
+		GLTFCORE_API void Clear(uint32 BinBufferKBytes, uint32 ExtraBinBufferKBytes);
 
 		/**
 		 * Will generate names for any entities(nodes, meshes, etc.) that have the name field missing.
@@ -130,24 +130,24 @@ namespace GLTF
 		 *
 		 * @param Prefix - prefix to add to the entities name.
 		 */
-		void GenerateNames();
+		GLTFCORE_API void GenerateNames();
 
 		/**
 		 * Finds the indices for the nodes which are root nodes.
 		 *
 		 * @param NodeIndices - array with the results
 		 */
-		void GetRootNodes(TArray<int32>& NodeIndices);
+		GLTFCORE_API void GetRootNodes(TArray<int32>& NodeIndices);
 
 		/**
 		 * Returns EValidationCheck::Valid if the asset passes the post-import validation checks.
 		 */
-		EValidationCheck ValidationCheck() const;
+		GLTFCORE_API EValidationCheck ValidationCheck() const;
 
 		/**
 		* Creates and sets up FBuffer and FBufferView for a given FAccessor (at the given index 'id')
 		*/
-		FAccessor& CreateBuffersForAccessorIndex(uint32 AccessorIndex);
+		GLTFCORE_API FAccessor& CreateBuffersForAccessorIndex(uint32 AccessorIndex);
 
 	private:
 		// Binary glTF files can have embedded data after JSON.
diff --git a/Engine/Plugins/Interchange/Runtime/Source/Parsers/GLTFCore/Public/GLTFMaterial.h b/Engine/Plugins/Interchange/Runtime/Source/Parsers/GLTFCore/Public/GLTFMaterial.h
--- a/Engine/Plugins/Interchange/Runtime/Source/Parsers/GLTFCore/Public/GLTFMaterial.h
+++ b/Engine/Plugins/Interchange/Runtime/Source/Parsers/GLTFCore/Public/GLTFMaterial.h
@@ -7,7 +7,7 @@
 
 namespace GLTF
 {
-	struct GLTFCORE_API FTextureTransform
+	struct FTextureTransform
 	{
 		float Offset[2];
 		float Rotation;
@@ -21,7 +21,7 @@ namespace GLTF
 		}
 	};
 
-	struct GLTFCORE_API FTextureMap
+	struct FTextureMap
 	{
 		int32 TextureIndex;
 		uint8 TexCoord;
@@ -47,7 +47,7 @@ namespace GLTF
 		}
 	};
 
-	struct GLTFCORE_API FMaterial
+	struct FMaterial
 	{
 		enum class EAlphaMode
 		{
@@ -227,6 +227,6 @@ namespace GLTF
 			return AlphaMode == EAlphaMode::Opaque;
 		}
 
-		FMD5Hash GetHash() const;
+		GLTFCORE_API FMD5Hash GetHash() const;
 	};
 }  // namespace GLTF
diff --git a/Engine/Plugins/Interchange/Runtime/Source/Parsers/GLTFCore/Public/GLTFMesh.h b/Engine/Plugins/Interchange/Runtime/Source/Parsers/GLTFCore/Public/GLTFMesh.h
--- a/Engine/Plugins/Interchange/Runtime/Source/Parsers/GLTFCore/Public/GLTFMesh.h
+++ b/Engine/Plugins/Interchange/Runtime/Source/Parsers/GLTFCore/Public/GLTFMesh.h
@@ -11,7 +11,7 @@ struct FMD5Hash;
 namespace GLTF
 {
 	// for storing triangle indices as a unit
-	struct GLTFCORE_API FTriangle
+	struct FTriangle
 	{
 		uint32 A, B, C;
 
@@ -23,7 +23,7 @@ namespace GLTF
 		}
 	};
 
-	struct GLTFCORE_API FJointInfluence
+	struct FJointInfluence
 	{
 		FVector4 Weight;
 		uint16   ID[4];
@@ -35,59 +35,59 @@ namespace GLTF
 		}
 	};
 
-	struct GLTFCORE_API FVariantMapping
+	struct FVariantMapping
 	{
 		int32 MaterialIndex;
 		TArray<int32> VariantIndices;
 	};
 
-	struct GLTFCORE_API FAttributeAccessors
+	struct FAttributeAccessors
 	{
 		//Expects elements in the exact same order as the EMeshAttributeType.
-		FAttributeAccessors(const TArray<TPair<EMeshAttributeType, FAccessor&>>& InAttributeAccessors, bool bInMorphTarget);
+		GLTFCORE_API FAttributeAccessors(const TArray<TPair<EMeshAttributeType, FAccessor&>>& InAttributeAccessors, bool bInMorphTarget);
 
 	protected:
-		bool HasAttributeAccessor(const EMeshAttributeType& Type) const;
-		const FAccessor& GetAttributeAccessor(const EMeshAttributeType& Type) const;
+		GLTFCORE_API bool HasAttributeAccessor(const EMeshAttributeType& Type) const;
+		GLTFCORE_API const FAccessor& GetAttributeAccessor(const EMeshAttributeType& Type) const;
 		
-		bool AreAttributeAccessorsValid(bool& bHasData, uint32 ExpectedElementCount) const;
+		GLTFCORE_API bool AreAttributeAccessorsValid(bool& bHasData, uint32 ExpectedElementCount) const;
 
-		void HashAttributes(FMD5& MD5) const;
+		GLTFCORE_API void HashAttributes(FMD5& MD5) const;
 
 	private:
 		TArray<TPair<EMeshAttributeType, FAccessor&>> AttributeAccessors;
 		bool bMorphTarget;
 	};
 
-	struct GLTFCORE_API FMorphTarget : FAttributeAccessors
+	struct FMorphTarget : FAttributeAccessors
 	{
-		FMorphTarget(const TArray<TPair<EMeshAttributeType, FAccessor&>>& InAttributeAccessors /*Displacements and Deltas*/);
+		GLTFCORE_API FMorphTarget(const TArray<TPair<EMeshAttributeType, FAccessor&>>& InAttributeAccessors /*Displacements and Deltas*/);
 
-		bool IsValid(uint32 ExpectedElementCount) const;
-		FMD5Hash GetHash() const;
+		GLTFCORE_API bool IsValid(uint32 ExpectedElementCount) const;
+		GLTFCORE_API FMD5Hash GetHash() const;
 
 		bool HasPositionDisplacements() const { return HasAttributeAccessor(EMeshAttributeType::POSITION); };
-		void GetPositionDisplacements(TArray<FVector3f>& Buffer) const;
-		int32 GetNumberOfPositionDisplacements() const;
+		GLTFCORE_API void GetPositionDisplacements(TArray<FVector3f>& Buffer) const;
+		GLTFCORE_API int32 GetNumberOfPositionDisplacements() const;
 
 		bool HasNormalDisplacements() const { return HasAttributeAccessor(EMeshAttributeType::NORMAL); };
-		void GetNormalDisplacements(TArray<FVector3f>& Buffer) const;
-		int32 GetNumberOfNormalDisplacements() const;
+		GLTFCORE_API void GetNormalDisplacements(TArray<FVector3f>& Buffer) const;
+		GLTFCORE_API int32 GetNumberOfNormalDisplacements() const;
 
 		bool HasTangentDisplacements() const { return HasAttributeAccessor(EMeshAttributeType::TANGENT); };
-		void GetTangentDisplacements(TArray<FVector4f>& Buffer) const;
-		int32 GetNumberOfTangentDisplacements() const;
+		GLTFCORE_API void GetTangentDisplacements(TArray<FVector4f>& Buffer) const;
+		GLTFCORE_API int32 GetNumberOfTangentDisplacements() const;
 
-		bool HasTexCoordDisplacements(uint32 Index) const;
-		void GetTexCoordDisplacements(uint32 Index, TArray<FVector2f>& Buffer) const;
-		int32 GetNumberOfTexCoordDisplacements(uint32 Index) const;
+		GLTFCORE_API bool HasTexCoordDisplacements(uint32 Index) const;
+		GLTFCORE_API void GetTexCoordDisplacements(uint32 Index, TArray<FVector2f>& Buffer) const;
+		GLTFCORE_API int32 GetNumberOfTexCoordDisplacements(uint32 Index) const;
 
 		bool HasColorDeltas() const { return HasAttributeAccessor(EMeshAttributeType::COLOR_0); };
-		void GetColorDeltas(TArray<FVector4f>& Buffer) const;
-		int32 GetNumberOfColorDeltas() const;
+		GLTFCORE_API void GetColorDeltas(TArray<FVector4f>& Buffer) const;
+		GLTFCORE_API int32 GetNumberOfColorDeltas() const;
 	};
 
-	struct GLTFCORE_API FPrimitive : FAttributeAccessors
+	struct FPrimitive : FAttributeAccessors
 	{
 		enum class EMode
 		{
@@ -105,8 +105,8 @@ namespace GLTF
 			Unknown = 7
 		};
 
-		static const TArray<EMode> SupportedModes;
-		static FString ToString(const EMode& Mode);
+		static GLTFCORE_API const TArray<EMode> SupportedModes;
+		static GLTFCORE_API FString ToString(const EMode& Mode);
 
 		const EMode Mode;
 		const int32 MaterialIndex;
@@ -114,31 +114,31 @@ namespace GLTF
 
 		TMap<FString, FString> Extras;
 
-		FPrimitive(EMode InMode, int32 InMaterial, const FAccessor& InIndices, const TArray<TPair<EMeshAttributeType, FAccessor&>>& InAttributeAccessors);
+		GLTFCORE_API FPrimitive(EMode InMode, int32 InMaterial, const FAccessor& InIndices, const TArray<TPair<EMeshAttributeType, FAccessor&>>& InAttributeAccessors);
 
-		void GenerateIsValidCache();
-		bool IsValid() const;
-		FMD5Hash GetHash() const;
+		GLTFCORE_API void GenerateIsValidCache();
+		GLTFCORE_API bool IsValid() const;
+		GLTFCORE_API FMD5Hash GetHash() const;
 
 		bool HasPositions() const { return HasAttributeAccessor(EMeshAttributeType::POSITION); }
 		bool HasNormals() const { return HasAttributeAccessor(EMeshAttributeType::NORMAL); }
 		bool HasTangents() const { return HasAttributeAccessor(EMeshAttributeType::TANGENT); }
-		bool HasTexCoords(uint32 Index) const;
+		GLTFCORE_API bool HasTexCoords(uint32 Index) const;
 		bool HasColors() const { return HasAttributeAccessor(EMeshAttributeType::COLOR_0); }
 		bool HasJointWeights() const { return HasAttributeAccessor(EMeshAttributeType::JOINTS_0) && HasAttributeAccessor(EMeshAttributeType::WEIGHTS_0); }
 
-		void GetPositions(TArray<FVector3f>& Buffer) const;
-		void GetNormals(TArray<FVector3f>& Buffer) const;
-		void GetTangents(TArray<FVector4f>& Buffer) const;
-		void GetTexCoords(uint32 Index, TArray<FVector2f>& Buffer) const;
-		void GetColors(TArray<FVector4f>& Buffer) const;
-		void GetJointInfluences(TArray<FJointInfluence>& Buffer) const;
+		GLTFCORE_API void GetPositions(TArray<FVector3f>& Buffer) const;
+		GLTFCORE_API void GetNormals(TArray<FVector3f>& Buffer) const;
+		GLTFCORE_API void GetTangents(TArray<FVector4f>& Buffer) const;
+		GLTFCORE_API void GetTexCoords(uint32 Index, TArray<FVector2f>& Buffer) const;
+		GLTFCORE_API void GetColors(TArray<FVector4f>& Buffer) const;
+		GLTFCORE_API void GetJointInfluences(TArray<FJointInfluence>& Buffer) const;
 
-		FTriangle TriangleVerts(uint32 T) const;
-		void      GetTriangleIndices(TArray<uint32>& Buffer) const;
+		GLTFCORE_API FTriangle TriangleVerts(uint32 T) const;
+		GLTFCORE_API void      GetTriangleIndices(TArray<uint32>& Buffer) const;
 
-		uint32 VertexCount() const;
-		uint32 TriangleCount() const;
+		GLTFCORE_API uint32 VertexCount() const;
+		GLTFCORE_API uint32 TriangleCount() const;
 
 		//
 		TArray<FMorphTarget> MorphTargets;
@@ -158,7 +158,7 @@ namespace GLTF
 	};
 
 
-	struct GLTFCORE_API FMesh
+	struct FMesh
 	{
 		FString				Name;
 		TArray<FPrimitive>	Primitives;
@@ -178,9 +178,9 @@ namespace GLTF
 
 		void GenerateIsValidCache(bool GenerateIsValidCacheForPrimitives = true);
 		bool IsValid() const;
-		FMD5Hash GetHash() const;
+		GLTFCORE_API FMD5Hash GetHash() const;
 
-		int32 NumberOfMorphTargetsPerPrimitive() const;
+		GLTFCORE_API int32 NumberOfMorphTargetsPerPrimitive() const;
 
 	private:
 		bool IsValidPrivate() const;
diff --git a/Engine/Plugins/Interchange/Runtime/Source/Parsers/GLTFCore/Public/GLTFNode.h b/Engine/Plugins/Interchange/Runtime/Source/Parsers/GLTFCore/Public/GLTFNode.h
--- a/Engine/Plugins/Interchange/Runtime/Source/Parsers/GLTFCore/Public/GLTFNode.h
+++ b/Engine/Plugins/Interchange/Runtime/Source/Parsers/GLTFCore/Public/GLTFNode.h
@@ -11,9 +11,9 @@ namespace GLTF { struct FAccessor; }
 
 namespace GLTF
 {
-	struct GLTFCORE_API FNode
+	struct FNode
 	{
-		struct GLTFCORE_API FLightInstanceIES
+		struct FLightInstanceIES
 		{
 			int32 Index = INDEX_NONE; // index into FAsset::LightsIES
 			
@@ -105,7 +105,7 @@ namespace GLTF
 		}
 	};
 
-	struct GLTFCORE_API FCamera
+	struct FCamera
 	{
 		struct FPerspective
 		{
@@ -147,7 +147,7 @@ namespace GLTF
 		}
 	};
 
-	struct GLTFCORE_API FLightIES
+	struct FLightIES
 	{
 		int32 Index; //serves as idintifer as well.
 		
@@ -164,7 +164,7 @@ namespace GLTF
 		FString UniqueId; //will be generated in FAsset::GenerateNames
 	};
 
-	struct GLTFCORE_API FLight
+	struct FLight
 	{
 		enum class EType
 		{
@@ -206,7 +206,7 @@ namespace GLTF
 		}
 	};
 
-	struct GLTFCORE_API FSkinInfo
+	struct FSkinInfo
 	{
 		const FAccessor& InverseBindMatrices;
 		FString          Name;
diff --git a/Engine/Plugins/Interchange/Runtime/Source/Parsers/GLTFCore/Public/GLTFTexture.h b/Engine/Plugins/Interchange/Runtime/Source/Parsers/GLTFCore/Public/GLTFTexture.h
--- a/Engine/Plugins/Interchange/Runtime/Source/Parsers/GLTFCore/Public/GLTFTexture.h
+++ b/Engine/Plugins/Interchange/Runtime/Source/Parsers/GLTFCore/Public/GLTFTexture.h
@@ -8,7 +8,7 @@ struct FScriptContainerElement;
 
 namespace GLTF
 {
-	struct GLTFCORE_API FImage
+	struct FImage
 	{
 		enum class EFormat
 		{
@@ -37,7 +37,7 @@ namespace GLTF
 		}
 	};
 
-	struct GLTFCORE_API FSampler
+	struct FSampler
 	{
 		enum class EFilter
 		{
@@ -58,7 +58,7 @@ namespace GLTF
 			ClampToEdge    = 33071
 		};
 
-		static const FSampler DefaultSampler;
+		static GLTFCORE_API const FSampler DefaultSampler;
 
 		EFilter MinFilter;
 		EFilter MagFilter;
@@ -76,7 +76,7 @@ namespace GLTF
 		}
 	};
 
-	struct GLTFCORE_API FTexture
+	struct FTexture
 	{
 		const FImage&   Source;
 		const FSampler& Sampler;
