[Interchange-glTF] Import crash fix for Sparse buffers without Base bufferviews.

#rnx

diff --git a/Engine/Plugins/Interchange/Runtime/Source/Parsers/GLTFCore/Private/GLTF/GLTFAccessor.cpp b/Engine/Plugins/Interchange/Runtime/Source/Parsers/GLTFCore/Private/GLTF/GLTFAccessor.cpp
--- a/Engine/Plugins/Interchange/Runtime/Source/Parsers/GLTFCore/Private/GLTF/GLTFAccessor.cpp
+++ b/Engine/Plugins/Interchange/Runtime/Source/Parsers/GLTFCore/Private/GLTF/GLTFAccessor.cpp
@@ -1057,7 +1057,7 @@ namespace GLTF
 			if (Type == EType::Scalar && !bNormalized)
 			{
 				uint32 Buffer;
-				if (AcquireUnsignedInt(Index, BufferView, ByteOffset, ByteStride, ComponentType, Buffer))
+				if (!BufferView.IsValid() || AcquireUnsignedInt(Index, BufferView, ByteOffset, ByteStride, ComponentType, Buffer))
 				{
 					UpdateUnsignedIntWithSparse(Index, Buffer);
 					return Buffer;
@@ -1077,7 +1077,7 @@ namespace GLTF
 		{
 			if (Type == EType::Vec4 && !bNormalized)
 			{
-				if (AcquireUnsignedInt16x4(Index, BufferView, ByteOffset, ByteStride, ComponentType, Values))
+				if (!BufferView.IsValid() || AcquireUnsignedInt16x4(Index, BufferView, ByteOffset, ByteStride, ComponentType, Values))
 				{
 					UpdateUnsignedInt16x4WithSparse(Index, Values);
 					return;
@@ -1096,7 +1096,7 @@ namespace GLTF
 			if (Type == EType::Scalar && !bNormalized)
 			{
 				float Buffer;
-				if (AcquireFloat(Index, BufferView, ByteOffset, ByteStride, ComponentType, Buffer))
+				if (!BufferView.IsValid() || AcquireFloat(Index, BufferView, ByteOffset, ByteStride, ComponentType, Buffer))
 				{
 					UpdateFloatWithSparse(Index, Buffer);
 					return Buffer;
@@ -1121,7 +1121,7 @@ namespace GLTF
 			if (Type == EType::Vec2)  // strict format match, unlike GPU shader fetch
 			{
 				FVector2D Buffer;
-				if (AcquireVec2(Index, BufferView, ByteOffset, ByteStride, ComponentType, bNormalized, Buffer))
+				if (!BufferView.IsValid() || AcquireVec2(Index, BufferView, ByteOffset, ByteStride, ComponentType, bNormalized, Buffer))
 				{
 					UpdateVec2WithSparse(Index, Buffer);
 					return Buffer;
@@ -1146,7 +1146,7 @@ namespace GLTF
 			if (Type == EType::Vec3)  // strict format match, unlike GPU shader fetch
 			{
 				FVector Buffer;
-				if (AcquireVec3(Index, BufferView, ByteOffset, ByteStride, ComponentType, bNormalized, Buffer))
+				if (!BufferView.IsValid() || AcquireVec3(Index, BufferView, ByteOffset, ByteStride, ComponentType, bNormalized, Buffer))
 				{
 					UpdateVec3WithSparse(Index, Buffer);
 					return Buffer;
@@ -1171,7 +1171,7 @@ namespace GLTF
 			if (Type == EType::Vec4)  // strict format match, unlike GPU shader fetch
 			{
 				FVector4 Buffer;
-				if (AcquireVec4(Index, BufferView, ByteOffset, ByteStride, ComponentType, bNormalized, Buffer))
+				if (!BufferView.IsValid() || AcquireVec4(Index, BufferView, ByteOffset, ByteStride, ComponentType, bNormalized, Buffer))
 				{
 					UpdateVec4WithSparse(Index, Buffer);
 					return Buffer;
@@ -1192,7 +1192,7 @@ namespace GLTF
 			if (Type == EType::Mat4)  // strict format match, unlike GPU shader fetch
 			{
 				FMatrix Buffer;
-				if (AcquireMat4(Index, BufferView, ByteOffset, ByteStride, ComponentType, Buffer))
+				if (!BufferView.IsValid() || AcquireMat4(Index, BufferView, ByteOffset, ByteStride, ComponentType, Buffer))
 				{
 					UpdateMat4WithSparse(Index, Buffer);
 					return Buffer;
@@ -1208,7 +1208,7 @@ namespace GLTF
 	{
 		if (Type == EType::Scalar && !bNormalized)
 		{
-			if (FillUnsignedIntArray(BufferView, ByteOffset, ComponentType, Count, Buffer))
+			if (!BufferView.IsValid() || FillUnsignedIntArray(BufferView, ByteOffset, ComponentType, Count, Buffer))
 			{
 				UpdateUnsignedIntArrayWithSparse(Buffer);
 				return;
@@ -1222,7 +1222,7 @@ namespace GLTF
 	{
 		if (Type == EType::Scalar)
 		{
-			if (FillFloatArray(bNormalized, BufferView, ByteOffset, ComponentType, ByteStride, Count, Buffer))
+			if (!BufferView.IsValid() || FillFloatArray(bNormalized, BufferView, ByteOffset, ComponentType, ByteStride, Count, Buffer))
 			{
 				UpdateFloatArrayWithSparse(Buffer);
 				return;
