Fixed crash on skeletal mesh reimport
- Root cause: skeletal mesh's rendering resources were not properly released nor skeletal mesh components using the reimported skeletal mesh properly updated.
- Solution: added the use of  FScopedSkeletalMeshPostEditChange when reimporting a skeletal mesh.
- Note: THis is what legacy FBX was doing on the reimport of skeletal mesh.

#rnx

diff --git a/Engine/Plugins/Interchange/Runtime/Source/Import/Private/Mesh/InterchangeSkeletalMeshFactory.cpp b/Engine/Plugins/Interchange/Runtime/Source/Import/Private/Mesh/InterchangeSkeletalMeshFactory.cpp
--- a/Engine/Plugins/Interchange/Runtime/Source/Import/Private/Mesh/InterchangeSkeletalMeshFactory.cpp
+++ b/Engine/Plugins/Interchange/Runtime/Source/Import/Private/Mesh/InterchangeSkeletalMeshFactory.cpp
@@ -63,6 +63,29 @@
 
 #endif //WITH_EDITORONLY_DATA
 
+FScopedSkeletalMeshReimportUtility::FScopedSkeletalMeshReimportUtility(USkeletalMesh* InSkeletalMesh)
+{
+#if WITH_EDITOR
+	// Scope PostEditChange before locking the skeletal mesh's properties
+	ScopedPostEditChange = MakeUnique<FScopedSkeletalMeshPostEditChange>(InSkeletalMesh);
+	// Lock the skeletal mesh's properties to prevent any async access
+	LockPropertiesEvent = InSkeletalMesh->LockPropertiesUntil();
+#endif
+}
+
+FScopedSkeletalMeshReimportUtility::~FScopedSkeletalMeshReimportUtility()
+{
+#if WITH_EDITOR
+	// Unlock the skeletal mesh's properties first
+	if (ensure(LockPropertiesEvent))
+	{
+		LockPropertiesEvent->Trigger();
+	}
+	// Now apply PostEditChange to update skeletal mesh's resources
+	ScopedPostEditChange.Reset();
+#endif
+}
+
 #if WITH_EDITOR
 namespace UE
 {
@@ -1004,13 +1027,13 @@ UInterchangeFactoryBase::FImportAssetResult UInterchangeSkeletalMeshFactory::Beg
 		}
 		return ImportAssetResult;
 	}
-	
+
 	SkeletalMesh->PreEditChange(nullptr);
 
 	//Lock the skeletalmesh properties if the skeletal mesh already exist (re-import)
 	if (ExistingAsset)
 	{
-		SkeletalMeshLockPropertiesEvent = SkeletalMesh->LockPropertiesUntil();
+		ScopedReimportUtility = MakeUnique<FScopedSkeletalMeshReimportUtility>(SkeletalMesh);
 	}
 
 	ImportAssetResult.ImportedObject = SkeletalMesh;
@@ -1937,11 +1960,9 @@ UInterchangeFactoryBase::FImportAssetResult UInterchangeSkeletalMeshFactory::End
 
 	ImportAssetResult.ImportedObject = SkeletalMesh;
 
-	//Release the promise so we unlock the skeletalmesh properties
-	if (SkeletalMeshLockPropertiesEvent)
+	if (ScopedReimportUtility.IsValid())
 	{
-		SkeletalMeshLockPropertiesEvent->Trigger();
-		SkeletalMeshLockPropertiesEvent = nullptr;
+		ScopedReimportUtility.Reset();
 	}
 #endif
 
@@ -1950,10 +1971,9 @@ UInterchangeFactoryBase::FImportAssetResult UInterchangeSkeletalMeshFactory::End
 
 void UInterchangeSkeletalMeshFactory::Cancel()
 {
-	if (SkeletalMeshLockPropertiesEvent)
+	if (ScopedReimportUtility.IsValid())
 	{
-		SkeletalMeshLockPropertiesEvent->Trigger();
-		SkeletalMeshLockPropertiesEvent = nullptr;
+		ScopedReimportUtility.Reset();
 	}
 }
 
@@ -2403,4 +2423,5 @@ bool UInterchangeSkeletalMeshFactory::FImportAssetObjectData::IsValid() const
 	}
 
 	return true;
-}
\ No newline at end of file
+}
+
diff --git a/Engine/Plugins/Interchange/Runtime/Source/Import/Public/Mesh/InterchangeSkeletalMeshFactory.h b/Engine/Plugins/Interchange/Runtime/Source/Import/Public/Mesh/InterchangeSkeletalMeshFactory.h
--- a/Engine/Plugins/Interchange/Runtime/Source/Import/Public/Mesh/InterchangeSkeletalMeshFactory.h
+++ b/Engine/Plugins/Interchange/Runtime/Source/Import/Public/Mesh/InterchangeSkeletalMeshFactory.h
@@ -9,6 +9,7 @@
 #include "InterchangeMeshNode.h"
 #include "Mesh/InterchangeMeshPayload.h"
 #include "Rendering/SkeletalMeshLODImporterData.h"
+#include "SkeletalMeshTypes.h"
 #include "UObject/Object.h"
 #include "UObject/ObjectMacros.h"
 
@@ -19,6 +20,29 @@ class UInterchangeSkeletalMeshFactoryNode;
 class USkeletalMesh;
 class USkeleton;
 
+// Utility class used on reimport only to: 
+// - Call PostEditchange at the end of the reimport which will reallocate the render resource
+// - Lock the skeletal mesh's properties while it is updated
+class FScopedSkeletalMeshReimportUtility
+{
+public:
+	/*
+	 * This constructor increment the skeletal mesh PostEditChangeStackCounter. If the stack counter is zero before the increment
+	 * the skeletal mesh component will be unregister from the world. The component will also release there rendering resources.
+	 * Parameters:
+	 * @param InbCallPostEditChange - if we are the first scope PostEditChange will be called.
+	 * @param InbReregisterComponents - if we are the first scope we will re register component from world and also component render data.
+	 */
+	FScopedSkeletalMeshReimportUtility(USkeletalMesh* InSkeletalMesh);
+	~FScopedSkeletalMeshReimportUtility();
+
+private:
+#if WITH_EDITOR
+	TUniquePtr<FScopedSkeletalMeshPostEditChange> ScopedPostEditChange;
+#endif
+	FEvent* LockPropertiesEvent = nullptr;
+};
+
 namespace UE::Interchange
 {
 	//Get the mesh node context for each MeshUids
@@ -105,9 +129,7 @@ public:
 	};
 
 private:
-	FEvent* SkeletalMeshLockPropertiesEvent = nullptr;
-
-	
+	TUniquePtr<FScopedSkeletalMeshReimportUtility> ScopedReimportUtility;
 
 	TMap<int32, FLodPayloads> PayloadsPerLodIndex;
 
