Interchange Analytics:
- Add MaterialX file version metadata

diff --git a/Engine/Plugins/Interchange/Runtime/Source/Import/Private/MaterialX/MaterialXUtils/MaterialXManager.cpp b/Engine/Plugins/Interchange/Runtime/Source/Import/Private/MaterialX/MaterialXUtils/MaterialXManager.cpp
--- a/Engine/Plugins/Interchange/Runtime/Source/Import/Private/MaterialX/MaterialXUtils/MaterialXManager.cpp
+++ b/Engine/Plugins/Interchange/Runtime/Source/Import/Private/MaterialX/MaterialXUtils/MaterialXManager.cpp
@@ -2,6 +2,7 @@
 
 #include "MaterialX/MaterialXUtils/MaterialXManager.h"
 #include "InterchangeImportModule.h"
+#include "Nodes/InterchangeSourceNode.h"
 #include "Misc/PackageName.h"
 #include "UObject/SoftObjectPath.h"
 
@@ -411,6 +412,8 @@ namespace
 
 bool FMaterialXManager::Translate(const FString& Filename, UInterchangeBaseNodeContainer& BaseNodeContainer, const UInterchangeTranslatorBase* Translator)
 {
+	using namespace UE::Interchange;
+
 	bool bIsDocumentValid = false;
 
 	if (!FPaths::FileExists(Filename))
@@ -449,6 +452,12 @@ bool FMaterialXManager::Translate(const FString& Filename, UInterchangeBaseNodeC
 		Document->importLibrary(MaterialXLibrary);
 
 		bIsDocumentValid = Translate(Document, BaseNodeContainer, Translator);
+
+		if (Document->hasVersionString())
+		{
+			UInterchangeSourceNode* SourceNode = UInterchangeSourceNode::FindOrCreateUniqueInstance(&BaseNodeContainer);
+			SourceNode->SetExtraInformation(FSourceNodeExtraInfoStaticData::GetApplicationVersionExtraInfoKey(), ANSI_TO_TCHAR(Document->getVersionString().c_str()));
+		}
 	}
 	catch (std::exception& Exception)
 	{
