Interchange MaterialX:
- Return the ShaderGraphNode if it was already in the node container. Fixes wrong connections on some materials that are used more than once (especially in USD).

diff --git a/Engine/Plugins/Interchange/Runtime/Source/Import/Private/MaterialX/MaterialXUtils/MaterialXSurfaceMaterial.cpp b/Engine/Plugins/Interchange/Runtime/Source/Import/Private/MaterialX/MaterialXUtils/MaterialXSurfaceMaterial.cpp
--- a/Engine/Plugins/Interchange/Runtime/Source/Import/Private/MaterialX/MaterialXUtils/MaterialXSurfaceMaterial.cpp
+++ b/Engine/Plugins/Interchange/Runtime/Source/Import/Private/MaterialX/MaterialXUtils/MaterialXSurfaceMaterial.cpp
@@ -38,6 +38,10 @@ UInterchangeBaseNode* FMaterialXSurfaceMaterial::Translate(MaterialX::NodePtr Su
 		NodeContainer.ClearNodeParentUid(ShaderGraphNodeUID);
 		bIsShaderGraphInContainer = false;
 	}
+	else
+	{
+		return ShaderGraphNode;
+	}
 
 	mx::InputPtr InputNormal;
 	mx::NodePtr ConnectedSurfaceShaderNode, ConnectedDisplacementShaderNode;
