[Interchange] Updating Analytics tracking for Translators.

#rnx

diff --git a/Engine/Plugins/Experimental/Enterprise/DatasmithInterchange/Source/DatasmithInterchange/Private/InterchangeDatasmithTranslator.cpp b/Engine/Plugins/Experimental/Enterprise/DatasmithInterchange/Source/DatasmithInterchange/Private/InterchangeDatasmithTranslator.cpp
--- a/Engine/Plugins/Experimental/Enterprise/DatasmithInterchange/Source/DatasmithInterchange/Private/InterchangeDatasmithTranslator.cpp
+++ b/Engine/Plugins/Experimental/Enterprise/DatasmithInterchange/Source/DatasmithInterchange/Private/InterchangeDatasmithTranslator.cpp
@@ -760,8 +760,9 @@ TOptional<UE::Interchange::FImportImage> UInterchangeDatasmithTranslator::GetTex
 		return TOptional<UE::Interchange::FImportImage>();
 	}
 
-	UE::Interchange::Private::FScopedTranslator ScopedTranslator(TextureElement->GetFile(), Results);
+	UE::Interchange::Private::FScopedTranslator ScopedTranslator(TextureElement->GetFile(), Results, AnalyticsHelper);
 	const IInterchangeTexturePayloadInterface* TextureTranslator = ScopedTranslator.GetPayLoadInterface< IInterchangeTexturePayloadInterface >();
+
 	if (!ensure(TextureTranslator))
 	{
 		return TOptional<UE::Interchange::FImportImage>();
@@ -780,7 +781,7 @@ TOptional<UE::Interchange::FImportLightProfile> UInterchangeDatasmithTranslator:
 	}
 
 	UInterchangeSourceData* PayloadSourceData = UInterchangeManager::GetInterchangeManager().CreateSourceData(PayloadKey);
-	UE::Interchange::Private::FScopedTranslator ScopedTranslator(PayloadKey, Results);
+	UE::Interchange::Private::FScopedTranslator ScopedTranslator(PayloadKey, Results, AnalyticsHelper);
 	const IInterchangeTextureLightProfilePayloadInterface* TextureTranslator = ScopedTranslator.GetPayLoadInterface< IInterchangeTextureLightProfilePayloadInterface >();
 	if (!ensure(TextureTranslator))
 	{
diff --git a/Engine/Plugins/Interchange/Extensions/OpenUSD/Source/Import/Private/InterchangeUsdTranslator.cpp b/Engine/Plugins/Interchange/Extensions/OpenUSD/Source/Import/Private/InterchangeUsdTranslator.cpp
--- a/Engine/Plugins/Interchange/Extensions/OpenUSD/Source/Import/Private/InterchangeUsdTranslator.cpp
+++ b/Engine/Plugins/Interchange/Extensions/OpenUSD/Source/Import/Private/InterchangeUsdTranslator.cpp
@@ -4718,7 +4718,8 @@ namespace UE::InterchangeUsdTranslator::Private
 		TOptional<FString>& AlternateTexturePath,
 		const UInterchangeTranslatorBase* Translator,
 		UInterchangeUSDTranslatorImpl& Impl,
-		UInterchangeResultsContainer* Results
+		UInterchangeResultsContainer* Results,
+		TSharedPtr<UE::Interchange::FAnalyticsHelper> AnalyticsHelper
 	)
 	{
 		TRACE_CPUPROFILER_EVENT_SCOPE(GetMaterialXTexturePayloadData)
@@ -4736,7 +4737,7 @@ namespace UE::InterchangeUsdTranslator::Private
 			CompressionSettings = TextureCompressionSettings(FCString::Atoi(*PayloadKey.Mid(IndexTextureCompression + 1)));
 		}
 #endif
-		UE::Interchange::Private::FScopedTranslator ScopedTranslator(Filename, Results);
+		UE::Interchange::Private::FScopedTranslator ScopedTranslator(Filename, Results, AnalyticsHelper);
 		const IInterchangeTexturePayloadInterface* TextureTranslator = ScopedTranslator.GetPayLoadInterface<IInterchangeTexturePayloadInterface>();
 
 		if (!TextureTranslator)
@@ -4781,7 +4782,7 @@ TOptional<UE::Interchange::FImportImage> UInterchangeUSDTranslator::GetTexturePa
 	if (bDecoded)
 	{
 		// Defer back to another translator to actually parse the texture raw data
-		UE::Interchange::Private::FScopedTranslator ScopedTranslator(FilePath, Results);
+		UE::Interchange::Private::FScopedTranslator ScopedTranslator(FilePath, Results, AnalyticsHelper);
 		const IInterchangeTexturePayloadInterface* TextureTranslator = ScopedTranslator.GetPayLoadInterface<IInterchangeTexturePayloadInterface>();
 		if (TextureTranslator)
 		{
@@ -4822,7 +4823,7 @@ TOptional<UE::Interchange::FImportImage> UInterchangeUSDTranslator::GetTexturePa
 	// copying the behavior from InterchangeMaterialXTranslator
 	if (!TexturePayloadData.IsSet())
 	{
-		TexturePayloadData = GetMaterialXTexturePayloadData(PayloadKey, AlternateTexturePath, this, *ImplPtr, Results);
+		TexturePayloadData = GetMaterialXTexturePayloadData(PayloadKey, AlternateTexturePath, this, *ImplPtr, Results, AnalyticsHelper);
 	}
 
 #endif	  // USE_USD_SDK
@@ -4887,7 +4888,7 @@ TOptional<UE::Interchange::FImportBlockedImage> UInterchangeUSDTranslator::GetBl
 		}
 
 		// Find another translator that actually supports that filetype to handle the texture
-		UE::Interchange::Private::FScopedTranslator ScopedTranslator(TileFilePath, Results);
+		UE::Interchange::Private::FScopedTranslator ScopedTranslator(TileFilePath, Results, AnalyticsHelper);
 		const IInterchangeTexturePayloadInterface* TextureTranslator = ScopedTranslator.GetPayLoadInterface<IInterchangeTexturePayloadInterface>();
 		if (!ensure(TextureTranslator))
 		{
diff --git a/Engine/Plugins/Interchange/Runtime/Source/Import/Private/Fbx/InterchangeFbxTranslator.cpp b/Engine/Plugins/Interchange/Runtime/Source/Import/Private/Fbx/InterchangeFbxTranslator.cpp
--- a/Engine/Plugins/Interchange/Runtime/Source/Import/Private/Fbx/InterchangeFbxTranslator.cpp
+++ b/Engine/Plugins/Interchange/Runtime/Source/Import/Private/Fbx/InterchangeFbxTranslator.cpp
@@ -319,7 +319,7 @@ TOptional<UE::Interchange::FImportImage> UInterchangeFbxTranslator::GetTexturePa
 	{
 		return TOptional<UE::Interchange::FImportImage>();
 	}
-	UE::Interchange::Private::FScopedTranslator ScopedTranslator(PayLoadKey, Results);
+	UE::Interchange::Private::FScopedTranslator ScopedTranslator(PayLoadKey, Results, AnalyticsHelper);
 	const IInterchangeTexturePayloadInterface* TextureTranslator = ScopedTranslator.GetPayLoadInterface<IInterchangeTexturePayloadInterface>();
 	if (!ensure(TextureTranslator))
 	{
diff --git a/Engine/Plugins/Interchange/Runtime/Source/Import/Private/Gltf/InterchangeGltfTranslator.cpp b/Engine/Plugins/Interchange/Runtime/Source/Import/Private/Gltf/InterchangeGltfTranslator.cpp
--- a/Engine/Plugins/Interchange/Runtime/Source/Import/Private/Gltf/InterchangeGltfTranslator.cpp
+++ b/Engine/Plugins/Interchange/Runtime/Source/Import/Private/Gltf/InterchangeGltfTranslator.cpp
@@ -11,6 +11,7 @@
 #include "GLTFReader.h"
 #include "GLTFTexture.h"
 
+#include "InterchangeAnalyticsHelper.h"
 #include "InterchangeAnimationTrackSetNode.h"
 #include "InterchangeCameraNode.h"
 #include "InterchangeImportLog.h"
@@ -152,68 +153,66 @@ namespace UE::Interchange::Gltf::Private
 		GLTFREADER_FAILED,
 		NOTSUPPORTED_EXTENSION_FOUND
 	};
-	void SendAnalytics(const TranslationResult& TranslationResult,
+	void AddAnalytics(const TranslationResult& TranslationResult,
+		TSharedPtr<UE::Interchange::FAnalyticsHelper> AnalyticsHelper,
 		const GLTF::FAsset& Asset = GLTF::FAsset(),
 		const FText& GLTFReaderLogMessage = FText::GetEmpty())
 	{
-		if (FEngineAnalytics::IsAvailable())
+		TMap<FString, FString> MetadataExtras;
+		for (GLTF::FMetadata::FExtraData ExtraData : Asset.Metadata.Extras)
 		{
-			TMap<FString, FString> MetadataExtras;
-			for (GLTF::FMetadata::FExtraData ExtraData : Asset.Metadata.Extras)
-			{
-				MetadataExtras.Add(ExtraData.Name, ExtraData.Value);
-			}
+			MetadataExtras.Add(ExtraData.Name, ExtraData.Value);
+		}
 
-			TSet<FString> AllExtensions;
-			AllExtensions.Append(Asset.ExtensionsUsed);
-			AllExtensions.Append(Asset.ExtensionsRequired);
+		TSet<FString> AllExtensions;
+		AllExtensions.Append(Asset.ExtensionsUsed);
+		AllExtensions.Append(Asset.ExtensionsRequired);
 
-			TArray<FString> ExtensionsSupported;
-			TArray<FString> ExtensionsUnsupported;
+		TArray<FString> ExtensionsSupported;
+		TArray<FString> ExtensionsUnsupported;
 
-			for (const FString& Extension : AllExtensions)
+		for (const FString& Extension : AllExtensions)
+		{
+			if (ImporterSupportedExtensions.Find(Extension) == INDEX_NONE)
 			{
-				if (ImporterSupportedExtensions.Find(Extension) == INDEX_NONE)
-				{
-					ExtensionsUnsupported.Add(Extension);
-				}
-				else
-				{
-					ExtensionsSupported.Add(Extension);
-				}
+				ExtensionsUnsupported.Add(Extension);
 			}
-
-			TArray<FAnalyticsEventAttribute> GLTFAnalytics;
-			if (Asset.ExtensionsUsed.Num() > 0)				GLTFAnalytics.Add(FAnalyticsEventAttribute(TEXT("ExtensionsUsed"), Asset.ExtensionsUsed));
-			if (Asset.ExtensionsRequired.Num() > 0)			GLTFAnalytics.Add(FAnalyticsEventAttribute(TEXT("ExtensionsRequired"), Asset.ExtensionsRequired));
-			if (ExtensionsSupported.Num() > 0)				GLTFAnalytics.Add(FAnalyticsEventAttribute(TEXT("ExtensionsSupported"), ExtensionsSupported));
-			if (ExtensionsUnsupported.Num() > 0)			GLTFAnalytics.Add(FAnalyticsEventAttribute(TEXT("ExtensionsUnsupported"), ExtensionsUnsupported));
-			if (Asset.Metadata.GeneratorName.Len() > 0)		GLTFAnalytics.Add(FAnalyticsEventAttribute(TEXT("MetaData.GeneratorName"), Asset.Metadata.GeneratorName));
-			if (MetadataExtras.Num() > 0)					GLTFAnalytics.Add(FAnalyticsEventAttribute(TEXT("MetaData.Extras"), MetadataExtras));
-			/*Version is always set at this point.*/		GLTFAnalytics.Add(FAnalyticsEventAttribute(TEXT("MetaData.Version"), Asset.Metadata.Version));
-			if (Asset.HasAbnormalInverseBindMatrices)		GLTFAnalytics.Add(FAnalyticsEventAttribute(TEXT("HasAbnormalInverseBindMatrices"), true));
-
-			switch (TranslationResult)
+			else
 			{
-			case SUCCESSFULL:
-				GLTFAnalytics.Add(FAnalyticsEventAttribute(TEXT("TranslationStatus"), "Successful."));
-				break;
-			case INPUT_FILE_NOTFOUND:
-				GLTFAnalytics.Add(FAnalyticsEventAttribute(TEXT("TranslationStatus"), "[Failed] Input File Not Found."));
-				break;
-			case GLTFREADER_FAILED:
-				GLTFAnalytics.Add(FAnalyticsEventAttribute(TEXT("TranslationStatus"), "[Failed] Parsing error."));
-				break;
-			case NOTSUPPORTED_EXTENSION_FOUND:
-				GLTFAnalytics.Add(FAnalyticsEventAttribute(TEXT("TranslationStatus"), "[Failed] Unsupported Extension Found."));
-				break;
-			default:
-				break;
+				ExtensionsSupported.Add(Extension);
 			}
+		}
 
-			//Send Analytics
-			FEngineAnalytics::GetProvider().RecordEvent(TEXT("Interchange.Usage.Import.GLTF"), GLTFAnalytics);
+		const FString GLTFAnalyticsIdentifier = TEXT("Interchange.Usage.Import.GLTF");
+		TArray<FAnalyticsEventAttribute> GLTFAnalytics;
+		if (Asset.ExtensionsUsed.Num() > 0)				GLTFAnalytics.Add(FAnalyticsEventAttribute(TEXT("ExtensionsUsed"), Asset.ExtensionsUsed));
+		if (Asset.ExtensionsRequired.Num() > 0)			GLTFAnalytics.Add(FAnalyticsEventAttribute(TEXT("ExtensionsRequired"), Asset.ExtensionsRequired));
+		if (ExtensionsSupported.Num() > 0)				GLTFAnalytics.Add(FAnalyticsEventAttribute(TEXT("ExtensionsSupported"), ExtensionsSupported));
+		if (ExtensionsUnsupported.Num() > 0)			GLTFAnalytics.Add(FAnalyticsEventAttribute(TEXT("ExtensionsUnsupported"), ExtensionsUnsupported));
+		if (Asset.Metadata.GeneratorName.Len() > 0)		GLTFAnalytics.Add(FAnalyticsEventAttribute(TEXT("MetaData.GeneratorName"), Asset.Metadata.GeneratorName));
+		if (MetadataExtras.Num() > 0)					GLTFAnalytics.Add(FAnalyticsEventAttribute(TEXT("MetaData.Extras"), MetadataExtras));
+		/*Version is always set at this point.*/		GLTFAnalytics.Add(FAnalyticsEventAttribute(TEXT("MetaData.Version"), Asset.Metadata.Version));
+		if (Asset.HasAbnormalInverseBindMatrices)		GLTFAnalytics.Add(FAnalyticsEventAttribute(TEXT("HasAbnormalInverseBindMatrices"), true));
+
+		switch (TranslationResult)
+		{
+		case SUCCESSFULL:
+			GLTFAnalytics.Add(FAnalyticsEventAttribute(TEXT("TranslationStatus"), "Successful."));
+			break;
+		case INPUT_FILE_NOTFOUND:
+			GLTFAnalytics.Add(FAnalyticsEventAttribute(TEXT("TranslationStatus"), "[Failed] Input File Not Found."));
+			break;
+		case GLTFREADER_FAILED:
+			GLTFAnalytics.Add(FAnalyticsEventAttribute(TEXT("TranslationStatus"), "[Failed] Parsing error."));
+			break;
+		case NOTSUPPORTED_EXTENSION_FOUND:
+			GLTFAnalytics.Add(FAnalyticsEventAttribute(TEXT("TranslationStatus"), "[Failed] Unsupported Extension Found."));
+			break;
+		default:
+			break;
 		}
+
+		AnalyticsHelper->AppendThreadSafe(GLTFAnalyticsIdentifier, GLTFAnalytics);
 	};
 }
 
@@ -501,7 +500,7 @@ bool UInterchangeGLTFTranslator::Translate( UInterchangeBaseNodeContainer& NodeC
 	FString FilePath = GetSourceData()->GetFilename();
 	if ( !FPaths::FileExists( FilePath ) )
 	{
-		SendAnalytics(TranslationResult::INPUT_FILE_NOTFOUND);
+		AddAnalytics(TranslationResult::INPUT_FILE_NOTFOUND, AnalyticsHelper);
 		return false;
 	}
 
@@ -583,7 +582,7 @@ bool UInterchangeGLTFTranslator::Translate( UInterchangeBaseNodeContainer& NodeC
 			ErrorResult->SourceAssetName = FileName;
 			ErrorResult->Text = LogMessage.Value;
 
-			SendAnalytics(TranslationResult::GLTFREADER_FAILED, GltfAsset, LogMessage.Value);
+			AddAnalytics(TranslationResult::GLTFREADER_FAILED, AnalyticsHelper, GltfAsset, LogMessage.Value);
 			return false;
 		}
 	}
@@ -607,7 +606,7 @@ bool UInterchangeGLTFTranslator::Translate( UInterchangeBaseNodeContainer& NodeC
 			LOCTEXT("UnsupportedRequiredExtensions", "Not all required extensions are supported. (Unsupported extensions: {0})"),
 			FText::FromString(NotSupportedRequiredExtensionsStringified));
 
-		SendAnalytics(TranslationResult::NOTSUPPORTED_EXTENSION_FOUND, GltfAsset);
+		AddAnalytics(TranslationResult::NOTSUPPORTED_EXTENSION_FOUND, AnalyticsHelper, GltfAsset);
 		return false;
 	}
 
@@ -982,7 +981,7 @@ bool UInterchangeGLTFTranslator::Translate( UInterchangeBaseNodeContainer& NodeC
 		UE_LOG(LogInterchangeImport, Warning, TEXT("glTF Mesh Import Warning: glTF mesh usage expectations are not met."));
 	}
 
-	SendAnalytics(TranslationResult::SUCCESSFULL, GltfAsset);
+	AddAnalytics(TranslationResult::SUCCESSFULL, AnalyticsHelper, GltfAsset);
 	return true;
 }
 
@@ -1020,7 +1019,7 @@ TOptional< UE::Interchange::FImportImage > UInterchangeGLTFTranslator::GetTextur
 	else
 	{
 		const FString TextureFilePath = FPaths::ConvertRelativePathToFull(GltfTexture.Source.FilePath);
-		UE::Interchange::Private::FScopedTranslator ScopedTranslator(TextureFilePath, Results);
+		UE::Interchange::Private::FScopedTranslator ScopedTranslator(TextureFilePath, Results, AnalyticsHelper);
 		const IInterchangeTexturePayloadInterface* TextureTranslator = ScopedTranslator.GetPayLoadInterface<IInterchangeTexturePayloadInterface>();
 		if (!ensure(TextureTranslator))
 		{
diff --git a/Engine/Plugins/Interchange/Runtime/Source/Import/Private/InterchangeTranslatorHelper.cpp b/Engine/Plugins/Interchange/Runtime/Source/Import/Private/InterchangeTranslatorHelper.cpp
--- a/Engine/Plugins/Interchange/Runtime/Source/Import/Private/InterchangeTranslatorHelper.cpp
+++ b/Engine/Plugins/Interchange/Runtime/Source/Import/Private/InterchangeTranslatorHelper.cpp
@@ -11,7 +11,7 @@
 
 namespace UE::Interchange::Private
 {
-	FScopedTranslator::FScopedTranslator(const FString& PayLoadKey, UInterchangeResultsContainer* Results)
+	FScopedTranslator::FScopedTranslator(const FString& PayLoadKey, UInterchangeResultsContainer* Results, TSharedPtr<UE::Interchange::FAnalyticsHelper> AnalyticsHelper)
 	{
 		PayloadSourceData = UInterchangeManager::GetInterchangeManager().CreateSourceData(PayLoadKey);
 		if (PayloadSourceData)
@@ -21,6 +21,7 @@ namespace UE::Interchange::Private
 			SourceTranslator = UInterchangeManager::GetInterchangeManager().GetTranslatorForSourceData(PayloadSourceData);
 			if (SourceTranslator)
 			{
+				SourceTranslator->AnalyticsHelper = AnalyticsHelper;
 				SourceTranslator->SetInternalFlags(EInternalObjectFlags::Async);
 				if (Results)
 				{
diff --git a/Engine/Plugins/Interchange/Runtime/Source/Import/Private/MaterialX/InterchangeMaterialXTranslator.cpp b/Engine/Plugins/Interchange/Runtime/Source/Import/Private/MaterialX/InterchangeMaterialXTranslator.cpp
--- a/Engine/Plugins/Interchange/Runtime/Source/Import/Private/MaterialX/InterchangeMaterialXTranslator.cpp
+++ b/Engine/Plugins/Interchange/Runtime/Source/Import/Private/MaterialX/InterchangeMaterialXTranslator.cpp
@@ -72,7 +72,7 @@ TOptional<UE::Interchange::FImportImage> UInterchangeMaterialXTranslator::GetTex
 	}
 #endif
 
-	UE::Interchange::Private::FScopedTranslator ScopedTranslator(Filename, Results);
+	UE::Interchange::Private::FScopedTranslator ScopedTranslator(Filename, Results, AnalyticsHelper);
 	const IInterchangeTexturePayloadInterface* TextureTranslator = ScopedTranslator.GetPayLoadInterface<IInterchangeTexturePayloadInterface>();
 
 	if(!ensure(TextureTranslator))
diff --git a/Engine/Plugins/Interchange/Runtime/Source/Import/Private/Mesh/InterchangeOBJTranslator.cpp b/Engine/Plugins/Interchange/Runtime/Source/Import/Private/Mesh/InterchangeOBJTranslator.cpp
--- a/Engine/Plugins/Interchange/Runtime/Source/Import/Private/Mesh/InterchangeOBJTranslator.cpp
+++ b/Engine/Plugins/Interchange/Runtime/Source/Import/Private/Mesh/InterchangeOBJTranslator.cpp
@@ -1608,7 +1608,7 @@ TOptional<UE::Interchange::FImportImage> UInterchangeOBJTranslator::GetTexturePa
 	// Here, we have a texture source referenced by the .obj which is ideally handled by the existing texture translators.
 	// In this case, we can just implement the texture payload interface here and use a temporary texture translator instance
 	// to generate the payload, but what if the translator needed to do non-trivial work during the Translate() phase?
-	UE::Interchange::Private::FScopedTranslator ScopedTranslator(PayLoadKey, Results);
+	UE::Interchange::Private::FScopedTranslator ScopedTranslator(PayLoadKey, Results, AnalyticsHelper);
 	const IInterchangeTexturePayloadInterface* TextureTranslator = ScopedTranslator.GetPayLoadInterface<IInterchangeTexturePayloadInterface>();
 	if (!ensure(TextureTranslator))
 	{
diff --git a/Engine/Plugins/Interchange/Runtime/Source/Import/Public/InterchangeTranslatorHelper.h b/Engine/Plugins/Interchange/Runtime/Source/Import/Public/InterchangeTranslatorHelper.h
--- a/Engine/Plugins/Interchange/Runtime/Source/Import/Public/InterchangeTranslatorHelper.h
+++ b/Engine/Plugins/Interchange/Runtime/Source/Import/Public/InterchangeTranslatorHelper.h
@@ -9,12 +9,17 @@
 
 class UInterchangeResultsContainer;
 
+namespace UE::Interchange
+{
+	struct FAnalyticsHelper;
+}
+
 namespace UE::Interchange::Private
 {
 	struct FScopedTranslator
 	{
 	public:
-		UE_API FScopedTranslator(const FString& PayLoadKey, UInterchangeResultsContainer* Results);
+		UE_API FScopedTranslator(const FString& PayLoadKey, UInterchangeResultsContainer* Results, TSharedPtr<UE::Interchange::FAnalyticsHelper> AnalyticsHelper);
 		UE_API ~FScopedTranslator();
 
 		template <class T>
diff --git a/Engine/Source/Runtime/Interchange/Core/Public/InterchangeTranslatorBase.h b/Engine/Source/Runtime/Interchange/Core/Public/InterchangeTranslatorBase.h
--- a/Engine/Source/Runtime/Interchange/Core/Public/InterchangeTranslatorBase.h
+++ b/Engine/Source/Runtime/Interchange/Core/Public/InterchangeTranslatorBase.h
@@ -20,6 +20,11 @@
 
 class UInterchangeBaseNodeContainer;
 
+namespace UE::Interchange
+{
+	struct FAnalyticsHelper;
+}
+
 /**
  * Base class for translator settings, the UInterchangeAssetImportData will store this data if the source translator exist
  */
@@ -172,4 +177,6 @@ public:
 
 	UPROPERTY()
 	TObjectPtr<const UInterchangeSourceData> SourceData;
+
+	TSharedPtr<UE::Interchange::FAnalyticsHelper> AnalyticsHelper = nullptr;
 };
diff --git a/Engine/Source/Runtime/Interchange/Engine/Private/InterchangeAnalyticsHelper.cpp b/Engine/Source/Runtime/Interchange/Engine/Private/InterchangeAnalyticsHelper.cpp
new file mode 100644
--- /dev/null
+++ b/Engine/Source/Runtime/Interchange/Engine/Private/InterchangeAnalyticsHelper.cpp
@@ -0,0 +1,40 @@
+// Copyright Epic Games, Inc. All Rights Reserved.
+
+#include "InterchangeAnalyticsHelper.h"
+#include "AnalyticsEventAttribute.h"
+#include "EngineAnalytics.h"
+
+namespace UE::Interchange
+{
+	void FAnalyticsHelper::AppendThreadSafe(const FString& Identifier, const TArray<FAnalyticsEventAttribute>& ToAdd)
+	{
+		FScopeLock ScopeLock(&CriticalSection);
+		TArray<FAnalyticsEventAttribute>& Analytics = AnalyticsAttributes.FindOrAdd(Identifier);
+		Analytics.Append(ToAdd);
+	}
+
+	//Make sure it is safe to call this directly.
+	void FAnalyticsHelper::Add(const FString& Identifier, const FAnalyticsEventAttribute& Entry)
+	{
+		TArray<FAnalyticsEventAttribute>& Analytics = AnalyticsAttributes.FindOrAdd(Identifier);
+		Analytics.Add(Entry);
+	}
+
+	void FAnalyticsHelper::AddThreadSafe(const FString& Identifier, const FAnalyticsEventAttribute& Entry)
+	{
+		FScopeLock ScopeLock(&CriticalSection);
+		TArray<FAnalyticsEventAttribute>& Analytics = AnalyticsAttributes.FindOrAdd(Identifier);
+		Analytics.Add(Entry);
+	}
+
+	void FAnalyticsHelper::SendAnalytics()
+	{
+		if (FEngineAnalytics::IsAvailable())
+		{
+			for (const TPair<FString, TArray<FAnalyticsEventAttribute>>& AnalyticsPerIdentifier : AnalyticsAttributes)
+			{
+				FEngineAnalytics::GetProvider().RecordEvent(AnalyticsPerIdentifier.Key, AnalyticsPerIdentifier.Value);
+			}
+		}
+	}
+};
\ No newline at end of file
diff --git a/Engine/Source/Runtime/Interchange/Engine/Private/InterchangeManager.cpp b/Engine/Source/Runtime/Interchange/Engine/Private/InterchangeManager.cpp
--- a/Engine/Source/Runtime/Interchange/Engine/Private/InterchangeManager.cpp
+++ b/Engine/Source/Runtime/Interchange/Engine/Private/InterchangeManager.cpp
@@ -23,6 +23,7 @@
 #include "InterchangeEngineLogPrivate.h"
 #include "InterchangeFactoryBase.h"
 #include "InterchangeHelper.h"
+#include "InterchangeAnalyticsHelper.h"
 #include "InterchangeProjectSettings.h"
 #include "InterchangePythonPipelineBase.h"
 #include "InterchangeSceneImportAsset.h"
@@ -1811,14 +1812,11 @@ UInterchangeManager::ImportInternal(const FString& ContentPath, const UInterchan
 	static int32 GeneratedUniqueID = 0;
 	int32 UniqueId = ++GeneratedUniqueID;
 
-	TArray<FAnalyticsEventAttribute> Attribs;
-	auto PreReturn = [&Attribs]()
+	TSharedPtr<UE::Interchange::FAnalyticsHelper, ESPMode::ThreadSafe> AnalyticsHelper = MakeShared< UE::Interchange::FAnalyticsHelper, ESPMode::ThreadSafe>();
+	const FString UsageImportAnalayticsId = TEXT("Interchange.Usage.Import");
+	auto PreReturn = [&AnalyticsHelper]()
 	{
-		if (FEngineAnalytics::IsAvailable())
-		{
-			FString EventString = TEXT("Interchange.Usage.Import");
-			FEngineAnalytics::GetProvider().RecordEvent(EventString, Attribs);
-		}
+		AnalyticsHelper->SendAnalytics();
 	};
 
 	auto EarlyExit = [&PreReturn]()
@@ -1830,24 +1828,20 @@ UInterchangeManager::ImportInternal(const FString& ContentPath, const UInterchan
 		return ImportResult;
 	};
 
-	if (FEngineAnalytics::IsAvailable())
+	AnalyticsHelper->Add(UsageImportAnalayticsId, FAnalyticsEventAttribute(TEXT("UniqueId"), UniqueId));
+	AnalyticsHelper->Add(UsageImportAnalayticsId, FAnalyticsEventAttribute(TEXT("ImportType"), ImportType == UE::Interchange::EImportType::ImportType_Asset ? TEXT("Asset") : TEXT("Scene")));
+	AnalyticsHelper->Add(UsageImportAnalayticsId, FAnalyticsEventAttribute(TEXT("Parameters.IsAutomated"), ImportAssetParameters.bIsAutomated));
+	AnalyticsHelper->Add(UsageImportAnalayticsId, FAnalyticsEventAttribute(TEXT("Parameters.IsReimport"), ImportAssetParameters.ReimportAsset != nullptr));
+	if (ImportAssetParameters.ReimportAsset != nullptr)
 	{
-		Attribs.Add(FAnalyticsEventAttribute(TEXT("UniqueId"), UniqueId));
-		Attribs.Add(FAnalyticsEventAttribute(TEXT("ImportType"), ImportType == UE::Interchange::EImportType::ImportType_Asset ? TEXT("Asset") : TEXT("Scene")));
-		Attribs.Add(FAnalyticsEventAttribute(TEXT("Parameters.IsAutomated"), ImportAssetParameters.bIsAutomated));
-		const bool bIsReimport = ImportAssetParameters.ReimportAsset != nullptr;
-		Attribs.Add(FAnalyticsEventAttribute(TEXT("Parameters.IsReimport"), bIsReimport));
-		if (ImportAssetParameters.ReimportAsset != nullptr)
-		{
-			Attribs.Add(FAnalyticsEventAttribute(TEXT("Parameters.ReimportSourceIndex"), ImportAssetParameters.ReimportSourceIndex));
-		}
-		const bool bIsPipelineOverride = ImportAssetParameters.OverridePipelines.Num() > 0;
-		Attribs.Add(FAnalyticsEventAttribute(TEXT("Parameters.IsPipelineOverrided"), bIsPipelineOverride));
-		Attribs.Add(FAnalyticsEventAttribute(TEXT("Parameters.bReplaceExisting"), ImportAssetParameters.bReplaceExisting));
-		if(!ImportAssetParameters.DestinationName.IsEmpty())
-		{
-			Attribs.Add(FAnalyticsEventAttribute(TEXT("Parameters.DestinationName"), ImportAssetParameters.DestinationName));
-		}
+		AnalyticsHelper->Add(UsageImportAnalayticsId, FAnalyticsEventAttribute(TEXT("Parameters.ReimportSourceIndex"), ImportAssetParameters.ReimportSourceIndex));
+	}
+	const bool bIsPipelineOverride = ImportAssetParameters.OverridePipelines.Num() > 0;
+	AnalyticsHelper->Add(UsageImportAnalayticsId, FAnalyticsEventAttribute(TEXT("Parameters.IsPipelineOverrided"), bIsPipelineOverride));
+	AnalyticsHelper->Add(UsageImportAnalayticsId, FAnalyticsEventAttribute(TEXT("Parameters.bReplaceExisting"), ImportAssetParameters.bReplaceExisting));
+	if (!ImportAssetParameters.DestinationName.IsEmpty())
+	{
+		AnalyticsHelper->Add(UsageImportAnalayticsId, FAnalyticsEventAttribute(TEXT("Parameters.DestinationName"), ImportAssetParameters.DestinationName));
 	}
 
 	if (this != &GetInterchangeManager())
@@ -1872,10 +1866,7 @@ UInterchangeManager::ImportInternal(const FString& ContentPath, const UInterchan
 		Translator->ReleaseSource();
 	}
 
-	if (FEngineAnalytics::IsAvailable())
-	{
-		Attribs.Add(FAnalyticsEventAttribute(TEXT("SourceExtension"), FPaths::GetExtension(SourceData->GetFilename())));
-	}
+	AnalyticsHelper->Add(UsageImportAnalayticsId, FAnalyticsEventAttribute(TEXT("SourceExtension"), FPaths::GetExtension(SourceData->GetFilename())));
 
 	bool bImportScene = ImportType == UE::Interchange::EImportType::ImportType_Scene;
 	const FInterchangeImportSettings& InterchangeImportSettings = FInterchangeProjectSettingsUtils::GetDefaultImportSettings(bImportScene);
@@ -1971,10 +1962,7 @@ UInterchangeManager::ImportInternal(const FString& ContentPath, const UInterchan
 							}
 
 							//Add an analytic so we can know if a not re-import was turn into a re-import
-							if (FEngineAnalytics::IsAvailable())
-							{
-								Attribs.Add(FAnalyticsEventAttribute(TEXT("Parameters.ConvertToReimport"), true));
-							}
+							AnalyticsHelper->Add(UsageImportAnalayticsId, FAnalyticsEventAttribute(TEXT("Parameters.ConvertToReimport"), true));
 						}
 					}
 				}
@@ -2070,6 +2058,9 @@ UInterchangeManager::ImportInternal(const FString& ContentPath, const UInterchan
 		{
 			return EarlyExit();
 		}
+
+		SourceTranslator->AnalyticsHelper = AnalyticsHelper;
+
 		if (bIsReimport)
 		{
 			//Set translator settings if we are doing a reimport
@@ -2172,10 +2163,8 @@ UInterchangeManager::ImportInternal(const FString& ContentPath, const UInterchan
 #else
 		bool bShowPipelineStacksConfigurationDialog = false;
 #endif
-		if (FEngineAnalytics::IsAvailable())
-		{
-			Attribs.Add(FAnalyticsEventAttribute(TEXT("ShowImportDialog"), bShowPipelineStacksConfigurationDialog));
-		}
+		
+		AnalyticsHelper->AddThreadSafe(UsageImportAnalayticsId, FAnalyticsEventAttribute(TEXT("ShowImportDialog"), bShowPipelineStacksConfigurationDialog));
 
 		//If we need to display the dialog we want to translate the source file before showing it
 		if (RegisteredPipelineConfiguration && bShowPipelineStacksConfigurationDialog)
@@ -2442,10 +2431,7 @@ UInterchangeManager::ImportInternal(const FString& ContentPath, const UInterchan
 	//Cancel the import do not queue task
 	if (bImportCanceled || bImportAborted)
 	{
-		if (FEngineAnalytics::IsAvailable())
-		{
-			Attribs.Add(FAnalyticsEventAttribute(TEXT("Canceled"), bImportCanceled));
-		}
+		AnalyticsHelper->AddThreadSafe(UsageImportAnalayticsId, FAnalyticsEventAttribute(TEXT("Canceled"), bImportCanceled));
 
 		AsyncHelper->InitCancel();
 		AsyncHelper->CleanUp();
diff --git a/Engine/Source/Runtime/Interchange/Engine/Public/InterchangeAnalyticsHelper.h b/Engine/Source/Runtime/Interchange/Engine/Public/InterchangeAnalyticsHelper.h
new file mode 100644
--- /dev/null
+++ b/Engine/Source/Runtime/Interchange/Engine/Public/InterchangeAnalyticsHelper.h
@@ -0,0 +1,29 @@
+// Copyright Epic Games, Inc. All Rights Reserved.
+
+#pragma once
+
+#include "CoreMinimal.h"
+struct FAnalyticsEventAttribute;
+
+namespace UE::Interchange
+{
+	struct FAnalyticsHelper
+	{
+	private:
+		//Analytics Attributes to share with the Translators. (to be sent when Import is finished)
+		TMap<FString, TArray<FAnalyticsEventAttribute>> AnalyticsAttributes;
+		FCriticalSection CriticalSection;
+
+	public:
+		INTERCHANGEENGINE_API void AppendThreadSafe(const FString& Identifier, const TArray<FAnalyticsEventAttribute>& ToAdd);
+
+		//Make sure it is safe to call this directly.
+		INTERCHANGEENGINE_API void Add(const FString& Identifier, const FAnalyticsEventAttribute& Entry);
+
+		//It is suggest to use the AppendThreadSafe, if more than one Entry is added as it will Scope Lock for every entry.
+		// (where append will Scope Lock only per append)
+		INTERCHANGEENGINE_API void AddThreadSafe(const FString& Identifier, const FAnalyticsEventAttribute& Entry);
+
+		INTERCHANGEENGINE_API void SendAnalytics();
+	};
+};
\ No newline at end of file
diff --git a/Engine/Source/Runtime/Interchange/Engine/Public/InterchangeManager.h b/Engine/Source/Runtime/Interchange/Engine/Public/InterchangeManager.h
--- a/Engine/Source/Runtime/Interchange/Engine/Public/InterchangeManager.h
+++ b/Engine/Source/Runtime/Interchange/Engine/Public/InterchangeManager.h
@@ -38,6 +38,11 @@ class UInterchangePipelineBase;
 class UInterchangePythonPipelineBase;
 class ULevel;
 
+namespace UE::Interchange
+{
+	struct FAnalyticsHelper;
+}
+
 /** Some utility delegates for automating Interchange. */
 DECLARE_DYNAMIC_DELEGATE_OneParam(FOnObjectImportDoneDynamic, UObject*, Object);
 DECLARE_DELEGATE_OneParam(FOnObjectImportDoneNative, UObject*);
