Interchange Reset - Mesh Actors do not reset the asset if set to none
- Actor factory node tracks the resolved asset factory node uid to use later for retrieving and setting the correct asset.

#rnx

diff --git a/Engine/Plugins/Interchange/Runtime/Source/FactoryNodes/Private/InterchangeMeshActorFactoryNode.cpp b/Engine/Plugins/Interchange/Runtime/Source/FactoryNodes/Private/InterchangeMeshActorFactoryNode.cpp
--- a/Engine/Plugins/Interchange/Runtime/Source/FactoryNodes/Private/InterchangeMeshActorFactoryNode.cpp
+++ b/Engine/Plugins/Interchange/Runtime/Source/FactoryNodes/Private/InterchangeMeshActorFactoryNode.cpp
@@ -13,6 +13,16 @@ UInterchangeMeshActorFactoryNode::UInterchangeMeshActorFactoryNode()
 	SlotMaterialDependencies.Initialize(Attributes.ToSharedRef(), TEXT("__SlotMaterialDependencies__"));
 }
 
+bool UInterchangeMeshActorFactoryNode::SetCustomInstancedAssetFactoryNodeUid(const FString& AttributeValue)
+{
+	IMPLEMENT_NODE_ATTRIBUTE_SETTER_NODELEGATE(InstancedAssetFactoryNodeUid, FString);
+}
+
+bool UInterchangeMeshActorFactoryNode::GetCustomInstancedAssetFactoryNodeUid(FString& AttributeValue) const
+{
+	IMPLEMENT_NODE_ATTRIBUTE_GETTER(InstancedAssetFactoryNodeUid, FString);
+}
+
 void UInterchangeMeshActorFactoryNode::GetSlotMaterialDependencies(TMap<FString, FString>& OutMaterialDependencies) const
 {
 	OutMaterialDependencies = SlotMaterialDependencies.ToMap();
diff --git a/Engine/Plugins/Interchange/Runtime/Source/FactoryNodes/Public/InterchangeMeshActorFactoryNode.h b/Engine/Plugins/Interchange/Runtime/Source/FactoryNodes/Public/InterchangeMeshActorFactoryNode.h
--- a/Engine/Plugins/Interchange/Runtime/Source/FactoryNodes/Public/InterchangeMeshActorFactoryNode.h
+++ b/Engine/Plugins/Interchange/Runtime/Source/FactoryNodes/Public/InterchangeMeshActorFactoryNode.h
@@ -29,6 +29,14 @@ public:
 		}
 	}
 
+	/** Set the resolved asset factory node. */
+	UFUNCTION(BlueprintCallable, Category = "Interchange | Node | Mesh")
+	UE_API bool SetCustomInstancedAssetFactoryNodeUid(const FString& AttributeValue);
+
+	/** Get the resolved asset factory node. */
+	UFUNCTION(BlueprintCallable, Category = "Interchange | Node | Mesh")
+	UE_API bool GetCustomInstancedAssetFactoryNodeUid(FString& AttributeValue) const;
+
 	/**
 	 * Retrieve the correspondence table between slot names and assigned materials for this object.
 	 */
@@ -78,6 +86,9 @@ public:
 private:
 	UE::Interchange::TMapAttributeHelper<FString, FString> SlotMaterialDependencies;
 
+	//A factory node can reference an asset on top of base asset:
+	const UE::Interchange::FAttributeKey Macro_CustomInstancedAssetFactoryNodeUidKey = UE::Interchange::FAttributeKey(TEXT("InstancedAssetFactoryNodeUid"));
+
 	//A scene node can reference an animation asset on top of base asset:
 	const UE::Interchange::FAttributeKey Macro_CustomAnimationAssetUidToPlayKey = UE::Interchange::FAttributeKey(TEXT("AnimationAssetUidToPlay"));
 
diff --git a/Engine/Plugins/Interchange/Runtime/Source/Import/Private/Scene/InterchangeActorHelper.cpp b/Engine/Plugins/Interchange/Runtime/Source/Import/Private/Scene/InterchangeActorHelper.cpp
--- a/Engine/Plugins/Interchange/Runtime/Source/Import/Private/Scene/InterchangeActorHelper.cpp
+++ b/Engine/Plugins/Interchange/Runtime/Source/Import/Private/Scene/InterchangeActorHelper.cpp
@@ -4,8 +4,11 @@
 
 #include "Engine/Engine.h"
 #include "Engine/Level.h"
+#include "GameFramework/WorldSettings.h"
 #include "InterchangeActorFactoryNode.h"
+#include "InterchangeAssetUserData.h"
 #include "InterchangeImportCommon.h"
+#include "InterchangeImportReset.h"
 #include "InterchangeMaterialFactoryNode.h"
 #include "InterchangeMeshActorFactoryNode.h"
 #include "InterchangeMeshNode.h"
@@ -23,201 +26,237 @@
 #include "Editor/EditorEngine.h"
 #endif
 
-AActor* UE::Interchange::ActorHelper::GetSpawnedParentActor(const UInterchangeBaseNodeContainer* NodeContainer, const UInterchangeActorFactoryNode* FactoryNode)
+namespace UE::Interchange::ActorHelper
 {
-	AActor* ParentActor = nullptr;
-
-	if (const UInterchangeFactoryBaseNode* ParentFactoryNode = NodeContainer->GetFactoryNode(FactoryNode->GetParentUid()))
+	AActor* GetSpawnedParentActor(const UInterchangeBaseNodeContainer* NodeContainer, const UInterchangeActorFactoryNode* FactoryNode)
 	{
-		FSoftObjectPath ReferenceObject;
-		ParentFactoryNode->GetCustomReferenceObject(ReferenceObject);
-		ParentActor = Cast<AActor>(ReferenceObject.TryLoad());
-	}
-
-	return  ParentActor;
-}
+		AActor* ParentActor = nullptr;
 
-AActor* UE::Interchange::ActorHelper::SpawnFactoryActor(const UInterchangeFactoryBase::FImportSceneObjectsParams& CreateSceneObjectsParams)
-{
-	UInterchangeActorFactoryNode* FactoryNode = Cast<UInterchangeActorFactoryNode>(CreateSceneObjectsParams.FactoryNode);
-	const UInterchangeBaseNodeContainer* NodeContainer = CreateSceneObjectsParams.NodeContainer;
+		if (const UInterchangeFactoryBaseNode* ParentFactoryNode = NodeContainer->GetFactoryNode(FactoryNode->GetParentUid()))
+		{
+			FSoftObjectPath ReferenceObject;
+			ParentFactoryNode->GetCustomReferenceObject(ReferenceObject);
+			ParentActor = Cast<AActor>(ReferenceObject.TryLoad());
+		}
 
-	if (!FactoryNode || !NodeContainer)
-	{
-		return nullptr;
+		return  ParentActor;
 	}
 
-	FActorSpawnParameters SpawnParameters;
-	SpawnParameters.Name = FName(CreateSceneObjectsParams.ObjectName);
-	SpawnParameters.NameMode = FActorSpawnParameters::ESpawnActorNameMode::Requested;
-	SpawnParameters.OverrideLevel = CreateSceneObjectsParams.Level;
-	SpawnParameters.SpawnCollisionHandlingOverride = ESpawnActorCollisionHandlingMethod::AlwaysSpawn;
-
-	AActor* ParentActor = UE::Interchange::ActorHelper::GetSpawnedParentActor(NodeContainer, FactoryNode);
-	UWorld* const World = [&SpawnParameters, &ParentActor]()
+	AActor* SpawnFactoryActor(const UInterchangeFactoryBase::FImportSceneObjectsParams& CreateSceneObjectsParams)
 	{
-		UWorld* DefaultWorld = nullptr;
+		UInterchangeActorFactoryNode* FactoryNode = Cast<UInterchangeActorFactoryNode>(CreateSceneObjectsParams.FactoryNode);
+		const UInterchangeBaseNodeContainer* NodeContainer = CreateSceneObjectsParams.NodeContainer;
 
-		if (ParentActor)
+		if (!FactoryNode || !NodeContainer)
 		{
-			DefaultWorld = ParentActor->GetWorld();
+			return nullptr;
 		}
-#if WITH_EDITOR
-		if (DefaultWorld == nullptr)
+
+		FActorSpawnParameters SpawnParameters;
+		SpawnParameters.Name = FName(CreateSceneObjectsParams.ObjectName);
+		SpawnParameters.NameMode = FActorSpawnParameters::ESpawnActorNameMode::Requested;
+		SpawnParameters.OverrideLevel = CreateSceneObjectsParams.Level;
+		SpawnParameters.SpawnCollisionHandlingOverride = ESpawnActorCollisionHandlingMethod::AlwaysSpawn;
+
+		AActor* ParentActor = GetSpawnedParentActor(NodeContainer, FactoryNode);
+		UWorld* const World = [&SpawnParameters, &ParentActor]()
 		{
-			UEditorEngine* EditorEngine = Cast<UEditorEngine>(GEngine);
-			if (GIsEditor && EditorEngine != nullptr)
+			UWorld* DefaultWorld = nullptr;
+
+			if (ParentActor)
 			{
-				DefaultWorld = EditorEngine->GetEditorWorldContext().World();
+				DefaultWorld = ParentActor->GetWorld();
+			}
+#if WITH_EDITOR
+			if (DefaultWorld == nullptr)
+			{
+				UEditorEngine* EditorEngine = Cast<UEditorEngine>(GEngine);
+				if (GIsEditor && EditorEngine != nullptr)
+				{
+					DefaultWorld = EditorEngine->GetEditorWorldContext().World();
+				}
 			}
-		}
 #endif
-		if (DefaultWorld == nullptr && GEngine)
-		{
-			DefaultWorld = GEngine->GetWorld();
-		}
+			if (DefaultWorld == nullptr && GEngine)
+			{
+				DefaultWorld = GEngine->GetWorld();
+			}
 
-		return SpawnParameters.OverrideLevel ? SpawnParameters.OverrideLevel->GetWorld() : DefaultWorld;
-	}();
+			return SpawnParameters.OverrideLevel ? SpawnParameters.OverrideLevel->GetWorld() : DefaultWorld;
+		}();
 
-	if (!World)
-	{
-		return nullptr;
-	}
+		if (!World)
+		{
+			return nullptr;
+		}
 
-	UClass* ActorClass = FactoryNode->GetObjectClass();
-	AActor* SpawnedActor = Cast<AActor>(CreateSceneObjectsParams.ReimportObject);
-	if (SpawnedActor)
-	{
-		if (SpawnedActor->GetClass() == ActorClass)
+		UClass* ActorClass = FactoryNode->GetObjectClass();
+		AActor* SpawnedActor = Cast<AActor>(CreateSceneObjectsParams.ReimportObject);
+		if (SpawnedActor)
 		{
-			// TODO: Check whether parenting is the same
+			if (SpawnedActor->GetClass() == ActorClass)
+			{
+				// TODO: Check whether parenting is the same
+			}
+			else
+			{
+				SpawnedActor = nullptr;
+			}
 		}
-		else
+		// The related actor has been deleted. Check on reimport policy
+		else if (CreateSceneObjectsParams.ReimportFactoryNode && !FactoryNode->ShouldForceNodeReimport())
 		{
-			SpawnedActor = nullptr;
+			// if reimport policy does not prioritize new content, do not recreate the actor
+			return nullptr;
 		}
-	}
-	// The related actor has been deleted. Check on reimport policy
-	else if (CreateSceneObjectsParams.ReimportFactoryNode && !FactoryNode->ShouldForceNodeReimport())
-	{
-		// if reimport policy does not prioritize new content, do not recreate the actor
-		return nullptr;
-	}
 
-	if (!SpawnedActor)
-	{
-		SpawnedActor = World->SpawnActor<AActor>(ActorClass, SpawnParameters);
-	}
+		if (!SpawnedActor)
+		{
+			SpawnedActor = World->SpawnActor<AActor>(ActorClass, SpawnParameters);
+		}
 
-	if (SpawnedActor)
-	{
+		if (SpawnedActor)
+		{
 #if WITH_EDITOR
-		SpawnedActor->SetActorLabel(FactoryNode->GetDisplayLabel());
+			SpawnedActor->SetActorLabel(FactoryNode->GetDisplayLabel());
 #endif
-		if (!SpawnedActor->GetRootComponent())
-		{
-			USceneComponent* RootComponent = NewObject<USceneComponent>(SpawnedActor, USceneComponent::GetDefaultSceneRootVariableName(), RF_Transactional);
+			if (!SpawnedActor->GetRootComponent())
+			{
+				USceneComponent* RootComponent = NewObject<USceneComponent>(SpawnedActor, USceneComponent::GetDefaultSceneRootVariableName(), RF_Transactional);
 #if WITH_EDITORONLY_DATA
-			RootComponent->bVisualizeComponent = true;
+				RootComponent->bVisualizeComponent = true;
 #endif
-			SpawnedActor->SetRootComponent(RootComponent);
-			SpawnedActor->AddInstanceComponent(RootComponent);
-		}
+				SpawnedActor->SetRootComponent(RootComponent);
+				SpawnedActor->AddInstanceComponent(RootComponent);
+			}
 
-		if (USceneComponent* RootComponent = SpawnedActor->GetRootComponent())
-		{
-			uint8 Mobility;
-			if (FactoryNode->GetCustomMobility(Mobility))
+			if (USceneComponent* RootComponent = SpawnedActor->GetRootComponent())
 			{
-				//Make sure we don't have a mobility that's more restrictive than our parent mobility, as that wouldn't be a valid setup.
-				EComponentMobility::Type TargetMobility = (EComponentMobility::Type)Mobility;
-
-				if (ParentActor && ParentActor->GetRootComponent())
+				uint8 Mobility;
+				if (FactoryNode->GetCustomMobility(Mobility))
 				{
-					TargetMobility = (EComponentMobility::Type)FMath::Max((uint8)Mobility, (uint8)ParentActor->GetRootComponent()->Mobility);
+					//Make sure we don't have a mobility that's more restrictive than our parent mobility, as that wouldn't be a valid setup.
+					EComponentMobility::Type TargetMobility = (EComponentMobility::Type)Mobility;
+
+					if (ParentActor && ParentActor->GetRootComponent())
+					{
+						TargetMobility = (EComponentMobility::Type)FMath::Max((uint8)Mobility, (uint8)ParentActor->GetRootComponent()->Mobility);
+					}
+
+					RootComponent->SetMobility(TargetMobility);
 				}
+			}
 
-				RootComponent->SetMobility(TargetMobility);
+			if (ParentActor || !SpawnedActor->IsAttachedTo(ParentActor))
+			{
+				SpawnedActor->AttachToActor(ParentActor, FAttachmentTransformRules::KeepWorldTransform);
 			}
 		}
 
-		if (ParentActor || !SpawnedActor->IsAttachedTo(ParentActor))
+		return SpawnedActor;
+	}
+
+	const UInterchangeFactoryBaseNode* FindAssetInstanceFactoryNode(const UInterchangeBaseNodeContainer* NodeContainer, const UInterchangeFactoryBaseNode* ActorFactoryNode)
+	{
+		TArray<FString> ActorTargetNodes;
+		ActorFactoryNode->GetTargetNodeUids(ActorTargetNodes);
+		const UInterchangeSceneNode* SceneNode = ActorTargetNodes.IsEmpty() ? nullptr : Cast<UInterchangeSceneNode>(NodeContainer->GetNode(ActorTargetNodes[0]));
+		if (!SceneNode)
 		{
-			SpawnedActor->AttachToActor(ParentActor, FAttachmentTransformRules::KeepWorldTransform);
+			return nullptr;
 		}
-	}
 
-	return SpawnedActor;
-}
+		FString AssetInstanceUid;
+		SceneNode->GetCustomAssetInstanceUid(AssetInstanceUid);
+		const UInterchangeBaseNode* AssetNode = NodeContainer->GetNode(AssetInstanceUid);
+		if (!AssetNode)
+		{
+			return nullptr;
+		}
 
-const UInterchangeFactoryBaseNode* UE::Interchange::ActorHelper::FindAssetInstanceFactoryNode(const UInterchangeBaseNodeContainer* NodeContainer, const UInterchangeFactoryBaseNode* ActorFactoryNode)
-{
-	TArray<FString> ActorTargetNodes;
-	ActorFactoryNode->GetTargetNodeUids(ActorTargetNodes);
-	const UInterchangeSceneNode* SceneNode = ActorTargetNodes.IsEmpty() ? nullptr : Cast<UInterchangeSceneNode>(NodeContainer->GetNode(ActorTargetNodes[0]));
-	if (!SceneNode)
-	{
-		return nullptr;
+		TArray<FString> AssetTargetNodeIds;
+		AssetNode->GetTargetNodeUids(AssetTargetNodeIds);
+		return AssetTargetNodeIds.IsEmpty() ? nullptr : NodeContainer->GetFactoryNode(AssetTargetNodeIds[0]);
 	}
 
-	FString AssetInstanceUid;
-	SceneNode->GetCustomAssetInstanceUid(AssetInstanceUid);
-	const UInterchangeBaseNode* AssetNode = NodeContainer->GetNode(AssetInstanceUid);
-	if (!AssetNode)
+	void ApplySlotMaterialDependencies(const UInterchangeBaseNodeContainer& NodeContainer, const UInterchangeMeshActorFactoryNode& MeshActorFactoryNode, UMeshComponent& MeshComponent)
 	{
-		return nullptr;
-	}
+		// Set material slots from imported materials
+		TMap<FString, FString> SlotMaterialDependencies;
+		MeshActorFactoryNode.GetSlotMaterialDependencies(SlotMaterialDependencies);
+		for (TPair<FString, FString>& SlotMaterialDependency : SlotMaterialDependencies)
+		{
+			const UInterchangeBaseMaterialFactoryNode* MaterialFactoryNode = Cast<UInterchangeBaseMaterialFactoryNode>(NodeContainer.GetNode(SlotMaterialDependency.Value));
+			if (!MaterialFactoryNode || !MaterialFactoryNode->IsEnabled())
+			{
+				continue;
+			}
+			FSoftObjectPath ReferenceObject;
+			MaterialFactoryNode->GetCustomReferenceObject(ReferenceObject);
+			if (!ReferenceObject.IsValid())
+			{
+				continue;
+			}
 
-	TArray<FString> AssetTargetNodeIds;
-	AssetNode->GetTargetNodeUids(AssetTargetNodeIds);
-	return AssetTargetNodeIds.IsEmpty() ? nullptr : NodeContainer->GetFactoryNode(AssetTargetNodeIds[0]);
-}
+			FName MaterialSlotName = *SlotMaterialDependency.Key;
+			UMaterialInterface* MaterialInterface = Cast<UMaterialInterface>(ReferenceObject.TryLoad());
+			if (!MaterialInterface)
+			{
+				MaterialInterface = UMaterial::GetDefaultMaterial(MD_Surface);
+			}
 
-void UE::Interchange::ActorHelper::ApplySlotMaterialDependencies(const UInterchangeBaseNodeContainer& NodeContainer, const UInterchangeMeshActorFactoryNode& MeshActorFactoryNode, UMeshComponent& MeshComponent)
-{
-	// Set material slots from imported materials
-	TMap<FString, FString> SlotMaterialDependencies;
-	MeshActorFactoryNode.GetSlotMaterialDependencies(SlotMaterialDependencies);
-	for (TPair<FString, FString>& SlotMaterialDependency : SlotMaterialDependencies)
+			int32 MaterialSlotIndex = MeshComponent.GetMaterialIndex(MaterialSlotName);
+			if (MaterialSlotIndex != INDEX_NONE)
+			{
+				MeshComponent.SetMaterial(MaterialSlotIndex, MaterialInterface);
+			}
+		}
+	}
+
+	void ApplyAllCustomAttributes(const UInterchangeFactoryBase::FImportSceneObjectsParams& CreateSceneObjectsParams, UObject& ObjectToUpdate)
 	{
-		const UInterchangeBaseMaterialFactoryNode* MaterialFactoryNode = Cast<UInterchangeBaseMaterialFactoryNode>(NodeContainer.GetNode(SlotMaterialDependency.Value));
-		if (!MaterialFactoryNode || !MaterialFactoryNode->IsEnabled())
+		if (CreateSceneObjectsParams.ReimportObject && (ObjectToUpdate.GetOuter() == CreateSceneObjectsParams.ReimportObject))
 		{
-			continue;
+				UInterchangeFactoryBaseNode* CurrentNode = UInterchangeFactoryBaseNode::DuplicateWithObject(CreateSceneObjectsParams.FactoryNode, &ObjectToUpdate);
+
+				FFactoryCommon::ApplyReimportStrategyToAsset(&ObjectToUpdate, CreateSceneObjectsParams.ReimportFactoryNode, CurrentNode, CreateSceneObjectsParams.FactoryNode);
 		}
-		FSoftObjectPath ReferenceObject;
-		MaterialFactoryNode->GetCustomReferenceObject(ReferenceObject);
-		if (!ReferenceObject.IsValid())
+		else
 		{
-			continue;
+			CreateSceneObjectsParams.FactoryNode->ApplyAllCustomAttributeToObject(&ObjectToUpdate);
 		}
+	}
 
-		FName MaterialSlotName = *SlotMaterialDependency.Key;
-		UMaterialInterface* MaterialInterface = Cast<UMaterialInterface>(ReferenceObject.TryLoad());
-		if (!MaterialInterface)
+	void AddInterchangeAssetUserDataToActor(AActor* Actor, const UInterchangeSceneImportAsset* SceneImportAsset, const UInterchangeFactoryBaseNode* FactoryNode)
+	{
+		if (!Actor || !SceneImportAsset || !FactoryNode)
 		{
-			MaterialInterface = UMaterial::GetDefaultMaterial(MD_Surface);
+			return;
 		}
 
-		int32 MaterialSlotIndex = MeshComponent.GetMaterialIndex(MaterialSlotName);
-		if (MaterialSlotIndex != INDEX_NONE)
+		const FSoftObjectPath SceneImportAssetPath((const UObject*)SceneImportAsset);
+		if (USceneComponent* RootComponent = Actor->GetRootComponent())
 		{
-			MeshComponent.SetMaterial(MaterialSlotIndex, MaterialInterface);
+			UInterchangeAssetUserData* AssetUserData = NewObject<UInterchangeAssetUserData>(Actor);
+			AssetUserData->MetaData.Add(UE::Interchange::InterchangeReset::Constants::SceneImportAssetPathKey, SceneImportAssetPath.ToString());
+			AssetUserData->MetaData.Add(UE::Interchange::InterchangeReset::Constants::FactoryNodeUidPathKey, FactoryNode->GetUniqueID());
+			RootComponent->AddAssetUserData(AssetUserData);
 		}
 	}
-}
-
-void UE::Interchange::ActorHelper::ApplyAllCustomAttributes(const UInterchangeFactoryBase::FImportSceneObjectsParams& CreateSceneObjectsParams, UObject& ObjectToUpdate)
-{
-	if (CreateSceneObjectsParams.ReimportObject && (ObjectToUpdate.GetOuter() == CreateSceneObjectsParams.ReimportObject))
+	
+	void AddInterchangeLevelAssetUserDataToWorld(UWorld* World, const UInterchangeSceneImportAsset* SceneImportAsset)
 	{
-		UInterchangeFactoryBaseNode* CurrentNode = UInterchangeFactoryBaseNode::DuplicateWithObject(CreateSceneObjectsParams.FactoryNode, &ObjectToUpdate);
+		if (!World || !SceneImportAsset)
+		{
+			return;
+		}
 
-		FFactoryCommon::ApplyReimportStrategyToAsset(&ObjectToUpdate, CreateSceneObjectsParams.ReimportFactoryNode, CurrentNode, CreateSceneObjectsParams.FactoryNode);
-	}
-	else
-	{
-		CreateSceneObjectsParams.FactoryNode->ApplyAllCustomAttributeToObject(&ObjectToUpdate);
+		const FSoftObjectPath SceneImportAssetPath((const UObject*)SceneImportAsset);
+		if (AWorldSettings* WorldSettings = World->GetWorldSettings())
+		{
+			UInterchangeLevelAssetUserData* AssetUserData = NewObject<UInterchangeLevelAssetUserData>(World);
+			AssetUserData->SceneImportPaths.Add(SceneImportAssetPath);
+			WorldSettings->AddAssetUserData(AssetUserData);
+		}
 	}
-}
+}
\ No newline at end of file
diff --git a/Engine/Plugins/Interchange/Runtime/Source/Import/Private/Scene/InterchangeSkeletalMeshActorFactory.cpp b/Engine/Plugins/Interchange/Runtime/Source/Import/Private/Scene/InterchangeSkeletalMeshActorFactory.cpp
--- a/Engine/Plugins/Interchange/Runtime/Source/Import/Private/Scene/InterchangeSkeletalMeshActorFactory.cpp
+++ b/Engine/Plugins/Interchange/Runtime/Source/Import/Private/Scene/InterchangeSkeletalMeshActorFactory.cpp
@@ -46,7 +46,45 @@ UObject* UInterchangeSkeletalMeshActorFactory::ProcessActor(AActor& SpawnedActor
 	}
 
 	return nullptr;
-};
+}
+
+void UInterchangeSkeletalMeshActorFactory::ExecuteResetObjectProperties(const UInterchangeBaseNodeContainer* BaseNodeContainer, UInterchangeFactoryBaseNode* FactoryNode, UObject* ImportedObject)
+{
+	Super::ExecuteResetObjectProperties(BaseNodeContainer, FactoryNode, ImportedObject);
+	if (ASkeletalMeshActor* SkeletalMeshActor = Cast<ASkeletalMeshActor>(ImportedObject))
+	{
+		if (UInterchangeMeshActorFactoryNode* MeshActorFactoryNode = Cast<UInterchangeMeshActorFactoryNode>(FactoryNode))
+		{
+			if (USkeletalMeshComponent* SkeletalMeshComponent = Cast<USkeletalMeshComponent>(SkeletalMeshActor->GetSkeletalMeshComponent()))
+			{
+				using namespace UE::Interchange;
+				FString InstancedAssetFactoryNodeUid;
+				if (MeshActorFactoryNode->GetCustomInstancedAssetFactoryNodeUid(InstancedAssetFactoryNodeUid))
+				{
+					if (const UInterchangeFactoryBaseNode* MeshNode = Cast<UInterchangeFactoryBaseNode>(BaseNodeContainer->GetNode(InstancedAssetFactoryNodeUid)))
+					{
+						FSoftObjectPath ReferenceObject;
+						MeshNode->GetCustomReferenceObject(ReferenceObject);
+						if (USkeletalMesh* SkeletalMesh = Cast<USkeletalMesh>(ReferenceObject.TryLoad()))
+						{
+							if (SkeletalMesh != SkeletalMeshComponent->GetSkeletalMeshAsset())
+							{
+								SkeletalMeshComponent->SetSkeletalMeshAsset(SkeletalMesh);
+							}
+						}
+					}
+					else
+					{
+						// TODO: Warn that new mesh has not been applied
+					}
+				}
+
+				SkeletalMeshComponent->EmptyOverrideMaterials();
+				ActorHelper::ApplySlotMaterialDependencies(*BaseNodeContainer, *MeshActorFactoryNode, *SkeletalMeshComponent);
+			}
+		}
+	}
+}
 
 UClass* UInterchangeSkeletalMeshActorFactory::GetFactoryClass() const
 {
@@ -68,7 +106,7 @@ void UInterchangeSkeletalMeshActorFactory::SetupObject_GameThread(const FSetupOb
 				SkeletalMeshFactoryNode->GetCustomReferenceObject(ReferenceObject);
 				if (USkeletalMesh* SkeletalMesh = Cast<USkeletalMesh>(ReferenceObject.TryLoad()))
 				{
-					SkeletalMeshComponent->SetSkeletalMesh(SkeletalMesh);
+					SkeletalMeshComponent->SetSkeletalMeshAsset(SkeletalMesh);
 
 					if (const UInterchangeMeshActorFactoryNode* MeshActorFactoryNode = Cast<UInterchangeMeshActorFactoryNode>(Arguments.FactoryNode))
 					{
diff --git a/Engine/Plugins/Interchange/Runtime/Source/Import/Private/Scene/InterchangeStaticMeshActorFactory.cpp b/Engine/Plugins/Interchange/Runtime/Source/Import/Private/Scene/InterchangeStaticMeshActorFactory.cpp
--- a/Engine/Plugins/Interchange/Runtime/Source/Import/Private/Scene/InterchangeStaticMeshActorFactory.cpp
+++ b/Engine/Plugins/Interchange/Runtime/Source/Import/Private/Scene/InterchangeStaticMeshActorFactory.cpp
@@ -160,16 +160,36 @@ void UInterchangeStaticMeshActorFactory::SetupObject_GameThread(const FSetupObje
 void UInterchangeStaticMeshActorFactory::ExecuteResetObjectProperties(const UInterchangeBaseNodeContainer* BaseNodeContainer, UInterchangeFactoryBaseNode* FactoryNode, UObject* ImportedObject)
 {
 	Super::ExecuteResetObjectProperties(BaseNodeContainer, FactoryNode, ImportedObject);
-	if (AActor* SpawnedActor = Cast<AActor>(ImportedObject))
+	if (AStaticMeshActor* StaticMeshActor = Cast<AStaticMeshActor>(ImportedObject))
 	{
 		if (UInterchangeMeshActorFactoryNode* MeshActorFactoryNode = Cast<UInterchangeMeshActorFactoryNode>(FactoryNode))
 		{
-			if (UMeshComponent* MeshComponent = Cast<UMeshComponent>(SpawnedActor->GetRootComponent()))
+			if (UStaticMeshComponent* StaticMeshComponent = StaticMeshActor->GetStaticMeshComponent())
 			{
 				using namespace UE::Interchange;
-				
-				MeshComponent->EmptyOverrideMaterials();
-				ActorHelper::ApplySlotMaterialDependencies(*BaseNodeContainer, *MeshActorFactoryNode, *MeshComponent);
+				FString InstancedAssetFactoryNodeUid;
+				if (MeshActorFactoryNode->GetCustomInstancedAssetFactoryNodeUid(InstancedAssetFactoryNodeUid))
+				{
+					if (const UInterchangeFactoryBaseNode* MeshNode = Cast<UInterchangeFactoryBaseNode>(BaseNodeContainer->GetNode(InstancedAssetFactoryNodeUid)))
+					{
+						FSoftObjectPath ReferenceObject;
+						MeshNode->GetCustomReferenceObject(ReferenceObject);
+						if (UStaticMesh* StaticMesh = Cast<UStaticMesh>(ReferenceObject.TryLoad()))
+						{
+							if (StaticMesh != StaticMeshComponent->GetStaticMesh())
+							{
+								StaticMeshComponent->SetStaticMesh(StaticMesh);
+							}
+						}
+					}
+					else
+					{
+						// TODO: Warn that new mesh has not been applied
+					}
+				}
+
+				StaticMeshComponent->EmptyOverrideMaterials();
+				ActorHelper::ApplySlotMaterialDependencies(*BaseNodeContainer, *MeshActorFactoryNode, *StaticMeshComponent);
 			}
 		}
 	}
diff --git a/Engine/Plugins/Interchange/Runtime/Source/Import/Public/InterchangeAssetUserData.h b/Engine/Plugins/Interchange/Runtime/Source/Import/Public/InterchangeAssetUserData.h
--- a/Engine/Plugins/Interchange/Runtime/Source/Import/Public/InterchangeAssetUserData.h
+++ b/Engine/Plugins/Interchange/Runtime/Source/Import/Public/InterchangeAssetUserData.h
@@ -25,8 +25,6 @@ class UInterchangeLevelAssetUserData : public UAssetUserData
 	GENERATED_BODY()
 
 public:
-#if WITH_EDITORONLY_DATA
 	UPROPERTY(VisibleAnywhere, BlueprintReadWrite, AdvancedDisplay, Category = "Interchange User Data", meta = (ScriptName = "SceneImportPaths", DisplayName = "Scene Import Paths"))
 	TArray<FSoftObjectPath> SceneImportPaths;
-#endif
 };
diff --git a/Engine/Plugins/Interchange/Runtime/Source/Import/Public/Scene/InterchangeActorHelper.h b/Engine/Plugins/Interchange/Runtime/Source/Import/Public/Scene/InterchangeActorHelper.h
--- a/Engine/Plugins/Interchange/Runtime/Source/Import/Public/Scene/InterchangeActorHelper.h
+++ b/Engine/Plugins/Interchange/Runtime/Source/Import/Public/Scene/InterchangeActorHelper.h
@@ -4,11 +4,17 @@
 
 #include "InterchangeFactoryBase.h"
 
-class UInterchangeFactoryBaseNode;
-class UInterchangeBaseNodeContainer;
+class AActor;
+
 class UInterchangeActorFactoryNode;
+class UInterchangeAssetUserData;
+class UInterchangeBaseNodeContainer;
+class UInterchangeFactoryBaseNode;
 class UInterchangeMeshActorFactoryNode;
+class UInterchangeSceneImportAsset;
+
 class UMeshComponent;
+class UWorld;
 
 namespace UE::Interchange::ActorHelper
 {
@@ -38,4 +44,14 @@ namespace UE::Interchange::ActorHelper
 	 * Reimport policies will be applied if applicable
 	 */
 	INTERCHANGEIMPORT_API void ApplyAllCustomAttributes(const UInterchangeFactoryBase::FImportSceneObjectsParams& CreateSceneObjectsParams, UObject& ObjectToUpdate);
+
+	/**
+	 * Adds metadata to Actor via UInterchangeAssetUserData
+	 */
+	INTERCHANGEIMPORT_API void AddInterchangeAssetUserDataToActor(AActor* Actor, const UInterchangeSceneImportAsset* SceneImportAsset, const UInterchangeFactoryBaseNode* FactoryNode);
+
+	/**
+	 * Adds metadata to Actor via UInterchangeAssetUserData
+	 */
+	INTERCHANGEIMPORT_API void AddInterchangeLevelAssetUserDataToWorld(UWorld* World, const UInterchangeSceneImportAsset* SceneImportAsset);
 }
diff --git a/Engine/Plugins/Interchange/Runtime/Source/Import/Public/Scene/InterchangeSkeletalMeshActorFactory.h b/Engine/Plugins/Interchange/Runtime/Source/Import/Public/Scene/InterchangeSkeletalMeshActorFactory.h
--- a/Engine/Plugins/Interchange/Runtime/Source/Import/Public/Scene/InterchangeSkeletalMeshActorFactory.h
+++ b/Engine/Plugins/Interchange/Runtime/Source/Import/Public/Scene/InterchangeSkeletalMeshActorFactory.h
@@ -37,6 +37,7 @@ protected:
 
 	UE_API virtual UObject* ProcessActor(class AActor& SpawnedActor, const UInterchangeActorFactoryNode& FactoryNode, const UInterchangeBaseNodeContainer& NodeContainer, const FImportSceneObjectsParams& Params) override;
 
+	virtual void ExecuteResetObjectProperties(const UInterchangeBaseNodeContainer* BaseNodeContainer, UInterchangeFactoryBaseNode* FactoryNode, UObject* ImportedObject) override;
 	// Interchange actor factory interface end
 	//////////////////////////////////////////////////////////////////////////
 };
diff --git a/Engine/Plugins/Interchange/Runtime/Source/Pipelines/Private/InterchangeGenericScenesPipeline.cpp b/Engine/Plugins/Interchange/Runtime/Source/Pipelines/Private/InterchangeGenericScenesPipeline.cpp
--- a/Engine/Plugins/Interchange/Runtime/Source/Pipelines/Private/InterchangeGenericScenesPipeline.cpp
+++ b/Engine/Plugins/Interchange/Runtime/Source/Pipelines/Private/InterchangeGenericScenesPipeline.cpp
@@ -36,6 +36,7 @@
 #include "InterchangeVolumeNode.h"
 #include "Nodes/InterchangeSourceNode.h"
 #include "Nodes/InterchangeUserDefinedAttribute.h"
+#include "Scene/InterchangeActorHelper.h"
 #include "Volume/InterchangeVolumeDefinitions.h"
 
 #include "Animation/SkeletalMeshActor.h"
@@ -848,6 +849,15 @@ void UInterchangeGenericLevelPipeline::SetUpFactoryNode(UInterchangeActorFactory
 			}
 
 			MeshActorFactoryNode->AddFactoryDependencyUid(UInterchangeFactoryBaseNode::BuildFactoryNodeUid(MeshNode->GetUniqueID()));
+			if (!bSkeletal)
+			{
+				MeshActorFactoryNode->SetCustomInstancedAssetFactoryNodeUid(UInterchangeFactoryBaseNode::BuildFactoryNodeUid(MeshNode->GetUniqueID()));
+			}
+			else
+			{
+				// Directly points to the SkeletalMeshActorFactory that holds reference to SkeletalMesh Asset in the CustomReferencedObject.
+				MeshActorFactoryNode->SetCustomInstancedAssetFactoryNodeUid(TargetNodeUids[0]);
+			}
 
 			FTransform GeometricTransform;
 			if (SceneNode->GetCustomGeometricTransform(GeometricTransform))
@@ -1123,11 +1133,10 @@ void UInterchangeGenericLevelPipeline::SetUpFactoryNode(UInterchangeActorFactory
 			if (FString* FoundUid = SlotNameToMaterialPath.Find(UE::Interchange::Volume::VolumetricMaterial))
 			{
 				MaterialFactoryNodeUid = UInterchangeFactoryBaseNode::BuildFactoryNodeUid(*FoundUid);
-			}
-
-			if (!MaterialFactoryNodeUid.IsEmpty())
-			{
-				FactoryNode->SetCustomVolumetricMaterialUid(MaterialFactoryNodeUid);
+				if (!MaterialFactoryNodeUid.IsEmpty())
+				{
+					FactoryNode->SetCustomVolumetricMaterialUid(MaterialFactoryNodeUid);
+				}
 			}
 		}
 	}
@@ -1220,7 +1229,7 @@ void UInterchangeGenericLevelPipeline::ExecutePostImportPipeline(const UIntercha
 			if (SceneHierarchyType == EInterchangeSceneHierarchyType::CreateLevelActors)
 			{
 				InBaseNodeContainer->IterateNodesOfType<UInterchangeActorFactoryNode>(
-					[&SceneImportAssetPath](const FString& NodeUid, UInterchangeActorFactoryNode* FactoryNode)
+					[&SceneImportAssetPath, &SceneImportAsset](const FString& NodeUid, UInterchangeActorFactoryNode* FactoryNode)
 					{
 						FSoftObjectPath ObjectPath;
 						if (FactoryNode->GetCustomReferenceObject(ObjectPath))
@@ -1229,12 +1238,7 @@ void UInterchangeGenericLevelPipeline::ExecutePostImportPipeline(const UIntercha
 							{
 								if (UWorld* World = ReferencedObject->GetWorld())
 								{
-									if (USceneComponent* RootComponent = ReferencedObject->GetRootComponent())
-									{
-										UInterchangeAssetUserData* AssetUserData = NewObject<UInterchangeAssetUserData>(ReferencedObject);
-										AssetUserData->MetaData.Add(UE::Interchange::InterchangeReset::Constants::SceneImportAssetPathKey, SceneImportAssetPath.ToString());
-										RootComponent->AddAssetUserData(AssetUserData);
-									}
+									UE::Interchange::ActorHelper::AddInterchangeAssetUserDataToActor(ReferencedObject, SceneImportAsset, FactoryNode);
 								}
 							}
 						}
@@ -1244,7 +1248,7 @@ void UInterchangeGenericLevelPipeline::ExecutePostImportPipeline(const UIntercha
 			else
 			{
 				InBaseNodeContainer->IterateNodesOfType<UInterchangeLevelFactoryNode>(
-					[&SceneImportAssetPath](const FString& NodeUid, UInterchangeLevelFactoryNode* FactoryNode)
+					[&SceneImportAssetPath, &SceneImportAsset](const FString& NodeUid, UInterchangeLevelFactoryNode* FactoryNode)
 					{
 						bool bWasLevelCreatedForImport = false;
 						if (FactoryNode->GetCustomShouldCreateLevel(bWasLevelCreatedForImport) && bWasLevelCreatedForImport)
@@ -1254,12 +1258,7 @@ void UInterchangeGenericLevelPipeline::ExecutePostImportPipeline(const UIntercha
 							{
 								if (UWorld* World = Cast<UWorld>(ObjectPath.TryLoad()))
 								{
-									if (AWorldSettings* WorldSettings = World->GetWorldSettings())
-									{
-										UInterchangeLevelAssetUserData* AssetUserData = NewObject<UInterchangeLevelAssetUserData>(World);
-										AssetUserData->SceneImportPaths.Add(SceneImportAssetPath);
-										WorldSettings->AddAssetUserData(AssetUserData);
-									}
+									UE::Interchange::ActorHelper::AddInterchangeLevelAssetUserDataToWorld(World, SceneImportAsset);
 								}
 							}
 						}
@@ -1411,14 +1410,9 @@ void UInterchangeGenericLevelPipeline::ExecutePostBroadcastPipeline(const UInter
 				{
 					if (CurrentLevelInstanceActor->GetWorldAsset().Get() == WorldAsset)
 					{
-						if (APackedLevelActor* PackedLevelActor = Cast<APackedLevelActor>(CurrentLevelInstanceActor))
-						{
-
-							PackedLevelActor->UpdateLevelInstanceFromWorldAsset();
-						}
-						else
+						CurrentLevelInstanceActor->UpdateLevelInstanceFromWorldAsset();
+						if (!CurrentLevelInstanceActor->IsA<APackedLevelActor>())
 						{
-							CurrentLevelInstanceActor->UpdateLevelInstanceFromWorldAsset();
 							CurrentLevelInstanceActor->LoadLevelInstance();
 						}
 					}
diff --git a/Engine/Source/Runtime/Interchange/Core/Private/Nodes/InterchangeFactoryBaseNode.cpp b/Engine/Source/Runtime/Interchange/Core/Private/Nodes/InterchangeFactoryBaseNode.cpp
--- a/Engine/Source/Runtime/Interchange/Core/Private/Nodes/InterchangeFactoryBaseNode.cpp
+++ b/Engine/Source/Runtime/Interchange/Core/Private/Nodes/InterchangeFactoryBaseNode.cpp
@@ -336,6 +336,8 @@ void UInterchangeFactoryBaseNode::RestoreAllCustomAttributeDelegates()
 		AttributesAppliedThroughDelegatesKeySet.Add(Key);
 	}
 
+	// The same attributes would be set again when the actual reset happens.
+	AttributesAppliedThroughDelegates.RemoveAllItems();
 	OnRestoreAllCustomAttributeDelegates();
 }
 
