[Interchange Reset] Light Actors can be reset now

#rnx

diff --git a/Engine/Plugins/Interchange/Runtime/Source/FactoryNodes/Private/InterchangeLightFactoryNode.cpp b/Engine/Plugins/Interchange/Runtime/Source/FactoryNodes/Private/InterchangeLightFactoryNode.cpp
--- a/Engine/Plugins/Interchange/Runtime/Source/FactoryNodes/Private/InterchangeLightFactoryNode.cpp
+++ b/Engine/Plugins/Interchange/Runtime/Source/FactoryNodes/Private/InterchangeLightFactoryNode.cpp
@@ -58,6 +58,15 @@ void UInterchangeBaseLightFactoryNode::CopyWithObject(const UInterchangeFactoryB
 	}
 }
 
+void UInterchangeBaseLightFactoryNode::OnRestoreAllCustomAttributeDelegates()
+{
+	Super::OnRestoreAllCustomAttributeDelegates();
+	REFILL_CUSTOM_ATTRIBUTE_APPLY_DELEGATE(LightColor, FColor);
+	REFILL_CUSTOM_ATTRIBUTE_APPLY_DELEGATE(Intensity, float);
+	REFILL_CUSTOM_ATTRIBUTE_APPLY_DELEGATE(Temperature, float);
+	REFILL_CUSTOM_ATTRIBUTE_APPLY_DELEGATE_EXPLICIT(Macro_CustombUseTemperatureKey.Key, GetCustomUseTemperature, SetCustomUseTemperature, bool);
+}
+
 bool UInterchangeLightFactoryNode::GetCustomIntensityUnits(ELightUnits& AttributeValue) const
 {
 	IMPLEMENT_NODE_ATTRIBUTE_GETTER(IntensityUnits, ELightUnits)
@@ -129,6 +138,14 @@ void UInterchangeLightFactoryNode::CopyWithObject(const UInterchangeFactoryBaseN
 	}
 }
 
+void UInterchangeLightFactoryNode::OnRestoreAllCustomAttributeDelegates()
+{
+	Super::OnRestoreAllCustomAttributeDelegates();
+	REFILL_CUSTOM_ATTRIBUTE_APPLY_DELEGATE(IntensityUnits, ELightUnits)
+	REFILL_CUSTOM_ATTRIBUTE_APPLY_DELEGATE(AttenuationRadius, float)
+	REFILL_CUSTOM_ATTRIBUTE_APPLY_DELEGATE(UseIESBrightness, bool)
+}
+
 bool UInterchangePointLightFactoryNode::GetCustomUseInverseSquaredFalloff(bool& AttributeValue) const
 {
 	IMPLEMENT_NODE_ATTRIBUTE_GETTER(bUseInverseSquaredFalloff, bool)
@@ -221,3 +238,11 @@ void UInterchangeRectLightFactoryNode::CopyWithObject(const UInterchangeFactoryB
 		COPY_NODE_DELEGATES(LightFactoryNode, SourceHeight, float, URectLightComponent)
 	}
 }
+
+void UInterchangeRectLightFactoryNode::OnRestoreAllCustomAttributeDelegates()
+{
+	Super::OnRestoreAllCustomAttributeDelegates();
+
+	REFILL_CUSTOM_ATTRIBUTE_APPLY_DELEGATE(SourceWidth, float)
+	REFILL_CUSTOM_ATTRIBUTE_APPLY_DELEGATE(SourceHeight, float)
+}
diff --git a/Engine/Plugins/Interchange/Runtime/Source/FactoryNodes/Public/InterchangeLightFactoryNode.h b/Engine/Plugins/Interchange/Runtime/Source/FactoryNodes/Public/InterchangeLightFactoryNode.h
--- a/Engine/Plugins/Interchange/Runtime/Source/FactoryNodes/Public/InterchangeLightFactoryNode.h
+++ b/Engine/Plugins/Interchange/Runtime/Source/FactoryNodes/Public/InterchangeLightFactoryNode.h
@@ -42,6 +42,7 @@ public:
 
 	UE_API virtual void CopyWithObject(const UInterchangeFactoryBaseNode* SourceNode, UObject* Object) override;
 
+	virtual void OnRestoreAllCustomAttributeDelegates() override;
 private:
 
 	IMPLEMENT_NODE_ATTRIBUTE_KEY(LightColor)
@@ -101,6 +102,7 @@ public:
 
 	UE_API virtual void CopyWithObject(const UInterchangeFactoryBaseNode* SourceNode, UObject* Object) override;
 
+	virtual void OnRestoreAllCustomAttributeDelegates() override;
 private:
 
 	IMPLEMENT_NODE_ATTRIBUTE_KEY(IntensityUnits)
@@ -132,6 +134,7 @@ public:
 
 	UE_API virtual void CopyWithObject(const UInterchangeFactoryBaseNode* SourceNode, UObject* Object) override;
 
+	virtual void OnRestoreAllCustomAttributeDelegates() override;
 private:
 	IMPLEMENT_NODE_ATTRIBUTE_KEY(SourceWidth)
 	IMPLEMENT_NODE_ATTRIBUTE_KEY(SourceHeight)
diff --git a/Engine/Source/Runtime/Interchange/Core/Public/Nodes/InterchangeFactoryBaseNode.h b/Engine/Source/Runtime/Interchange/Core/Public/Nodes/InterchangeFactoryBaseNode.h
--- a/Engine/Source/Runtime/Interchange/Core/Public/Nodes/InterchangeFactoryBaseNode.h
+++ b/Engine/Source/Runtime/Interchange/Core/Public/Nodes/InterchangeFactoryBaseNode.h
@@ -138,6 +138,18 @@ if(AttributesAppliedThroughDelegatesKeySet.Contains(Macro_Custom##AttributeName#
 	}																																			\
 }																																				\
 
+#define REFILL_CUSTOM_ATTRIBUTE_APPLY_DELEGATE_EXPLICIT(AttributeKey, Getter, Setter, AttributeType)																	\
+if(AttributesAppliedThroughDelegatesKeySet.Contains(AttributeKey))																				\
+{																																				\
+	AttributeType ValueData;																													\
+	if(Getter(ValueData))																														\
+	{																																			\
+		constexpr const bool bAddApplyDelegate = true;																							\
+		Setter(ValueData, bAddApplyDelegate);																									\
+	}																																			\
+}
+
+
 namespace UE::Interchange
 {
 
@@ -670,4 +682,4 @@ inline bool UInterchangeFactoryBaseNode::FillAttributeFromObject<bool>(const FSt
 	}
 
 	return false;
-}
+}
\ No newline at end of file
