Interchange MaterialX:
- Connnect local position to the "Absolute World Position" of the Material Function HeightToNormalSmooth

diff --git a/Engine/Plugins/Interchange/Runtime/Source/Common/Public/MaterialX/InterchangeMaterialXDefinitions.h b/Engine/Plugins/Interchange/Runtime/Source/Common/Public/MaterialX/InterchangeMaterialXDefinitions.h
--- a/Engine/Plugins/Interchange/Runtime/Source/Common/Public/MaterialX/InterchangeMaterialXDefinitions.h
+++ b/Engine/Plugins/Interchange/Runtime/Source/Common/Public/MaterialX/InterchangeMaterialXDefinitions.h
@@ -221,6 +221,7 @@ namespace UE
 		namespace Inputs
 		{
 			static constexpr const TCHAR* A = TEXT("A");
+			static constexpr const TCHAR* AbsoluteWorldPosition = TEXT("Absolute World Position");
 			static constexpr const TCHAR* AEqualsB = TEXT("AEqualsB");
 			static constexpr const TCHAR* AGreaterThanB = TEXT("AGreaterThanB");
 			static constexpr const TCHAR* Alpha = TEXT("Alpha");
diff --git a/Engine/Plugins/Interchange/Runtime/Source/Import/Private/MaterialX/MaterialXUtils/MaterialXSurfaceShaderAbstract.cpp b/Engine/Plugins/Interchange/Runtime/Source/Import/Private/MaterialX/MaterialXUtils/MaterialXSurfaceShaderAbstract.cpp
--- a/Engine/Plugins/Interchange/Runtime/Source/Import/Private/MaterialX/MaterialXUtils/MaterialXSurfaceShaderAbstract.cpp
+++ b/Engine/Plugins/Interchange/Runtime/Source/Import/Private/MaterialX/MaterialXUtils/MaterialXSurfaceShaderAbstract.cpp
@@ -1229,12 +1229,21 @@ void FMaterialXSurfaceShaderAbstract::ConnectHeightToNormalInputToOutput(const F
 			UInterchangeFunctionCallShaderNode* HeightToNormalSmoothShaderNode =
 				CreateFunctionCallShaderNode(Connect.UpstreamNode,
 											 Connect.UpstreamNode->getName().c_str(),
-											 TEXT("/Engine/Functions/Engine_MaterialFunctions03/Procedurals/HeightToNormalSmooth.HeightToNormalSmooth"));
+											 UE::MaterialFunctions::Path::HeightToNormalSmooth);
 
 			UInterchangeShaderNode* TransformNode = CreateShaderNode(Connect.UpstreamNode, (Connect.UpstreamNode->getName() + "_TS").c_str(), TransformVector::Name.ToString());
 			TransformNode->AddInt32Attribute(TransformVector::Attributes::TransformSourceType.ToString(), EMaterialVectorCoordTransformSource::TRANSFORMSOURCE_World);
 			TransformNode->AddInt32Attribute(TransformVector::Attributes::TransformType.ToString(), EMaterialVectorCoordTransform::TRANSFORM_Tangent);
 
+			// By default positions in MaterialX are in object space so we need to create a local position expression here.
+			// TODO: Handling it as a geompropdef so all nodes refer to the same LocalPosition expression
+			UInterchangeShaderNode* PositionNode = CreateShaderNode(Connect.UpstreamNode, (Connect.UpstreamNode->getName() + "_Position").c_str(), UE::Expressions::Names::LocalPosition);
+			UInterchangeShaderNode* UnitNode = CreateShaderNode(Connect.UpstreamNode, (Connect.UpstreamNode->getName() + "_Unit").c_str(), Multiply::Name.ToString());
+
+			UnitNode->AddFloatAttribute(UInterchangeShaderPortsAPI::MakeInputValueKey(Multiply::Inputs::B.ToString()), 0.01f);
+			UInterchangeShaderPortsAPI::ConnectDefaultOuputToInput(UnitNode, Multiply::Inputs::A.ToString(), PositionNode->GetUniqueID());
+			UInterchangeShaderPortsAPI::ConnectDefaultOuputToInput(HeightToNormalSmoothShaderNode, UE::Expressions::Inputs::AbsoluteWorldPosition, UnitNode->GetUniqueID());
+
 			UInterchangeShaderPortsAPI::ConnectDefaultOuputToInput(TransformNode, TransformVector::Inputs::Input.ToString(), HeightToNormalSmoothShaderNode->GetUniqueID());
 			UInterchangeShaderPortsAPI::ConnectDefaultOuputToInput(Connect.ParentShaderNode, Connect.InputChannelName, TransformNode->GetUniqueID());
 		}
