[BUGFIX] Interchange GLTF Translator Analytics Helper Crash
- Added an early return when analytics helper is not valid.

#rnx

diff --git a/Engine/Plugins/Interchange/Runtime/Source/Import/Private/Gltf/InterchangeGltfTranslator.cpp b/Engine/Plugins/Interchange/Runtime/Source/Import/Private/Gltf/InterchangeGltfTranslator.cpp
--- a/Engine/Plugins/Interchange/Runtime/Source/Import/Private/Gltf/InterchangeGltfTranslator.cpp
+++ b/Engine/Plugins/Interchange/Runtime/Source/Import/Private/Gltf/InterchangeGltfTranslator.cpp
@@ -158,6 +158,11 @@ namespace UE::Interchange::Gltf::Private
 		const GLTF::FAsset& Asset = GLTF::FAsset(),
 		const FText& GLTFReaderLogMessage = FText::GetEmpty())
 	{
+		if (!AnalyticsHelper.IsValid())
+		{
+			return;
+		}
+
 		TMap<FString, FString> MetadataExtras;
 		for (GLTF::FMetadata::FExtraData ExtraData : Asset.Metadata.Extras)
 		{
