[Interchange] MorphTarget Name empty related fixes.

#rnx

diff --git a/Engine/Plugins/Interchange/Runtime/Source/Import/Private/Mesh/InterchangeMeshHelper.cpp b/Engine/Plugins/Interchange/Runtime/Source/Import/Private/Mesh/InterchangeMeshHelper.cpp
--- a/Engine/Plugins/Interchange/Runtime/Source/Import/Private/Mesh/InterchangeMeshHelper.cpp
+++ b/Engine/Plugins/Interchange/Runtime/Source/Import/Private/Mesh/InterchangeMeshHelper.cpp
@@ -590,6 +590,11 @@ namespace UE::Interchange::Private::MeshHelper
 			const FString& MorphTargetNameBase = NamePerUIDs.Key;
 			FString MorphTargetName(MorphTargetNameBase);
 
+			if (MorphTargetName.IsEmpty())
+			{
+				MorphTargetName = TEXT("__Empty__");
+			}
+
 			for (int32 Suffix = 1; UsedMorphTargetNames.Contains(*MorphTargetName); Suffix++)
 			{
 				MorphTargetName = FString::Printf(TEXT("%s_%d"), *MorphTargetNameBase, Suffix);
diff --git a/Engine/Plugins/Interchange/Runtime/Source/Parsers/Fbx/Private/FbxMesh.cpp b/Engine/Plugins/Interchange/Runtime/Source/Parsers/Fbx/Private/FbxMesh.cpp
--- a/Engine/Plugins/Interchange/Runtime/Source/Parsers/Fbx/Private/FbxMesh.cpp
+++ b/Engine/Plugins/Interchange/Runtime/Source/Parsers/Fbx/Private/FbxMesh.cpp
@@ -1865,7 +1865,7 @@ void FFbxMesh::AddAllMeshes(FbxScene* SDKScene, FbxGeometryConverter* SDKGeometr
 				const int32 CurrentChannelMorphTargetCount = Channel->GetTargetShapeCount();
 				FString ChannelName = Parser.GetFbxHelper()->GetFbxObjectName(Channel);
 				// Maya adds the name of the MorphTarget and an underscore to the front of the channel name, so remove it
-				if (ChannelName.StartsWith(MorphTargetName))
+				if (ChannelName.StartsWith(MorphTargetName) && (ChannelName.Compare(MorphTargetName, ESearchCase::IgnoreCase) == 0))
 				{
 					ChannelName.RightInline(ChannelName.Len() - (MorphTargetName.Len() + 1), EAllowShrinking::No);
 				}
@@ -1889,6 +1889,12 @@ void FFbxMesh::AddAllMeshes(FbxScene* SDKScene, FbxGeometryConverter* SDKGeometr
 							ShapeName = ChannelName;
 						}
 					}
+
+					if (ShapeName.IsEmpty())
+					{
+						ShapeName = TEXT("MorphTarget_") + FString::FromInt(ChannelIndex) + TEXT("_") + FString::FromInt(ChannelMorphTargetIndex);
+					}
+
 					ensure(!ShapeNameToFbxShape.Contains(ShapeName));
 					ShapeNameToFbxShape.Add(ShapeName, Shape);
 					FString MorphTargetAttributeName = Parser.GetFbxHelper()->GetMeshName(Shape);
