Interchange:
- Add pipeline option to set the displacement center of materials, the option overrides any value set by shader graph nodes

diff --git a/Engine/Plugins/Interchange/Runtime/Source/Pipelines/Private/InterchangeGenericMaterialPipeline.cpp b/Engine/Plugins/Interchange/Runtime/Source/Pipelines/Private/InterchangeGenericMaterialPipeline.cpp
--- a/Engine/Plugins/Interchange/Runtime/Source/Pipelines/Private/InterchangeGenericMaterialPipeline.cpp
+++ b/Engine/Plugins/Interchange/Runtime/Source/Pipelines/Private/InterchangeGenericMaterialPipeline.cpp
@@ -1918,7 +1918,7 @@ void UInterchangeGenericMaterialPipeline::HandleCommonParameters(const UIntercha
 				MaterialFactoryNode->ConnectOutputToDisplacement(DisplacementExpressionFactoryNode.Get<0>()->GetUniqueID(), DisplacementExpressionFactoryNode.Get<1>());
 			}
 
-			if (float DisplacementCenter; ShaderGraphNode->GetCustomDisplacementCenterMode(DisplacementCenter))
+			if (float DisplacementCenter; !bOverrideDisplacement && ShaderGraphNode->GetCustomDisplacementCenterMode(DisplacementCenter))
 			{
 				MaterialFactoryNode->SetCustomDisplacementCenter(DisplacementCenter);
 			}
@@ -3476,6 +3476,11 @@ UInterchangeMaterialFactoryNode* UInterchangeGenericMaterialPipeline::CreateMate
 	UInterchangeMaterialFactoryNode* MaterialFactoryNode = Cast<UInterchangeMaterialFactoryNode>( CreateBaseMaterialFactoryNode(ShaderGraphNode, UInterchangeMaterialFactoryNode::StaticClass()) );
 
 	UInterchangeUserDefinedAttributesAPI::DuplicateAllUserDefinedAttribute(ShaderGraphNode, MaterialFactoryNode, false);
+	
+	if (bOverrideDisplacement)
+	{
+		MaterialFactoryNode->SetCustomDisplacementCenter(OverrideDisplacementCenter);
+	}
 
 	if(HandleSubstrate(ShaderGraphNode, MaterialFactoryNode))
 	{
@@ -3942,7 +3947,7 @@ bool UInterchangeGenericMaterialPipeline::HandleSubstrate(const UInterchangeShad
 				UInterchangeShaderPortsAPI::ConnectOuputToInputByName(MaterialFactoryNode, SubstrateMaterial::Parameters::Displacement.ToString(), DisplacementFactoryNode.Get<0>()->GetUniqueID(), DisplacementFactoryNode.Get<1>());
 			}
 
-			if (float DisplacementCenter; ShaderGraphNode->GetCustomDisplacementCenterMode(DisplacementCenter))
+			if (float DisplacementCenter; !bOverrideDisplacement && ShaderGraphNode->GetCustomDisplacementCenterMode(DisplacementCenter))
 			{
 				MaterialFactoryNode->SetCustomDisplacementCenter(DisplacementCenter);
 			}
diff --git a/Engine/Plugins/Interchange/Runtime/Source/Pipelines/Public/InterchangeGenericMaterialPipeline.h b/Engine/Plugins/Interchange/Runtime/Source/Pipelines/Public/InterchangeGenericMaterialPipeline.h
--- a/Engine/Plugins/Interchange/Runtime/Source/Pipelines/Public/InterchangeGenericMaterialPipeline.h
+++ b/Engine/Plugins/Interchange/Runtime/Source/Pipelines/Public/InterchangeGenericMaterialPipeline.h
@@ -105,6 +105,14 @@ public:
 	UPROPERTY(VisibleAnywhere, BlueprintReadWrite, Instanced, Category = "Textures")
 	TObjectPtr<UInterchangeSparseVolumeTexturePipeline> SparseVolumeTexturePipeline;
 
+	/** If enabled, it will override the displacement center set by shader graph nodes, if any*/
+	UPROPERTY(EditAnywhere, BlueprintReadWrite, Category = "Materials", Meta = (InlineEditConditionToggle))
+	bool bOverrideDisplacement = false;
+
+	/** Set the value of the displacement center. If enabled it will also override any displacement center value set by shader graph nodes*/
+	UPROPERTY(EditAnywhere, BlueprintReadWrite, Category = "Materials", Meta = (EditCondition = "bOverrideDisplacement", ClampMin = "0", DisplayName = "Displacement Center"))
+	float OverrideDisplacementCenter = 0.5f;
+
 	/** BEGIN UInterchangePipelineBase overrides */
 	UE_API virtual void PreDialogCleanup(const FName PipelineStackName) override;
 	UE_API virtual bool IsSettingsAreValid(TOptional<FText>& OutInvalidReason) const override;
