Interchange - FBX - Level import - Light properties not properly imported

#rnx

diff --git a/Engine/Plugins/Interchange/Runtime/Source/Parsers/Fbx/Private/FbxLight.cpp b/Engine/Plugins/Interchange/Runtime/Source/Parsers/Fbx/Private/FbxLight.cpp
--- a/Engine/Plugins/Interchange/Runtime/Source/Parsers/Fbx/Private/FbxLight.cpp
+++ b/Engine/Plugins/Interchange/Runtime/Source/Parsers/Fbx/Private/FbxLight.cpp
@@ -49,9 +49,39 @@ namespace UE
 					Message->Text = LOCTEXT("CannotAllocateNode", "Cannot allocate a node when importing FBX.");
 					return nullptr;
 				}
-
 				NodeContainer.SetupNode(BaseLightNode, NodeUid, NodeName, EInterchangeNodeContainerType::TranslatedAsset);
 
+				const FbxDouble3 LightColor = LightAttribute.Color.Get();
+				const FColor UnrealLightColor(uint8(255.0 * LightColor[0]), uint8(255.0 * LightColor[1]), uint8(255.0 * LightColor[2]));
+				BaseLightNode->SetCustomLightColor(UnrealLightColor);
+				
+				const FbxDouble Intensity = LightAttribute.Intensity.Get();
+				BaseLightNode->SetCustomIntensity((float)Intensity);
+
+				if (UInterchangeLightNode* InterchangeLightNode = Cast<UInterchangeLightNode>(BaseLightNode))
+				{
+					InterchangeLightNode->SetCustomAttenuationRadius(LightAttribute.DecayStart.Get());
+					const FbxLight::EDecayType Decay = LightAttribute.DecayType.Get();
+					if (Decay == FbxLight::EDecayType::eNone)
+					{
+						InterchangeLightNode->SetCustomAttenuationRadius(FBXSDK_FLOAT_MAX);
+					}
+				}
+
+				// NOTE: Area Lights do not have dimensions property on FBX Lights to be able to set the width and height of the lights.				
+				switch (LightAttribute.LightType.Get())
+				{
+				case FbxLight::eSpot:
+					if (UInterchangeSpotLightNode* SpotLightNode = Cast<UInterchangeSpotLightNode>(BaseLightNode))
+					{
+						SpotLightNode->SetCustomInnerConeAngle(LightAttribute.InnerAngle.Get());
+						SpotLightNode->SetCustomOuterConeAngle(LightAttribute.OuterAngle.Get());
+					}
+					break;
+				default:
+					break;
+				}
+
 				return BaseLightNode;
 			}
 
