[Interchange] Reset properties does not reset an interchange tracked property
- ActorFactoryNode now gets default visibility values through attributes for Interchange Reset to Work correctly as they will add appropriate delegates.

#rnx

diff --git a/Engine/Plugins/Interchange/Runtime/Source/Pipelines/Private/InterchangeGenericScenesPipeline.cpp b/Engine/Plugins/Interchange/Runtime/Source/Pipelines/Private/InterchangeGenericScenesPipeline.cpp
--- a/Engine/Plugins/Interchange/Runtime/Source/Pipelines/Private/InterchangeGenericScenesPipeline.cpp
+++ b/Engine/Plugins/Interchange/Runtime/Source/Pipelines/Private/InterchangeGenericScenesPipeline.cpp
@@ -709,15 +709,13 @@ void UInterchangeGenericLevelPipeline::ExecuteSceneNodePreImport(const FTransfor
 	}
 
 	bool bCustomVisibility = true;
-	if (SceneNode->GetCustomComponentVisibility(bCustomVisibility))
-	{
-		ActorFactoryNode->SetCustomComponentVisibility(bCustomVisibility);
-	}
-	if (SceneNode->GetCustomActorVisibility(bCustomVisibility))
-	{
-		ActorFactoryNode->SetCustomActorVisibility(bCustomVisibility);
-	}
-
+	SceneNode->GetCustomComponentVisibility(bCustomVisibility);
+	ActorFactoryNode->SetCustomComponentVisibility(bCustomVisibility);
+	
+	bCustomVisibility = true;
+	SceneNode->GetCustomActorVisibility(bCustomVisibility);
+	ActorFactoryNode->SetCustomActorVisibility(bCustomVisibility);
+	
 	ActorFactoryNode->SetCustomMobility(EComponentMobility::Static);
 
 	if (TranslatedAssetNode)
