Interchange: Allow using existing material instances.

diff --git a/Engine/Plugins/Interchange/Runtime/Source/Import/Private/Material/InterchangeMaterialFactory.cpp b/Engine/Plugins/Interchange/Runtime/Source/Import/Private/Material/InterchangeMaterialFactory.cpp
--- a/Engine/Plugins/Interchange/Runtime/Source/Import/Private/Material/InterchangeMaterialFactory.cpp
+++ b/Engine/Plugins/Interchange/Runtime/Source/Import/Private/Material/InterchangeMaterialFactory.cpp
@@ -975,6 +975,12 @@ UInterchangeFactoryBase::FImportAssetResult UInterchangeMaterialFactory::BeginIm
 	}
 #endif //WITH_EDITOR
 
+	if (bSkipImport)
+	{
+		// Don't assign to ImportAssetResult.ImportedObject here or else Interchange will dirty the object
+		return ImportAssetResult;
+	}
+
 	// Setup material instance based on reimport policy if applicable
 	if (const UInterchangeMaterialInstanceFactoryNode* MaterialInstanceFactoryNode = Cast<UInterchangeMaterialInstanceFactoryNode>(MaterialFactoryNode))
 	{
diff --git a/Engine/Plugins/Interchange/Runtime/Source/Pipelines/Private/InterchangeGenericMaterialPipeline.cpp b/Engine/Plugins/Interchange/Runtime/Source/Pipelines/Private/InterchangeGenericMaterialPipeline.cpp
--- a/Engine/Plugins/Interchange/Runtime/Source/Pipelines/Private/InterchangeGenericMaterialPipeline.cpp
+++ b/Engine/Plugins/Interchange/Runtime/Source/Pipelines/Private/InterchangeGenericMaterialPipeline.cpp
@@ -1050,35 +1050,46 @@ void UInterchangeGenericMaterialPipeline::ExecutePipeline(UInterchangeBaseNodeCo
 		OverrideAssetName = AssetName;
 	}
 
+	// Enable or disable factory nodes according to import/reimport context
 	const bool bCombinedImportMaterials = bImportMaterials || bImportUnusedMaterial;
 	for (UInterchangeBaseMaterialFactoryNode* BaseMaterialFactoryNode : BaseMaterialFactoryNodes)
 	{
 		using namespace UE::Interchange::InterchangeGenericMaterialPipeline::Private;
 
-		BaseMaterialFactoryNode->SetCustomIsMaterialImportEnabled(bCombinedImportMaterials);
+		bool bEnableFactoryNode = bCombinedImportMaterials;
 
-		// Disable all materials if the toggle is off
-		if (!bCombinedImportMaterials)
+		if (bEnableFactoryNode)
 		{
-			BaseMaterialFactoryNode->SetEnabled(false);
-		}
+			FString MaterialName = BaseMaterialFactoryNode->GetDisplayLabel();
+			if (UMaterialInterface* ExistingMaterial = FindExistingMaterialFromSearchLocation(MaterialName, ContentBasePath, SearchLocation))
+			{
+				BaseMaterialFactoryNode->SetCustomReferenceObject(ExistingMaterial);
 
-		// See if we can assign an existing material from the search location
-		FString MaterialName = BaseMaterialFactoryNode->GetDisplayLabel();
-		if (UMaterialInterface* ExistingMaterial = FindExistingMaterialFromSearchLocation(MaterialName, ContentBasePath, SearchLocation))
-		{
-			// Make sure we have the correct type of material (can be material instance) before setting the custom object reference.
-			const bool bIsMaterial = BaseMaterialFactoryNode->IsA<UInterchangeMaterialFactoryNode>() && ExistingMaterial->IsA<UMaterial>();
-			const bool bIsMaterialinstance = BaseMaterialFactoryNode->IsA<UInterchangeMaterialInstanceFactoryNode>()
-				&& ExistingMaterial->IsA<UMaterialInstance>();
+				// We're reimporting the existing material we found: Enable the factory node if we can.
+				// TODO: Can this potentially lead to issues, since the existing material to reimport is discovered by name and search location?
+				if (CacheContextParam.ReimportAsset && CacheContextParam.ReimportAsset == ExistingMaterial)
+				{
+					const bool bIsMaterial = BaseMaterialFactoryNode->IsA<UInterchangeMaterialFactoryNode>()
+											 && ExistingMaterial->IsA<UMaterial>();
 
-			BaseMaterialFactoryNode->SetCustomReferenceObject(ExistingMaterial);
+					const bool bIsMaterialinstance = BaseMaterialFactoryNode->IsA<UInterchangeMaterialInstanceFactoryNode>()
+													 && ExistingMaterial->IsA<UMaterialInstance>();
 
-			// Reimport can only be done on material instances
-			const bool bEnableReimport = !bIsMaterial && bIsMaterialinstance;
-			BaseMaterialFactoryNode->SetCustomIsMaterialImportEnabled(bEnableReimport);
-			BaseMaterialFactoryNode->SetEnabled(bEnableReimport);
+					// Reimport can only be done on material instances
+					bEnableFactoryNode = bEnableFactoryNode && !bIsMaterial && bIsMaterialinstance;
+				}
+				// Found an existing material on the search location during an import --> Just use it as-is
+				else
+				{
+					// Disabling the factory node here is important so that we don't overwrite the existing asset.
+					// The existing material will still actually be used for the import's material assignments and etc.
+					bEnableFactoryNode = false;
+				}
+			}
 		}
+
+		BaseMaterialFactoryNode->SetCustomIsMaterialImportEnabled(bEnableFactoryNode);
+		BaseMaterialFactoryNode->SetEnabled(bEnableFactoryNode);
 	}
 
 	if (IsStandAlonePipeline() && !OverrideAssetName.IsEmpty() && BaseMaterialFactoryNodes.Num() == 1)
diff --git a/Engine/Source/Runtime/Interchange/Core/Private/InterchangePipelineBase.cpp b/Engine/Source/Runtime/Interchange/Core/Private/InterchangePipelineBase.cpp
--- a/Engine/Source/Runtime/Interchange/Core/Private/InterchangePipelineBase.cpp
+++ b/Engine/Source/Runtime/Interchange/Core/Private/InterchangePipelineBase.cpp
@@ -72,6 +72,31 @@ FString UInterchangePipelineBase::GetDefaultConfigFileName()
 	return PipelineConfigFilename;
 }
 
+FArchive& operator<<(FArchive& Ar, FInterchangePipelineContextParams& Params)
+{
+	Ar << Params.ContextType;
+	Ar << Params.ImportObjectType;
+	Ar << Params.ReimportAsset;
+	Ar << Params.BaseNodeContainer;
+	return Ar;
+}
+
+FArchive& operator<<(FArchive& Ar, FInterchangePipelinePropertyStatePerContext& States)
+{
+	Ar << States.bVisible;
+	return Ar;
+}
+
+FArchive& operator<<(FArchive& Ar, FInterchangePipelinePropertyStates& States)
+{
+	Ar << States.bLocked;
+	Ar << States.bPreDialogReset;
+	Ar << States.BasicLayoutStates;
+	Ar << States.ImportStates;
+	Ar << States.ReimportStates;
+	return Ar;
+}
+
 void UInterchangePipelineBase::LoadSettings(const FName PipelineStackName, bool bResetPreDialog /*= false*/)
 {
 	bool bOutRequiresSaving = false;
@@ -546,4 +571,21 @@ void UInterchangePipelineBase::PostDuplicate(bool bDuplicateForPIE)
 	Super::PostDuplicate(bDuplicateForPIE);
 
 	UpdateWeakObjectPtrs();
-}
\ No newline at end of file
+}
+
+void UInterchangePipelineBase::Serialize(FArchive& Ar)
+{
+	Super::Serialize(Ar);
+
+	if (Ar.GetPortFlags() & PPF_Duplicate)
+	{
+		Ar << bAllowPropertyStatesEdition;
+		Ar << bIsReimportContext;
+		Ar << bIsShowEssentials;
+		Ar << bFromReimportOrOverride;
+		Ar << Results;
+		Ar << PropertiesStates;
+		Ar << CachePropertiesStates;
+		Ar << CacheContextParam;
+	}
+}
diff --git a/Engine/Source/Runtime/Interchange/Core/Public/InterchangePipelineBase.h b/Engine/Source/Runtime/Interchange/Core/Public/InterchangePipelineBase.h
--- a/Engine/Source/Runtime/Interchange/Core/Public/InterchangePipelineBase.h
+++ b/Engine/Source/Runtime/Interchange/Core/Public/InterchangePipelineBase.h
@@ -60,11 +60,19 @@ struct FInterchangePipelineContextParams
 {
 	GENERATED_BODY()
 
+	UPROPERTY()
 	EInterchangePipelineContext ContextType = EInterchangePipelineContext::None;
-	UClass* ImportObjectType = nullptr;
+
+	UPROPERTY()
+	TObjectPtr<UClass> ImportObjectType = nullptr;
+
+	UPROPERTY()
 	TObjectPtr<UObject> ReimportAsset = nullptr;
-	const UInterchangeBaseNodeContainer* BaseNodeContainer = nullptr;
 
+	UPROPERTY()
+	TObjectPtr<const UInterchangeBaseNodeContainer> BaseNodeContainer = nullptr;
+
+	friend FArchive& operator<<(FArchive& Ar, FInterchangePipelineContextParams& Params);
 };
 
 USTRUCT(BlueprintType)
@@ -75,6 +83,8 @@ struct FInterchangePipelinePropertyStatePerContext
 	/** If true, the property is visible. */
 	UPROPERTY(EditAnywhere, BlueprintReadWrite, Category = "Property States")
 	bool bVisible = true;
+
+	friend FArchive& operator<<(FArchive& Ar, FInterchangePipelinePropertyStatePerContext& Params);
 };
 
 USTRUCT(BlueprintType)
@@ -156,6 +166,8 @@ struct FInterchangePipelinePropertyStates
 	/** The property states for the reimport context. */
 	UPROPERTY(EditAnywhere, BlueprintReadWrite, Category = "Context Properties States")
 	FInterchangePipelinePropertyStatePerContext ReimportStates;
+
+	friend FArchive& operator<<(FArchive& Ar, FInterchangePipelinePropertyStates& Params);
 };
 
 struct FInterchangeConflictInfo
@@ -528,6 +540,7 @@ public:
 
 	/** Begin UObject overrides*/
 	INTERCHANGECORE_API virtual void PostDuplicate(bool bDuplicateForPIE) override;
+	INTERCHANGECORE_API virtual void Serialize(FArchive& Ar) override;
 	/** End UObject overrides */
 
 	/*
@@ -622,6 +635,7 @@ protected:
 	 * Note: If you open a pipeline asset in the Content Browser, you will be able to edit the property states.
 	 *       If you import a file with Interchange, the import dialog will apply property states.
 	 */
+	UPROPERTY(Transient)
 	bool bAllowPropertyStatesEdition = true;
 
 	/**
@@ -630,6 +644,7 @@ protected:
 	 *
 	 * Note: This context must be set by the owner instancing this pipeline. This context will be used to determine whether to hide some properties.
 	 */
+	UPROPERTY(Transient)
 	bool bIsReimportContext = false;
 
 	/**
@@ -638,26 +653,30 @@ protected:
 	 *
 	 * Note: This layout must be set by the owner instancing this pipeline. This layout will be use to hide or not some properties.
 	 */
+	UPROPERTY(Transient)
 	bool bIsShowEssentials = false;
 
 	/*
 	 * If true, this pipeline was create to re-import an asset or override the project settings pipelines.
 	 * That kind of pipeline will not be treat like project settings pipeline in the UI. PredialogCleanup will not be called.
 	 */
-	UPROPERTY()
+	UPROPERTY(Transient)
 	bool bFromReimportOrOverride = false;
 
-	UPROPERTY()
+	UPROPERTY(Transient)
 	TObjectPtr<UInterchangeResultsContainer> Results;
 
 	/**
 	 * Map of property path and lock status. Any properties that have a true lock status will be readonly when showing the import dialog.
 	 * Use the API to Get and Set the properties states.
 	*/
-	UPROPERTY()
+	UPROPERTY(Transient)
 	TMap<FName, FInterchangePipelinePropertyStates> PropertiesStates;
 
+	UPROPERTY(Transient)
 	mutable TMap<FName, FInterchangePipelinePropertyStates> CachePropertiesStates;
+
+	UPROPERTY(Transient)
 	mutable FInterchangePipelineContextParams CacheContextParam;
 
 	TArray<FInterchangeConflictInfo> ConflictInfos; 
