USD: Enable bSnapToClosestFrameBoundary by default on the USD generic asset pipelines.

#virtualized

diff --git a/Engine/Plugins/Interchange/Runtime/Source/Pipelines/Private/InterchangeGenericAnimationPipeline.cpp b/Engine/Plugins/Interchange/Runtime/Source/Pipelines/Private/InterchangeGenericAnimationPipeline.cpp
--- a/Engine/Plugins/Interchange/Runtime/Source/Pipelines/Private/InterchangeGenericAnimationPipeline.cpp
+++ b/Engine/Plugins/Interchange/Runtime/Source/Pipelines/Private/InterchangeGenericAnimationPipeline.cpp
@@ -747,7 +747,7 @@ void UInterchangeGenericAnimationPipeline::CreateAnimSequenceFactoryNode(UInterc
 				Message->SourceAssetName = SourceDatas[0]->GetFilename();
 				Message->DestinationAssetName = TrackNode.GetDisplayLabel();
 				Message->AssetType = UAnimSequence::StaticClass();
-				Message->Text = FText::Format(NSLOCTEXT("UInterchangeGenericAnimationPipeline", "WrongSequenceLength", "Animation length {0} is not compatible with import frame-rate {1} (sub frame {2}). The animation has to be frame-border aligned."),
+				Message->Text = FText::Format(NSLOCTEXT("UInterchangeGenericAnimationPipeline", "WrongSequenceLength", "Animation length {0} is not compatible with import frame-rate {1} (sub frame {2}). The animation has to be frame-border aligned if the 'Snap to Closest Frame Boundary' pipeline option is disabled."),
 					FText::AsNumber(SequenceLength), FrameRate.ToPrettyText(), FText::AsNumber(SubFrame));
 				//Skip this anim sequence factory node
 				return;
