Interchange: Allow forcing collision primitive generation even if it disagrees with provided mesh data.

#virtualized

diff --git a/Engine/Plugins/Interchange/Runtime/Source/FactoryNodes/Private/InterchangeStaticMeshLodDataNode.cpp b/Engine/Plugins/Interchange/Runtime/Source/FactoryNodes/Private/InterchangeStaticMeshLodDataNode.cpp
--- a/Engine/Plugins/Interchange/Runtime/Source/FactoryNodes/Private/InterchangeStaticMeshLodDataNode.cpp
+++ b/Engine/Plugins/Interchange/Runtime/Source/FactoryNodes/Private/InterchangeStaticMeshLodDataNode.cpp
@@ -332,6 +332,16 @@ bool UInterchangeStaticMeshLodDataNode::SetImportCollisionType(EInterchangeMeshC
 	IMPLEMENT_NODE_ATTRIBUTE_SETTER_NODELEGATE(ImportCollisionType, EInterchangeMeshCollision)
 }
 
+bool UInterchangeStaticMeshLodDataNode::GetForceCollisionPrimitiveGeneration(bool& AttributeValue) const
+{
+	IMPLEMENT_NODE_ATTRIBUTE_GETTER(ForceCollisionPrimitiveGeneration, bool);
+}
+
+bool UInterchangeStaticMeshLodDataNode::SetForceCollisionPrimitiveGeneration(bool AttributeValue)
+{
+	IMPLEMENT_NODE_ATTRIBUTE_SETTER_NODELEGATE(ForceCollisionPrimitiveGeneration, bool);
+}
+
 bool UInterchangeStaticMeshLodDataNode::GetImportCollision(bool& AttributeValue) const
 {
 	IMPLEMENT_NODE_ATTRIBUTE_GETTER(ImportCollision, bool);
diff --git a/Engine/Plugins/Interchange/Runtime/Source/FactoryNodes/Public/InterchangeStaticMeshLodDataNode.h b/Engine/Plugins/Interchange/Runtime/Source/FactoryNodes/Public/InterchangeStaticMeshLodDataNode.h
--- a/Engine/Plugins/Interchange/Runtime/Source/FactoryNodes/Public/InterchangeStaticMeshLodDataNode.h
+++ b/Engine/Plugins/Interchange/Runtime/Source/FactoryNodes/Public/InterchangeStaticMeshLodDataNode.h
@@ -136,6 +136,20 @@ public:
 	UFUNCTION(BlueprintCallable, Category = "Interchange | Node | StaticMeshLodData")
 	UE_API bool SetImportCollisionType(EInterchangeMeshCollision AttributeValue);
 
+	/** 
+	 * Gets whether we're generating collision primitive shapes even if the mesh data 
+	 * doesn't match the desired shape very well
+	 */
+	 UFUNCTION(BlueprintCallable, Category = "Interchange | Node | StaticMeshLodData")
+	 UE_API bool GetForceCollisionPrimitiveGeneration(bool& bGenerate) const;
+	 
+	 /** 
+	  * Sets whether we're generating collision primitive shapes even if the mesh data 
+	  * doesn't match the desired shape very well
+	  */
+	UFUNCTION(BlueprintCallable, Category = "Interchange | Node | StaticMeshLodData")
+	UE_API bool SetForceCollisionPrimitiveGeneration(bool bGenerate);
+
 	UFUNCTION(BlueprintCallable, Category = "Interchange | Node | StaticMeshLodData")
 	UE_API bool GetImportCollision(bool& AttributeValue) const;
 
@@ -154,6 +168,7 @@ private:
 	UE::Interchange::TArrayAttributeHelper<FString> ConvexCollisionMeshUids;
 		
 	IMPLEMENT_NODE_ATTRIBUTE_KEY(OneConvexHullPerUCX)
+	IMPLEMENT_NODE_ATTRIBUTE_KEY(ForceCollisionPrimitiveGeneration)
 	IMPLEMENT_NODE_ATTRIBUTE_KEY(ImportCollision)
 	IMPLEMENT_NODE_ATTRIBUTE_KEY(ImportCollisionType)
 };
diff --git a/Engine/Plugins/Interchange/Runtime/Source/Import/Private/Mesh/InterchangeMeshHelper.cpp b/Engine/Plugins/Interchange/Runtime/Source/Import/Private/Mesh/InterchangeMeshHelper.cpp
--- a/Engine/Plugins/Interchange/Runtime/Source/Import/Private/Mesh/InterchangeMeshHelper.cpp
+++ b/Engine/Plugins/Interchange/Runtime/Source/Import/Private/Mesh/InterchangeMeshHelper.cpp
@@ -1,31 +1,829 @@
-// Copyright Epic Games, Inc. All Rights Reserved. 
+// Copyright Epic Games, Inc. All Rights Reserved.
+
 #include "Mesh/InterchangeMeshHelper.h"
 
+#if WITH_EDITOR
+#include "BSPOps.h"
+#endif
+#include "Engine/Polys.h"
 #include "Engine/SkeletalMesh.h"
 #include "Engine/SkinnedAssetCommon.h"
 #include "Engine/StaticMesh.h"
+#include "Engine/StaticMeshSocket.h"
 #include "GenericPlatform/GenericPlatformMisc.h"
-#include "InterchangeHelper.h"
-#include "InterchangeImportLog.h"
-#include "InterchangeMaterialFactoryNode.h"
-#include "InterchangeSceneNode.h"
 #include "MaterialDomain.h"
 #include "Materials/Material.h"
 #include "Materials/MaterialInterface.h"
 #include "Math/GenericOctree.h"
+#include "Mesh/InterchangeMeshPayload.h"
 #include "MeshDescription.h"
 #include "MeshUVChannelInfo.h"
-#include "Mesh/InterchangeMeshPayload.h"
+#include "Model.h"
 #include "Nodes/InterchangeBaseNodeContainer.h"
+#include "PhysicsEngine/AggregateGeom.h"
+#include "PhysicsEngine/BodySetup.h"
+#include "PhysicsEngine/ConvexElem.h"
 #include "Rendering/SkeletalMeshModel.h"
 #include "SkeletalMeshAttributes.h"
 #include "SkeletalMeshTypes.h"
 #include "StaticMeshAttributes.h"
 #include "StaticMeshOperations.h"
 
+#include "InterchangeCommonPipelineDataFactoryNode.h"
+#include "InterchangeHelper.h"
+#include "InterchangeImportLog.h"
+#include "InterchangeMaterialFactoryNode.h"
+#include "InterchangeMeshNode.h"
+#include "InterchangeSceneNode.h"
+#include "InterchangeStaticMeshFactoryNode.h"
+#include "InterchangeStaticMeshLodDataNode.h"
+#include "Nodes/InterchangeSourceNode.h"
 
 namespace UE::Interchange::Private::MeshHelper
 {
+	bool AddConvexGeomFromVertices(
+		const UInterchangeFactoryBase::FImportAssetObjectParams& Arguments,
+		const FMeshDescription& MeshDescription,
+		FKAggregateGeom& AggGeom
+	)
+	{
+		FStaticMeshConstAttributes Attributes(MeshDescription);
+		TVertexAttributesRef<const FVector3f> VertexPositions = Attributes.GetVertexPositions();
+
+		if (VertexPositions.GetNumElements() == 0)
+		{
+			return false;
+		}
+
+		FKConvexElem& ConvexElem = AggGeom.ConvexElems.Emplace_GetRef();
+		ConvexElem.VertexData.AddUninitialized(VertexPositions.GetNumElements());
+
+		for (int32 Index = 0; Index < VertexPositions.GetNumElements(); Index++)
+		{
+			ConvexElem.VertexData[Index] = FVector(VertexPositions[Index]);
+		}
+
+		ConvexElem.UpdateElemBox();
+
+		return true;
+	}
+
+	bool DecomposeConvexMesh(
+		const UInterchangeFactoryBase::FImportAssetObjectParams& Arguments,
+		const FMeshDescription& MeshDescription,
+		UBodySetup* BodySetup
+	)
+	{
+#if WITH_EDITOR
+
+		// Construct a bit array containing a bit for each triangle ID in the mesh description.
+		// We are assuming the mesh description is compact, i.e. it has no holes, and so the number of triangles is equal to the array size.
+		// The aim is to identify 'islands' of adjacent triangles which will form separate convex hulls
+
+		check(MeshDescription.Triangles().Num() == MeshDescription.Triangles().GetArraySize());
+		TBitArray<> BitArray(0, MeshDescription.Triangles().Num());
+
+		// Here we build the groups of triangle IDs
+
+		TArray<TArray<FTriangleID>> TriangleGroups;
+
+		int32 FirstIndex = BitArray.FindAndSetFirstZeroBit();
+		while (FirstIndex != INDEX_NONE)
+		{
+			// Find the first index we haven't used yet, and use it as the beginning of a new triangle group
+
+			TArray<FTriangleID>& TriangleGroup = TriangleGroups.Emplace_GetRef();
+			TriangleGroup.Emplace(FirstIndex);
+
+			// Now iterate through the TriangleGroup array, finding unused adjacent triangles to each index, and appending them
+			// to the end of the array.  Note we deliberately check the array size each time round the loop, as each iteration
+			// can cause it to grow.
+
+			for (int32 CheckIndex = 0; CheckIndex < TriangleGroup.Num(); ++CheckIndex)
+			{
+				for (FTriangleID AdjacentTriangleID : MeshDescription.GetTriangleAdjacentTriangles(TriangleGroup[CheckIndex]))
+				{
+					if (BitArray[AdjacentTriangleID] == 0)
+					{
+						// Append unused adjacent triangles to the TriangleGroup, to be considered for adjacency later
+						TriangleGroup.Emplace(AdjacentTriangleID);
+						BitArray[AdjacentTriangleID] = 1;
+					}
+				}
+			}
+
+			// When we exhaust the triangle group array, there are no more triangles in this island.
+			// Now find the start of the next group.
+
+			FirstIndex = BitArray.FindAndSetFirstZeroBit();
+		}
+
+		// Now iterate through the triangle groups, adding each as a convex hull to the AggGeom
+
+		UModel* TempModel = NewObject<UModel>();
+		TempModel->RootOutside = true;
+		TempModel->EmptyModel(true, true);
+		TempModel->Polys->ClearFlags(RF_Transactional);
+
+		FStaticMeshConstAttributes Attributes(MeshDescription);
+		TTriangleAttributesRef<TArrayView<const FVertexID>> TriangleVertices = Attributes.GetTriangleVertexIndices();
+		TVertexAttributesRef<const FVector3f> VertexPositions = Attributes.GetVertexPositions();
+
+		bool bSuccess = true;
+
+		for (const TArray<FTriangleID>& TriangleGroup : TriangleGroups)
+		{
+			// Initialize a new brush
+			TempModel->Polys->Element.Empty();
+
+			// Add each triangle to the brush
+			int32 Index = 0;
+			for (FTriangleID TriangleID : TriangleGroup)
+			{
+				FPoly& Poly = TempModel->Polys->Element.Emplace_GetRef();
+				Poly.Init();
+				Poly.iLink = Index++;
+
+				// For reasons lost in time, BSP poly vertices have the opposite winding order to regular mesh vertices.
+				// So add them backwards (sigh)
+				Poly.Vertices.Emplace(FVector(VertexPositions[TriangleVertices[TriangleID][2]]));
+				Poly.Vertices.Emplace(FVector(VertexPositions[TriangleVertices[TriangleID][1]]));
+				Poly.Vertices.Emplace(FVector(VertexPositions[TriangleVertices[TriangleID][0]]));
+
+				Poly.CalcNormal(true);
+			}
+
+			// Build bounding box
+			TempModel->BuildBound();
+
+			// Build BSP for the brush
+			FBSPOps::bspBuild(TempModel, FBSPOps::BSP_Good, 15, 70, 1, 0);
+			FBSPOps::bspRefresh(TempModel, true);
+			FBSPOps::bspBuildBounds(TempModel);
+
+			bSuccess &= BodySetup->CreateFromModel(TempModel, false);
+		}
+
+		TempModel->ClearInternalFlags(EInternalObjectFlags::Async);
+		TempModel->Polys->ClearInternalFlags(EInternalObjectFlags::Async);
+
+		return bSuccess;
+
+#else	 // #if WITH_EDITOR
+
+		return false;
+
+#endif
+	}
+
+	static bool AreEqual(float A, float B)
+	{
+		constexpr float MeshToPrimTolerance = 0.001f;
+		return FMath::Abs(A - B) < MeshToPrimTolerance;
+	}
+
+	static bool AreParallel(const FVector3f& A, const FVector3f& B)
+	{
+		float Dot = FVector3f::DotProduct(A, B);
+
+		return (AreEqual(FMath::Abs(Dot), 1.0f));
+	}
+
+	static FVector3f GetTriangleNormal(
+		TVertexAttributesRef<const FVector3f> VertexPositions,
+		TArrayView<const FVertexID> VertexIndices
+	)
+	{
+		const FVector3f& V0 = VertexPositions[VertexIndices[0]];
+		const FVector3f& V1 = VertexPositions[VertexIndices[1]];
+		const FVector3f& V2 = VertexPositions[VertexIndices[2]];
+		// @todo: LWC conversions everywhere here; surely this can be more elegant?
+		return FVector3f(FVector(FVector3f::CrossProduct(V1 - V0, V2 - V0).GetSafeNormal()));
+	}
+
+	bool AddBoxGeomFromTris(
+		const FMeshDescription& MeshDescription,
+		FKAggregateGeom& AggGeom,
+		bool bForcePrimitiveGeneration
+	)
+	{
+		FStaticMeshConstAttributes Attributes{MeshDescription};
+		TTriangleAttributesRef<TArrayView<const FVertexID>> TriangleVertices = Attributes.GetTriangleVertexIndices();
+		TVertexAttributesRef<const FVector3f> VertexPositions = Attributes.GetVertexPositions();
+		FBox Box{ForceInit};
+
+		// Maintain an array of the planes we have encountered so far.
+		// We are expecting two instances of three unique plane orientations, one for each side of the box.
+
+		struct FPlaneInfo
+		{
+			FPlaneInfo(const FVector3f InNormal, float InFirstDistance)
+				: Normal(InNormal)
+				, DistCount(1)
+				, PlaneDist{InFirstDistance, 0.0f}
+			{
+			}
+
+			FVector3f Normal;
+			int32 DistCount;
+			float PlaneDist[2];
+		};
+
+		TArray<FPlaneInfo> Planes;		
+
+		for (FTriangleID TriangleID : MeshDescription.Triangles().GetElementIDs())
+		{
+			TArrayView<const FVertexID> VertexIndices = TriangleVertices[TriangleID];
+
+			// Maintain an AABB, adding points from each triangle.
+			Box += FVector(VertexPositions[VertexIndices[0]]);
+			Box += FVector(VertexPositions[VertexIndices[1]]);
+			Box += FVector(VertexPositions[VertexIndices[2]]);
+
+			FVector3f TriangleNormal = GetTriangleNormal(VertexPositions, VertexIndices);
+			if (TriangleNormal.IsNearlyZero())
+			{
+				continue;
+			}
+
+			bool bFoundPlane = false;
+			bool bFailedPlanes = false;
+			for (int32 PlaneIndex = 0; PlaneIndex < Planes.Num() && !bFoundPlane; PlaneIndex++)
+			{
+				// if this triangle plane is already known...
+				if (AreParallel(TriangleNormal, Planes[PlaneIndex].Normal))
+				{
+					// Always use the same normal when comparing distances, to ensure consistent sign.
+					float Dist = FVector3f::DotProduct(VertexPositions[VertexIndices[0]], Planes[PlaneIndex].Normal);
+
+					// we only have one distance, and its not that one, add it.
+					if (Planes[PlaneIndex].DistCount == 1 && !AreEqual(Dist, Planes[PlaneIndex].PlaneDist[0]))
+					{
+						Planes[PlaneIndex].PlaneDist[1] = Dist;
+						Planes[PlaneIndex].DistCount = 2;
+					}
+					// if we have a second distance, and its not that either, something is wrong.
+					else if (Planes[PlaneIndex].DistCount == 2 
+						&& !AreEqual(Dist, Planes[PlaneIndex].PlaneDist[0])
+						&& !AreEqual(Dist, Planes[PlaneIndex].PlaneDist[1])
+					)
+					{
+						bFailedPlanes = true;
+						break;
+					}
+
+					bFoundPlane = true;
+				}
+			}
+
+			// If this triangle does not match an existing plane, add to list.
+			if (!bFailedPlanes && !bFoundPlane)
+			{
+				Planes.Emplace(TriangleNormal, FVector3f::DotProduct(VertexPositions[VertexIndices[0]], TriangleNormal));
+			}
+		}
+
+		// Now we have our candidate planes, see if we can match all the requirements
+
+		// Check for the right number of planes.
+		if (Planes.Num() == 3)
+		{
+			// Check if we have the 3 pairs we need
+			if ((Planes[0].DistCount == 2) && (Planes[1].DistCount == 2) && (Planes[2].DistCount == 2))
+			{
+				// ensure valid TM by cross-product
+				if (AreParallel(FVector3f::CrossProduct(Planes[0].Normal, Planes[1].Normal), Planes[2].Normal))
+				{
+					// Allocate box in array
+					FKBoxElem BoxElem;
+			
+					// In case we have a box oriented with the world axis system we want to reorder the plane to not introduce axis swap.
+					// If the box was turned, the order of the planes will be arbitrary and the box rotation will make the collision
+					// not playing well if the asset is built or place in a level with a non uniform scale.
+					FVector3f Axis[3] = {FVector3f::XAxisVector, FVector3f::YAxisVector, FVector3f::ZAxisVector};
+					int32 Reorder[3] = {INDEX_NONE, INDEX_NONE, INDEX_NONE};
+					for (int32 PlaneIndex = 0; PlaneIndex < 3; ++PlaneIndex)
+					{
+						for (int32 AxisIndex = 0; AxisIndex < 3; ++AxisIndex)
+						{
+							if (AreParallel(Planes[PlaneIndex].Normal, Axis[AxisIndex]))
+							{
+								Reorder[PlaneIndex] = AxisIndex;
+								break;
+							}
+						}
+					}
+			
+					if (Reorder[0] == INDEX_NONE || Reorder[1] == INDEX_NONE || Reorder[2] == INDEX_NONE)
+					{
+						Reorder[0] = 0;
+						Reorder[1] = 1;
+						Reorder[2] = 2;
+					}
+			
+					BoxElem.SetTransform(
+						FTransform(FVector(Planes[Reorder[0]].Normal), FVector(Planes[Reorder[1]].Normal), FVector(Planes[Reorder[2]].Normal), Box.GetCenter())
+					);
+			
+					// distance between parallel planes is box edge lengths.
+					BoxElem.X = FMath::Abs(Planes[Reorder[0]].PlaneDist[0] - Planes[Reorder[0]].PlaneDist[1]);
+					BoxElem.Y = FMath::Abs(Planes[Reorder[1]].PlaneDist[0] - Planes[Reorder[1]].PlaneDist[1]);
+					BoxElem.Z = FMath::Abs(Planes[Reorder[2]].PlaneDist[0] - Planes[Reorder[2]].PlaneDist[1]);
+			
+					AggGeom.BoxElems.Add(BoxElem);
+					
+					return true;
+				}
+			}
+		}
+
+		// We couldn't produce the box we desired, so fallback to a simple AABB if we must produce anything
+		if (bForcePrimitiveGeneration && Box.IsValid)
+		{
+			FKBoxElem BoxElem;			
+			BoxElem.Center = Box.GetCenter();
+
+			FVector Extents = Box.GetExtent();
+			BoxElem.X = 2.0 * Extents.X;
+			BoxElem.Y = 2.0 * Extents.Y;
+			BoxElem.Z = 2.0 * Extents.Z;
+
+			AggGeom.BoxElems.Add(BoxElem);
+			return true;
+		}
+
+		return false;
+	}
+
+	bool AddSphereGeomFromVertices(
+		const UInterchangeFactoryBase::FImportAssetObjectParams& Arguments,
+		const FMeshDescription& MeshDescription,
+		FKAggregateGeom& AggGeom,
+		bool bForcePrimitiveGeneration
+	)
+	{
+		FStaticMeshConstAttributes Attributes(MeshDescription);
+		TVertexAttributesRef<const FVector3f> VertexPositions = Attributes.GetVertexPositions();
+
+		if (VertexPositions.GetNumElements() == 0)
+		{
+			return false;
+		}
+
+		FBox Box(ForceInit);
+
+		for (const FVector3f& VertexPosition : VertexPositions.GetRawArray())
+		{
+			Box += FVector(VertexPosition);
+		}
+
+		FVector Center;
+		FVector Extents;
+		Box.GetCenterAndExtents(Center, Extents);
+		float Longest = 2.0f * Extents.GetMax();
+		float Radius = 0.5f * Longest;
+		
+		// Validation
+		if (!bForcePrimitiveGeneration)
+		{
+			float Shortest = 2.0f * Extents.GetMin();
+
+			// check that the AABB is roughly a square (5% tolerance)
+			if ((Longest - Shortest) / Longest > 0.05f)
+			{
+				return false;
+			}
+
+			// Test that all vertices are a similar radius (5%) from the sphere centre.
+			float MaxR = 0;
+			float MinR = BIG_NUMBER;
+
+			for (const FVector3f& VertexPosition : VertexPositions.GetRawArray())
+			{
+				FVector CToV = FVector(VertexPosition) - Center;
+				float RSqr = CToV.SizeSquared();
+
+				MaxR = FMath::Max(RSqr, MaxR);
+
+				// Sometimes vertex at centre, so reject it.
+				if (RSqr > KINDA_SMALL_NUMBER)
+				{
+					MinR = FMath::Min(RSqr, MinR);
+				}
+			}
+
+			MaxR = FMath::Sqrt(MaxR);
+			MinR = FMath::Sqrt(MinR);
+
+			if ((MaxR - MinR) / Radius > 0.05f)
+			{
+				return false;
+			}
+		}
+
+		// Allocate sphere in array
+		FKSphereElem SphereElem;
+		SphereElem.Center = Center;
+		SphereElem.Radius = Radius;
+		AggGeom.SphereElems.Add(SphereElem);
+
+		return true;
+	}
+
+	bool AddCapsuleGeomFromVertices(
+		const UInterchangeFactoryBase::FImportAssetObjectParams& Arguments,
+		const FMeshDescription& MeshDescription,
+		FKAggregateGeom& AggGeom
+	)
+	{
+		FStaticMeshConstAttributes Attributes(MeshDescription);
+		TVertexAttributesRef<const FVector3f> VertexPositions = Attributes.GetVertexPositions();
+
+		if (VertexPositions.GetNumElements() == 0)
+		{
+			return false;
+		}
+
+		FVector AxisStart = FVector::ZeroVector;
+		FVector AxisEnd = FVector::ZeroVector;
+		float MaxDistSqr = 0.f;
+
+		for (int32 IndexA = 0; IndexA < VertexPositions.GetNumElements() - 1; IndexA++)
+		{
+			for (int32 IndexB = IndexA + 1; IndexB < VertexPositions.GetNumElements(); IndexB++)
+			{
+				FVector TransformedA = FVector(VertexPositions[IndexA]);
+				FVector TransformedB = FVector(VertexPositions[IndexB]);
+
+				float DistSqr = (TransformedA - TransformedB).SizeSquared();
+				if (DistSqr > MaxDistSqr)
+				{
+					AxisStart = TransformedA;
+					AxisEnd = TransformedB;
+					MaxDistSqr = DistSqr;
+				}
+			}
+		}
+
+		// If we got a valid axis, find vertex furthest from it
+		if (MaxDistSqr > SMALL_NUMBER)
+		{
+			float MaxRadius = 0.0f;
+
+			const FVector LineOrigin = AxisStart;
+			const FVector LineDir = (AxisEnd - AxisStart).GetSafeNormal();
+
+			for (int32 IndexA = 0; IndexA < VertexPositions.GetNumElements(); IndexA++)
+			{
+				FVector TransformedA = FVector(VertexPositions[IndexA]);
+
+				float DistToAxis = FMath::PointDistToLine(TransformedA, LineDir, LineOrigin);
+				if (DistToAxis > MaxRadius)
+				{
+					MaxRadius = DistToAxis;
+				}
+			}
+
+			if (MaxRadius > SMALL_NUMBER)
+			{
+				// Allocate capsule in array
+				FKSphylElem SphylElem;
+				SphylElem.Center = 0.5f * (AxisStart + AxisEnd);
+				SphylElem.Rotation = FQuat::FindBetweenVectors(FVector::ZAxisVector, LineDir).Rotator();	// Get quat that takes you from z axis to
+																											// desired axis
+				SphylElem.Radius = MaxRadius;
+				SphylElem.Length = FMath::Max(FMath::Sqrt(MaxDistSqr) - (2.0f * MaxRadius), 0.0f);	  // subtract two radii from total length to get
+																									  // segment length (ensure > 0)
+				AggGeom.SphylElems.Add(SphylElem);
+				return true;
+			}
+		}
+
+		return false;
+	}
+
+	bool ImportBoxCollision(
+		const UInterchangeFactoryBase::FImportAssetObjectParams& Arguments,
+		const TMap<FInterchangeMeshPayLoadKey, FMeshPayload>& BoxCollisionPayloads,
+		UStaticMesh* StaticMesh,
+		bool bForcePrimitiveGeneration
+	)
+	{
+		using namespace UE::Interchange;
+
+		bool bResult = false;
+
+		FKAggregateGeom& AggGeo = StaticMesh->GetBodySetup()->AggGeom;
+
+		for (const TPair<FInterchangeMeshPayLoadKey, FMeshPayload>& BoxCollisionPayload : BoxCollisionPayloads)
+		{
+			const TOptional<FMeshPayloadData>& PayloadData = BoxCollisionPayload.Value.PayloadData;
+			if (!PayloadData.IsSet())
+			{
+				continue;
+			}
+
+			if (AddBoxGeomFromTris(PayloadData->MeshDescription, AggGeo, bForcePrimitiveGeneration))
+			{
+				bResult = true;
+				FKBoxElem& NewElem = AggGeo.BoxElems.Last();
+
+				// Now test the last element in the AggGeo list and remove it if its a duplicate
+				// @TODO: determine why we have to do this. Was it to prevent duplicate boxes accumulating when reimporting?
+				for (int32 ElementIndex = 0; ElementIndex < AggGeo.BoxElems.Num() - 1; ++ElementIndex)
+				{
+					FKBoxElem& CurrentElem = AggGeo.BoxElems[ElementIndex];
+
+					if (CurrentElem == NewElem)
+					{
+						// The new element is a duplicate, remove it
+						AggGeo.BoxElems.RemoveAt(AggGeo.BoxElems.Num() - 1);
+						break;
+					}
+				}
+			}
+		}
+
+		return bResult;
+	}
+
+	bool ImportCapsuleCollision(
+		const UInterchangeFactoryBase::FImportAssetObjectParams& Arguments,
+		const TMap<FInterchangeMeshPayLoadKey, FMeshPayload>& CapsuleCollisionPayloads,
+		UStaticMesh* StaticMesh
+	)
+	{
+		using namespace UE::Interchange;
+
+		bool bResult = false;
+
+		FKAggregateGeom& AggGeo = StaticMesh->GetBodySetup()->AggGeom;
+
+		for (const TPair<FInterchangeMeshPayLoadKey, FMeshPayload>& CapsuleCollisionPayload : CapsuleCollisionPayloads)
+		{
+			const TOptional<FMeshPayloadData>& PayloadData = CapsuleCollisionPayload.Value.PayloadData;
+
+			if (!PayloadData.IsSet())
+			{
+				continue;
+			}
+
+			if (AddCapsuleGeomFromVertices(Arguments, PayloadData->MeshDescription, AggGeo))
+			{
+				bResult = true;
+
+				FKSphylElem& NewElem = AggGeo.SphylElems.Last();
+
+				// Now test the late element in the AggGeo list and remove it if its a duplicate
+				// @TODO: determine why we have to do this. Was it to prevent duplicate boxes accumulating when reimporting?
+				for (int32 ElementIndex = 0; ElementIndex < AggGeo.SphylElems.Num() - 1; ++ElementIndex)
+				{
+					FKSphylElem& CurrentElem = AggGeo.SphylElems[ElementIndex];
+					if (CurrentElem == NewElem)
+					{
+						// The new element is a duplicate, remove it
+						AggGeo.SphylElems.RemoveAt(AggGeo.SphylElems.Num() - 1);
+						break;
+					}
+				}
+			}
+		}
+
+		return bResult;
+	}
+
+	bool ImportSphereCollision(
+		const UInterchangeFactoryBase::FImportAssetObjectParams& Arguments,
+		const TMap<FInterchangeMeshPayLoadKey, FMeshPayload>& SphereCollisionPayloads,
+		UStaticMesh* StaticMesh,
+		bool bForcePrimitiveGeneration
+	)
+	{
+		using namespace UE::Interchange;
+
+		bool bResult = false;
+
+		FKAggregateGeom& AggGeo = StaticMesh->GetBodySetup()->AggGeom;
+
+		for (const TPair<FInterchangeMeshPayLoadKey, FMeshPayload>& SphereCollisionPayload : SphereCollisionPayloads)
+		{
+			const TOptional<FMeshPayloadData>& PayloadData = SphereCollisionPayload.Value.PayloadData;
+
+			if (!PayloadData.IsSet())
+			{
+				continue;
+			}
+
+			if (AddSphereGeomFromVertices(Arguments, PayloadData->MeshDescription, AggGeo, bForcePrimitiveGeneration))
+			{
+				bResult = true;
+
+				FKSphereElem& NewElem = AggGeo.SphereElems.Last();
+
+				// Now test the last element in the AggGeo list and remove it if its a duplicate
+				// @TODO: determine why we have to do this. Was it to prevent duplicate boxes accumulating when reimporting?
+				for (int32 ElementIndex = 0; ElementIndex < AggGeo.SphereElems.Num() - 1; ++ElementIndex)
+				{
+					FKSphereElem& CurrentElem = AggGeo.SphereElems[ElementIndex];
+
+					if (CurrentElem == NewElem)
+					{
+						// The new element is a duplicate, remove it
+						AggGeo.SphereElems.RemoveAt(AggGeo.SphereElems.Num() - 1);
+						break;
+					}
+				}
+			}
+		}
+
+		return bResult;
+	}
+
+	bool ImportConvexCollision(
+		const UInterchangeFactoryBase::FImportAssetObjectParams& Arguments,
+		const TMap<FInterchangeMeshPayLoadKey, FMeshPayload>& ConvexCollisionPayloads,
+		UStaticMesh* StaticMesh,
+		const UInterchangeStaticMeshLodDataNode* LodDataNode
+	)
+	{
+		using namespace UE::Interchange;
+
+		bool bResult = false;
+
+		bool bOneConvexHullPerUCX;
+		if (!LodDataNode->GetOneConvexHullPerUCX(bOneConvexHullPerUCX) || !bOneConvexHullPerUCX)
+		{
+			for (const TPair<FInterchangeMeshPayLoadKey, FMeshPayload>& ConvexCollisionPayload : ConvexCollisionPayloads)
+			{
+				const TOptional<FMeshPayloadData>& PayloadData = ConvexCollisionPayload.Value.PayloadData;
+
+				if (!PayloadData.IsSet())
+				{
+					continue;
+				}
+
+				if (DecomposeConvexMesh(Arguments, PayloadData->MeshDescription, StaticMesh->GetBodySetup()))
+				{
+					bResult = true;
+				}
+			}
+		}
+		else
+		{
+			FKAggregateGeom& AggGeo = StaticMesh->GetBodySetup()->AggGeom;
+
+			for (const TPair<FInterchangeMeshPayLoadKey, FMeshPayload>& ConvexCollisionPayload : ConvexCollisionPayloads)
+			{
+				TOptional<FMeshPayloadData> PayloadData = ConvexCollisionPayload.Value.PayloadData;
+
+				if (!PayloadData.IsSet())
+				{
+					continue;
+				}
+
+				if (AddConvexGeomFromVertices(Arguments, PayloadData->MeshDescription, AggGeo))
+				{
+					bResult = true;
+
+					FKConvexElem& NewElem = AggGeo.ConvexElems.Last();
+
+					// Now test the late element in the AggGeo list and remove it if its a duplicate
+					// @TODO: determine why the importer used to do this. Was it something to do with reimport not adding extra collision or
+					// something?
+					for (int32 ElementIndex = 0; ElementIndex < AggGeo.ConvexElems.Num() - 1; ++ElementIndex)
+					{
+						FKConvexElem& CurrentElem = AggGeo.ConvexElems[ElementIndex];
+
+						if (CurrentElem.VertexData.Num() == NewElem.VertexData.Num())
+						{
+							bool bFoundDifference = false;
+							for (int32 VertexIndex = 0; VertexIndex < NewElem.VertexData.Num(); ++VertexIndex)
+							{
+								if (CurrentElem.VertexData[VertexIndex] != NewElem.VertexData[VertexIndex])
+								{
+									bFoundDifference = true;
+									break;
+								}
+							}
+
+							if (!bFoundDifference)
+							{
+								// The new collision geo is a duplicate, delete it
+								AggGeo.ConvexElems.RemoveAt(AggGeo.ConvexElems.Num() - 1);
+								break;
+							}
+						}
+					}
+				}
+			}
+		}
+
+		return bResult;
+	}
+
+	bool ImportSockets(
+		const UInterchangeFactoryBase::FImportAssetObjectParams& Arguments,
+		UStaticMesh* StaticMesh,
+		const UInterchangeStaticMeshFactoryNode* FactoryNode
+	)
+	{
+		bool bImportSockets = false;
+		FactoryNode->GetCustomImportSockets(bImportSockets);
+		if (!bImportSockets)
+		{
+			// Skip Import sockets
+			return true;
+		}
+
+		TArray<FString> SocketUids;
+		FactoryNode->GetSocketUids(SocketUids);
+
+		TSet<FName> ImportedSocketNames;
+
+		FTransform GlobalOffsetTransform = FTransform::Identity;
+
+		bool bBakeMeshes = false;
+		bool bBakePivotMeshes = false;
+		if (UInterchangeCommonPipelineDataFactoryNode* CommonPipelineDataFactoryNode = UInterchangeCommonPipelineDataFactoryNode::GetUniqueInstance(
+				Arguments.NodeContainer
+			))
+		{
+			CommonPipelineDataFactoryNode->GetCustomGlobalOffsetTransform(GlobalOffsetTransform);
+			CommonPipelineDataFactoryNode->GetBakeMeshes(bBakeMeshes);
+			if (!bBakeMeshes)
+			{
+				CommonPipelineDataFactoryNode->GetBakePivotMeshes(bBakePivotMeshes);
+			}
+		}
+
+		for (const FString& SocketUid : SocketUids)
+		{
+			if (const UInterchangeSceneNode* SceneNode = Cast<UInterchangeSceneNode>(Arguments.NodeContainer->GetNode(SocketUid)))
+			{
+				FString NodeDisplayName = SceneNode->GetDisplayLabel();
+				if (NodeDisplayName.StartsWith(UInterchangeMeshFactoryNode::GetMeshSocketPrefix()))
+				{
+					NodeDisplayName.RightChopInline(UInterchangeMeshFactoryNode::GetMeshSocketPrefix().Len(), EAllowShrinking::No);
+				}
+				FName SocketName = FName(NodeDisplayName);
+				ImportedSocketNames.Add(SocketName);
+
+				FTransform Transform;
+				if (bBakeMeshes)
+				{
+					SceneNode->GetCustomGlobalTransform(Arguments.NodeContainer, GlobalOffsetTransform, Transform);
+				}
+
+				UE::Interchange::Private::MeshHelper::AddSceneNodeGeometricAndPivotToGlobalTransform(
+					Transform,
+					SceneNode,
+					bBakeMeshes,
+					bBakePivotMeshes
+				);
+
+				// Apply axis transformation inverse to get correct Socket Transform:
+				const UInterchangeSourceNode* SourceNode = UInterchangeSourceNode::GetUniqueInstance(Arguments.NodeContainer);
+				FTransform AxisConversionInverseTransform;
+				if (SourceNode->GetCustomAxisConversionInverseTransform(AxisConversionInverseTransform))
+				{
+					Transform = AxisConversionInverseTransform * Transform;
+				}
+
+				UStaticMeshSocket* Socket = StaticMesh->FindSocket(SocketName);
+				if (!Socket)
+				{
+					// If the socket didn't exist create a new one now
+					Socket = NewObject<UStaticMeshSocket>(StaticMesh);
+#if WITH_EDITORONLY_DATA
+					Socket->bSocketCreatedAtImport = true;
+#endif
+					Socket->SocketName = SocketName;
+					StaticMesh->AddSocket(Socket);
+				}
+
+				Socket->RelativeLocation = Transform.GetLocation();
+				Socket->RelativeRotation = Transform.GetRotation().Rotator();
+				Socket->RelativeScale = Transform.GetScale3D();
+			}
+		}
+
+		// Delete any sockets which were previously imported but which no longer exist in the imported scene
+		for (TArray<TObjectPtr<UStaticMeshSocket>>::TIterator It = StaticMesh->Sockets.CreateIterator(); It; ++It)
+		{
+			UStaticMeshSocket* Socket = *It;
+			if (
+#if WITH_EDITORONLY_DATA
+			Socket->bSocketCreatedAtImport &&
+#endif
+			!ImportedSocketNames.Contains(Socket->SocketName))
+			{
+				It.RemoveCurrent();
+			}
+		}
+
+		return true;
+	}
+
 	void RemapPolygonGroups(const FMeshDescription& SourceMesh, FMeshDescription& TargetMesh, PolygonGroupMap& RemapPolygonGroup)
 	{
 		FStaticMeshConstAttributes SourceAttributes(SourceMesh);
diff --git a/Engine/Plugins/Interchange/Runtime/Source/Import/Private/Mesh/InterchangeStaticMeshFactory.cpp b/Engine/Plugins/Interchange/Runtime/Source/Import/Private/Mesh/InterchangeStaticMeshFactory.cpp
--- a/Engine/Plugins/Interchange/Runtime/Source/Import/Private/Mesh/InterchangeStaticMeshFactory.cpp
+++ b/Engine/Plugins/Interchange/Runtime/Source/Import/Private/Mesh/InterchangeStaticMeshFactory.cpp
@@ -3,34 +3,30 @@
 #include "Mesh/InterchangeStaticMeshFactory.h"
 
 #if WITH_EDITOR
-#include "BSPOps.h"
 #include "Editor/UnrealEd/Private/GeomFitUtils.h"
 #endif
 #include "Components.h"
-#include "Engine/Polys.h"
 #include "Engine/StaticMesh.h"
-#include "Engine/StaticMeshSocket.h"
 #include "InterchangeAssetImportData.h"
 #include "InterchangeCommonPipelineDataFactoryNode.h"
 #include "InterchangeImportCommon.h"
 #include "InterchangeImportLog.h"
-#include "InterchangeMaterialFactoryNode.h"
+#include "InterchangeManager.h"
 #include "InterchangeMeshNode.h"
 #include "InterchangeSceneNode.h"
-#include "InterchangeStaticMeshLodDataNode.h"
-#include "InterchangeStaticMeshFactoryNode.h"
 #include "InterchangeSourceData.h"
+#include "InterchangeStaticMeshFactoryNode.h"
+#include "InterchangeStaticMeshLodDataNode.h"
 #include "InterchangeTranslatorBase.h"
 #include "MaterialDomain.h"
 #include "Materials/Material.h"
-#include "Materials/MaterialInstance.h"
 #include "Materials/MaterialInterface.h"
 #include "Mesh/InterchangeMeshHelper.h"
 #include "Mesh/InterchangeMeshPayloadInterface.h"
 #include "MeshBudgetProjectSettings.h"
-#include "Model.h"
 #include "Nodes/InterchangeBaseNode.h"
 #include "Nodes/InterchangeBaseNodeContainer.h"
+#include "Nodes/InterchangeSourceNode.h"
 #include "PhysicsEngine/BodySetup.h"
 #include "PhysicsEngine/BoxElem.h"
 #include "PhysicsEngine/ConvexElem.h"
@@ -39,8 +35,7 @@
 #include "StaticMeshAttributes.h"
 #include "StaticMeshCompiler.h"
 #include "StaticMeshOperations.h"
-#include "InterchangeManager.h"
-#include "Nodes/InterchangeSourceNode.h"
+#include "StaticMeshResources.h"
 
 #include UE_INLINE_GENERATED_CPP_BY_NAME(InterchangeStaticMeshFactory)
 
@@ -109,7 +104,7 @@ void UInterchangeStaticMeshFactory::CreatePayloadTasks(const FImportAssetObjectP
 		FLodPayloads& LodPayloads = PayloadsPerLodIndex.FindOrAdd(LodIndex);
 
 		auto AddMeshPayloads = [&Arguments, MeshTranslatorPayloadInterface, bBakeMeshes, bBakePivotMeshes, &GlobalOffsetTransform, &PayloadTasks, bAsync, &PayloadAttributes]
-			(const TArray<FString>& MeshUids, TMap<FInterchangeMeshPayLoadKey, FMeshPayload>& PayloadPerKey)
+			(const TArray<FString>& MeshUids, TMap<FInterchangeMeshPayLoadKey, UE::Interchange::FMeshPayload>& PayloadPerKey)
 			{
 				for (const FString& MeshUid : MeshUids)
 				{
@@ -166,7 +161,7 @@ void UInterchangeStaticMeshFactory::CreatePayloadTasks(const FImportAssetObjectP
 					GlobalPayLoadKey.UniqueId += FInterchangeMeshPayLoadKey::GetTransformString(GlobalMeshTransform);
 					if (!PayloadPerKey.Contains(GlobalPayLoadKey))
 					{
-						FMeshPayload& Payload = PayloadPerKey.FindOrAdd(GlobalPayLoadKey);
+						UE::Interchange::FMeshPayload& Payload = PayloadPerKey.FindOrAdd(GlobalPayLoadKey);
 						Payload.Transform = GlobalMeshTransform;
 						Payload.MeshName = PayLoadKey.UniqueId;
 						PayloadAttributes.RegisterAttribute(UE::Interchange::FAttributeKey{ UE::Interchange::MeshPayload::Attributes::MeshGlobalTransform }, GlobalMeshTransform);
@@ -412,7 +407,7 @@ UInterchangeFactoryBase::FImportAssetResult UInterchangeStaticMeshFactory::Impor
 
 		// Just move the mesh description from the first valid payload then append the rest
 		bool bFirstValidMoved = false;
-		for (TPair<FInterchangeMeshPayLoadKey, FMeshPayload>& KeyAndPayload : LodPayloads.MeshPayloadPerKey)
+		for (TPair<FInterchangeMeshPayLoadKey, UE::Interchange::FMeshPayload>& KeyAndPayload : LodPayloads.MeshPayloadPerKey)
 		{
 			const TOptional<UE::Interchange::FMeshPayloadData>& LodMeshPayload = KeyAndPayload.Value.PayloadData;
 			if (!LodMeshPayload.IsSet())
@@ -518,10 +513,17 @@ UInterchangeFactoryBase::FImportAssetResult UInterchangeStaticMeshFactory::Impor
 					StaticMesh->GetBodySetup()->AggGeom.EmptyElements();
 				}
 
-				bImportedCustomCollision |= ImportBoxCollision(Arguments, StaticMesh);
-				bImportedCustomCollision |= ImportCapsuleCollision(Arguments, StaticMesh);
-				bImportedCustomCollision |= ImportSphereCollision(Arguments, StaticMesh);
-				bImportedCustomCollision |= ImportConvexCollision(Arguments, StaticMesh, LodDataNode);
+				bool bForceGeneration = false;
+				if (!LodDataNode->GetForceCollisionPrimitiveGeneration(bForceGeneration))
+				{
+					bForceGeneration = false;
+				}
+
+				using namespace UE::Interchange::Private::MeshHelper;
+				bImportedCustomCollision |= ImportBoxCollision(Arguments, LodPayloads.CollisionBoxPayloadPerKey, StaticMesh, bForceGeneration);
+				bImportedCustomCollision |= ImportCapsuleCollision(Arguments, LodPayloads.CollisionCapsulePayloadPerKey, StaticMesh);
+				bImportedCustomCollision |= ImportSphereCollision(Arguments, LodPayloads.CollisionSpherePayloadPerKey, StaticMesh, bForceGeneration);
+				bImportedCustomCollision |= ImportConvexCollision(Arguments, LodPayloads.CollisionConvexPayloadPerKey, StaticMesh, LodDataNode);
 			}
 		}
 
@@ -715,7 +717,7 @@ UInterchangeFactoryBase::FImportAssetResult UInterchangeStaticMeshFactory::EndIm
 
 	CommitMeshDescriptions(*StaticMesh);
 
-	ImportSockets(Arguments, StaticMesh, StaticMeshFactoryNode);
+	UE::Interchange::Private::MeshHelper::ImportSockets(Arguments, StaticMesh, StaticMeshFactoryNode);
 
 	if (!Arguments.ReimportObject)
 	{
@@ -1077,759 +1079,6 @@ void UInterchangeStaticMeshFactory::BuildObject_GameThread(const FSetupObjectPar
 #endif
 }
 
-bool UInterchangeStaticMeshFactory::AddConvexGeomFromVertices(const FImportAssetObjectParams& Arguments, const FMeshDescription& MeshDescription, const FTransform& Transform, FKAggregateGeom& AggGeom)
-{
-	FStaticMeshConstAttributes Attributes(MeshDescription);
-	TVertexAttributesRef<const FVector3f> VertexPositions = Attributes.GetVertexPositions();
-	
-	if (VertexPositions.GetNumElements() == 0)
-	{
-		return false;
-	}
-
-	FKConvexElem& ConvexElem = AggGeom.ConvexElems.Emplace_GetRef();
-	ConvexElem.VertexData.AddUninitialized(VertexPositions.GetNumElements());
-
-	for (int32 Index = 0; Index < VertexPositions.GetNumElements(); Index++)
-	{
-		ConvexElem.VertexData[Index] = Transform.TransformPosition(FVector(VertexPositions[Index]));
-	}
-
-	ConvexElem.UpdateElemBox();
-
-	return true;
-}
-
-
-bool UInterchangeStaticMeshFactory::DecomposeConvexMesh(const FImportAssetObjectParams& Arguments, const FMeshDescription& MeshDescription, const FTransform& Transform, UBodySetup* BodySetup)
-{
-#if WITH_EDITOR
-
-	// Construct a bit array containing a bit for each triangle ID in the mesh description.
-	// We are assuming the mesh description is compact, i.e. it has no holes, and so the number of triangles is equal to the array size.
-	// The aim is to identify 'islands' of adjacent triangles which will form separate convex hulls
-
-	check(MeshDescription.Triangles().Num() == MeshDescription.Triangles().GetArraySize());
-	TBitArray<> BitArray(0, MeshDescription.Triangles().Num());
-
-	// Here we build the groups of triangle IDs
-
-	TArray<TArray<FTriangleID>> TriangleGroups;
-
-	int32 FirstIndex = BitArray.FindAndSetFirstZeroBit();
-	while (FirstIndex != INDEX_NONE)
-	{
-		// Find the first index we haven't used yet, and use it as the beginning of a new triangle group
-
-		TArray<FTriangleID>& TriangleGroup = TriangleGroups.Emplace_GetRef();
-		TriangleGroup.Emplace(FirstIndex);
-
-		// Now iterate through the TriangleGroup array, finding unused adjacent triangles to each index, and appending them
-		// to the end of the array.  Note we deliberately check the array size each time round the loop, as each iteration
-		// can cause it to grow.
-
-		for (int32 CheckIndex = 0; CheckIndex < TriangleGroup.Num(); ++CheckIndex)
-		{
-			for (FTriangleID AdjacentTriangleID : MeshDescription.GetTriangleAdjacentTriangles(TriangleGroup[CheckIndex]))
-			{
-				if (BitArray[AdjacentTriangleID] == 0)
-				{
-					// Append unused adjacent triangles to the TriangleGroup, to be considered for adjacency later
-					TriangleGroup.Emplace(AdjacentTriangleID);
-					BitArray[AdjacentTriangleID] = 1;
-				}
-			}
-		}
-
-		// When we exhaust the triangle group array, there are no more triangles in this island.
-		// Now find the start of the next group.
-
-		FirstIndex = BitArray.FindAndSetFirstZeroBit();
-	}
-
-	// Now iterate through the triangle groups, adding each as a convex hull to the AggGeom
-
-	UModel* TempModel = NewObject<UModel>();
-	TempModel->RootOutside = true;
-	TempModel->EmptyModel(true, true);
-	TempModel->Polys->ClearFlags(RF_Transactional);
-
-	FStaticMeshConstAttributes Attributes(MeshDescription);
-	TTriangleAttributesRef<TArrayView<const FVertexID>> TriangleVertices = Attributes.GetTriangleVertexIndices();
-	TVertexAttributesRef<const FVector3f> VertexPositions = Attributes.GetVertexPositions();
-
-	bool bSuccess = true;
-
-	for (const TArray<FTriangleID>& TriangleGroup : TriangleGroups)
-	{
-		// Initialize a new brush
-		TempModel->Polys->Element.Empty();
-
-		// Add each triangle to the brush
-		int32 Index = 0;
-		for (FTriangleID TriangleID : TriangleGroup)
-		{
-			FPoly& Poly = TempModel->Polys->Element.Emplace_GetRef();
-			Poly.Init();
-			Poly.iLink = Index++;
-
-			// For reasons lost in time, BSP poly vertices have the opposite winding order to regular mesh vertices.
-			// So add them backwards (sigh)
-			Poly.Vertices.Emplace(Transform.TransformPosition(FVector(VertexPositions[TriangleVertices[TriangleID][2]])));
-			Poly.Vertices.Emplace(Transform.TransformPosition(FVector(VertexPositions[TriangleVertices[TriangleID][1]])));
-			Poly.Vertices.Emplace(Transform.TransformPosition(FVector(VertexPositions[TriangleVertices[TriangleID][0]])));
-
-			Poly.CalcNormal(true);
-		}
-
-		// Build bounding box
-		TempModel->BuildBound();
-
-		// Build BSP for the brush
-		FBSPOps::bspBuild(TempModel, FBSPOps::BSP_Good, 15, 70, 1, 0);
-		FBSPOps::bspRefresh(TempModel, true);
-		FBSPOps::bspBuildBounds(TempModel);
-
-		bSuccess &= BodySetup->CreateFromModel(TempModel, false);
-	}
-
-	TempModel->ClearInternalFlags(EInternalObjectFlags::Async);
-	TempModel->Polys->ClearInternalFlags(EInternalObjectFlags::Async);
-
-	return bSuccess;
-
-#else // #if WITH_EDITOR
-
-	return false;
-
-#endif
-}
-
-
-static bool AreEqual(float A, float B)
-{
-	constexpr float MeshToPrimTolerance = 0.001f;
-	return FMath::Abs(A - B) < MeshToPrimTolerance;
-}
-
-
-static bool AreParallel(const FVector3f& A, const FVector3f& B)
-{
-	float Dot = FVector3f::DotProduct(A, B);
-
-	return (AreEqual(FMath::Abs(Dot), 1.0f));
-}
-
-
-static FVector3f GetTriangleNormal(const FTransform& Transform, TVertexAttributesRef<const FVector3f> VertexPositions, TArrayView<const FVertexID> VertexIndices)
-{
-	const FVector3f& V0 = VertexPositions[VertexIndices[0]];
-	const FVector3f& V1 = VertexPositions[VertexIndices[1]];
-	const FVector3f& V2 = VertexPositions[VertexIndices[2]];
-	// @todo: LWC conversions everywhere here; surely this can be more elegant?
-	return FVector3f(Transform.TransformVector(FVector(FVector3f::CrossProduct(V1 - V0, V2 - V0).GetSafeNormal())));
-}
-
-
-bool UInterchangeStaticMeshFactory::AddBoxGeomFromTris(const FMeshDescription& MeshDescription, const FTransform& Transform, FKAggregateGeom& AggGeom)
-{
-	// Maintain an array of the planes we have encountered so far.
-	// We are expecting two instances of three unique plane orientations, one for each side of the box.
-
-	struct FPlaneInfo
-	{
-		FPlaneInfo(const FVector3f InNormal, float InFirstDistance)
-			: Normal(InNormal),
-			  DistCount(1),
-			  PlaneDist{InFirstDistance, 0.0f}
-		{}
-
-		FVector3f Normal;
-		int32 DistCount;
-		float PlaneDist[2];
-	};
-
-	TArray<FPlaneInfo> Planes;
-	FBox Box(ForceInit);
-
-	FStaticMeshConstAttributes Attributes(MeshDescription);
-	TTriangleAttributesRef<TArrayView<const FVertexID>> TriangleVertices = Attributes.GetTriangleVertexIndices();
-	TVertexAttributesRef<const FVector3f> VertexPositions = Attributes.GetVertexPositions();
-
-	for (FTriangleID TriangleID : MeshDescription.Triangles().GetElementIDs())
-	{
-		TArrayView<const FVertexID> VertexIndices = TriangleVertices[TriangleID];
-
-		FVector3f TriangleNormal = GetTriangleNormal(Transform, VertexPositions, VertexIndices);
-		if (TriangleNormal.IsNearlyZero())
-		{
-			continue;
-		}
-
-		bool bFoundPlane = false;
-		for (int32 PlaneIndex = 0; PlaneIndex < Planes.Num() && !bFoundPlane; PlaneIndex++)
-		{
-			// if this triangle plane is already known...
-			if (AreParallel(TriangleNormal, Planes[PlaneIndex].Normal))
-			{
-				// Always use the same normal when comparing distances, to ensure consistent sign.
-				float Dist = FVector3f::DotProduct(VertexPositions[VertexIndices[0]], Planes[PlaneIndex].Normal);
-
-				// we only have one distance, and its not that one, add it.
-				if (Planes[PlaneIndex].DistCount == 1 && !AreEqual(Dist, Planes[PlaneIndex].PlaneDist[0]))
-				{
-					Planes[PlaneIndex].PlaneDist[1] = Dist;
-					Planes[PlaneIndex].DistCount = 2;
-				}
-				// if we have a second distance, and its not that either, something is wrong.
-				else if (Planes[PlaneIndex].DistCount == 2 && !AreEqual(Dist, Planes[PlaneIndex].PlaneDist[0]) && !AreEqual(Dist, Planes[PlaneIndex].PlaneDist[1]))
-				{
-					// Error
-//					UE_LOG(LogStaticMeshImportUtils, Log, TEXT("AddBoxGeomFromTris (%s): Found more than 2 planes with different distances."), ObjName);
-					return false;
-				}
-
-				bFoundPlane = true;
-			}
-		}
-
-		// If this triangle does not match an existing plane, add to list.
-		if (!bFoundPlane)
-		{
-			Planes.Emplace(TriangleNormal, FVector3f::DotProduct(VertexPositions[VertexIndices[0]], TriangleNormal));
-		}
-
-		// Maintain an AABB, adding points from each triangle.
-		// We will use this to determine the origin of the box transform.
-
-		Box += Transform.TransformPosition(FVector(VertexPositions[VertexIndices[0]]));
-		Box += Transform.TransformPosition(FVector(VertexPositions[VertexIndices[1]]));
-		Box += Transform.TransformPosition(FVector(VertexPositions[VertexIndices[2]]));
-	}
-
-	// Now we have our candidate planes, see if there are any problems
-
-	// Wrong number of planes.
-	if (Planes.Num() != 3)
-	{
-		// Error
-//		UE_LOG(LogStaticMeshImportUtils, Log, TEXT("AddBoxGeomFromTris (%s): Not very box-like (need 3 sets of planes)."), ObjName);
-		return false;
-	}
-
-	// If we don't have 3 pairs, we can't carry on.
-	if ((Planes[0].DistCount != 2) || (Planes[1].DistCount != 2) || (Planes[2].DistCount != 2))
-	{
-		// Error
-//		UE_LOG(LogStaticMeshImportUtils, Log, TEXT("AddBoxGeomFromTris (%s): Incomplete set of planes (need 2 per axis)."), ObjName);
-		return false;
-	}
-
-	// ensure valid TM by cross-product
-	if (!AreParallel(FVector3f::CrossProduct(Planes[0].Normal, Planes[1].Normal), Planes[2].Normal))
-	{
-		// Error
-//		UE_LOG(LogStaticMeshImportUtils, Log, TEXT("AddBoxGeomFromTris (%s): Box axes are not perpendicular."), ObjName);
-		return false;
-	}
-
-	// Allocate box in array
-	FKBoxElem BoxElem;
-
-	//In case we have a box oriented with the world axis system we want to reorder the plane to not introduce axis swap.
-	//If the box was turned, the order of the planes will be arbitrary and the box rotation will make the collision
-	//not playing well if the asset is built or place in a level with a non uniform scale.
-	FVector3f Axis[3] = { FVector3f::XAxisVector, FVector3f::YAxisVector, FVector3f::ZAxisVector };
-	int32 Reorder[3] = { INDEX_NONE, INDEX_NONE, INDEX_NONE };
-	for (int32 PlaneIndex = 0; PlaneIndex < 3; ++PlaneIndex)
-	{
-		for (int32 AxisIndex = 0; AxisIndex < 3; ++AxisIndex)
-		{
-			if (AreParallel(Planes[PlaneIndex].Normal, Axis[AxisIndex]))
-			{
-				Reorder[PlaneIndex] = AxisIndex;
-				break;
-			}
-		}
-	}
-
-	if (Reorder[0] == INDEX_NONE || Reorder[1] == INDEX_NONE || Reorder[2] == INDEX_NONE)
-	{
-		Reorder[0] = 0;
-		Reorder[1] = 1;
-		Reorder[2] = 2;
-	}
-
-	BoxElem.SetTransform(FTransform(FVector(Planes[Reorder[0]].Normal), FVector(Planes[Reorder[1]].Normal), FVector(Planes[Reorder[2]].Normal), Box.GetCenter()));
-
-	// distance between parallel planes is box edge lengths.
-	BoxElem.X = FMath::Abs(Planes[Reorder[0]].PlaneDist[0] - Planes[Reorder[0]].PlaneDist[1]);
-	BoxElem.Y = FMath::Abs(Planes[Reorder[1]].PlaneDist[0] - Planes[Reorder[1]].PlaneDist[1]);
-	BoxElem.Z = FMath::Abs(Planes[Reorder[2]].PlaneDist[0] - Planes[Reorder[2]].PlaneDist[1]);
-
-	AggGeom.BoxElems.Add(BoxElem);
-
-	return true;
-}
-
-
-bool UInterchangeStaticMeshFactory::AddSphereGeomFromVertices(const FImportAssetObjectParams& Arguments, const FMeshDescription& MeshDescription, const FTransform& Transform, FKAggregateGeom& AggGeom)
-{
-	FStaticMeshConstAttributes Attributes(MeshDescription);
-	TVertexAttributesRef<const FVector3f> VertexPositions = Attributes.GetVertexPositions();
-
-	if (VertexPositions.GetNumElements() == 0)
-	{
-		return false;
-	}
-
-	FBox Box(ForceInit);
-
-	for (const FVector3f& VertexPosition : VertexPositions.GetRawArray())
-	{
-		Box += Transform.TransformPosition(FVector(VertexPosition));
-	}
-
-	FVector Center;
-	FVector Extents;
-	Box.GetCenterAndExtents(Center, Extents);
-	float Longest = 2.0f * Extents.GetMax();
-	float Shortest = 2.0f * Extents.GetMin();
-
-	// check that the AABB is roughly a square (5% tolerance)
-	if ((Longest - Shortest) / Longest > 0.05f)
-	{
-		// Error
-		//UE_LOG(LogStaticMeshImportUtils, Log, TEXT("AddSphereGeomFromVerts (%s): Sphere bounding box not square."), ObjName);
-		return false;
-	}
-
-	float Radius = 0.5f * Longest;
-
-	// Test that all vertices are a similar radius (5%) from the sphere centre.
-	float MaxR = 0;
-	float MinR = BIG_NUMBER;
-
-
-	for (const FVector3f& VertexPosition : VertexPositions.GetRawArray())
-	{
-		FVector CToV = Transform.TransformPosition(FVector(VertexPosition)) - Center;
-		float RSqr = CToV.SizeSquared();
-
-		MaxR = FMath::Max(RSqr, MaxR);
-
-		// Sometimes vertex at centre, so reject it.
-		if (RSqr > KINDA_SMALL_NUMBER)
-		{
-			MinR = FMath::Min(RSqr, MinR);
-		}
-	}
-
-	MaxR = FMath::Sqrt(MaxR);
-	MinR = FMath::Sqrt(MinR);
-
-	if ((MaxR - MinR) / Radius > 0.05f)
-	{
-		// Error
-		//UE_LOG(LogStaticMeshImportUtils, Log, TEXT("AddSphereGeomFromVerts (%s): Vertices not at constant radius."), ObjName);
-		return false;
-	}
-
-	// Allocate sphere in array
-	FKSphereElem SphereElem;
-	SphereElem.Center = Center;
-	SphereElem.Radius = Radius;
-	AggGeom.SphereElems.Add(SphereElem);
-
-	return true;
-}
-
-
-bool UInterchangeStaticMeshFactory::AddCapsuleGeomFromVertices(const FImportAssetObjectParams& Arguments, const FMeshDescription& MeshDescription, const FTransform& Transform, FKAggregateGeom& AggGeom)
-{
-	FStaticMeshConstAttributes Attributes(MeshDescription);
-	TVertexAttributesRef<const FVector3f> VertexPositions = Attributes.GetVertexPositions();
-
-	if (VertexPositions.GetNumElements() == 0)
-	{
-		return false;
-	}
-
-	FVector AxisStart = FVector::ZeroVector;
-	FVector AxisEnd = FVector::ZeroVector;
-	float MaxDistSqr = 0.f;
-
-	for (int32 IndexA = 0; IndexA < VertexPositions.GetNumElements() - 1; IndexA++)
-	{
-		for (int32 IndexB = IndexA + 1; IndexB < VertexPositions.GetNumElements(); IndexB++)
-		{
-			FVector TransformedA = Transform.TransformPosition(FVector(VertexPositions[IndexA]));
-			FVector TransformedB = Transform.TransformPosition(FVector(VertexPositions[IndexB]));
-
-			float DistSqr = (TransformedA - TransformedB).SizeSquared();
-			if (DistSqr > MaxDistSqr)
-			{
-				AxisStart = TransformedA;
-				AxisEnd = TransformedB;
-				MaxDistSqr = DistSqr;
-			}
-		}
-	}
-
-	// If we got a valid axis, find vertex furthest from it
-	if (MaxDistSqr > SMALL_NUMBER)
-	{
-		float MaxRadius = 0.0f;
-
-		const FVector LineOrigin = AxisStart;
-		const FVector LineDir = (AxisEnd - AxisStart).GetSafeNormal();
-
-		for (int32 IndexA = 0; IndexA < VertexPositions.GetNumElements(); IndexA++)
-		{
-			FVector TransformedA = Transform.TransformPosition(FVector(VertexPositions[IndexA]));
-
-			float DistToAxis = FMath::PointDistToLine(TransformedA, LineDir, LineOrigin);
-			if (DistToAxis > MaxRadius)
-			{
-				MaxRadius = DistToAxis;
-			}
-		}
-
-		if (MaxRadius > SMALL_NUMBER)
-		{
-			// Allocate capsule in array
-			FKSphylElem SphylElem;
-			SphylElem.Center = 0.5f * (AxisStart + AxisEnd);
-			SphylElem.Rotation = FQuat::FindBetweenVectors(FVector::ZAxisVector, LineDir).Rotator(); // Get quat that takes you from z axis to desired axis
-			SphylElem.Radius = MaxRadius;
-			SphylElem.Length = FMath::Max(FMath::Sqrt(MaxDistSqr) - (2.0f * MaxRadius), 0.0f); // subtract two radii from total length to get segment length (ensure > 0)
-			AggGeom.SphylElems.Add(SphylElem);
-			return true;
-		}
-	}
-
-	return false;
-
-}
-
-
-bool UInterchangeStaticMeshFactory::ImportBoxCollision(const FImportAssetObjectParams& Arguments, UStaticMesh* StaticMesh)
-{
-	using namespace UE::Interchange;
-
-	bool bResult = false;
-
-	TMap<FInterchangeMeshPayLoadKey, FMeshPayload> BoxCollisionPayloads = PayloadsPerLodIndex.FindChecked(0).CollisionBoxPayloadPerKey;
-
-	FKAggregateGeom& AggGeo = StaticMesh->GetBodySetup()->AggGeom;
-
-	for (TPair<FInterchangeMeshPayLoadKey, FMeshPayload>& BoxCollisionPayload : BoxCollisionPayloads)
-	{
-		const FTransform Transform = FTransform::Identity;
-		const TOptional<FMeshPayloadData>& PayloadData = BoxCollisionPayload.Value.PayloadData;
-
-		if (!PayloadData.IsSet())
-		{
-			// warning here
-			continue;
-		}
-
-		if (AddBoxGeomFromTris(PayloadData->MeshDescription, Transform, AggGeo))
-		{
-			bResult = true;
-			FKBoxElem& NewElem = AggGeo.BoxElems.Last();
-
-			// Now test the last element in the AggGeo list and remove it if its a duplicate
-			// @TODO: determine why we have to do this. Was it to prevent duplicate boxes accumulating when reimporting?
-			for (int32 ElementIndex = 0; ElementIndex < AggGeo.BoxElems.Num() - 1; ++ElementIndex)
-			{
-				FKBoxElem& CurrentElem = AggGeo.BoxElems[ElementIndex];
-
-				if (CurrentElem == NewElem)
-				{
-					// The new element is a duplicate, remove it
-					AggGeo.BoxElems.RemoveAt(AggGeo.BoxElems.Num() - 1);
-					break;
-				}
-			}
-		}
-	}
-
-	return bResult;
-}
-
-
-bool UInterchangeStaticMeshFactory::ImportCapsuleCollision(const FImportAssetObjectParams& Arguments, UStaticMesh* StaticMesh)
-{
-	using namespace UE::Interchange;
-
-	bool bResult = false;
-
-	TMap<FInterchangeMeshPayLoadKey, FMeshPayload> CapsuleCollisionPayloads = PayloadsPerLodIndex.FindChecked(0).CollisionCapsulePayloadPerKey;
-
-	FKAggregateGeom& AggGeo = StaticMesh->GetBodySetup()->AggGeom;
-
-	for (TPair<FInterchangeMeshPayLoadKey, FMeshPayload>& CapsuleCollisionPayload : CapsuleCollisionPayloads)
-	{
-		const FTransform& Transform = FTransform::Identity;
-		const TOptional<FMeshPayloadData>& PayloadData = CapsuleCollisionPayload.Value.PayloadData;
-
-		if (!PayloadData.IsSet())
-		{
-			// warning here
-			continue;
-		}
-
-		if (AddCapsuleGeomFromVertices(Arguments, PayloadData->MeshDescription, Transform, AggGeo))
-		{
-			bResult = true;
-
-			FKSphylElem& NewElem = AggGeo.SphylElems.Last();
-
-			// Now test the late element in the AggGeo list and remove it if its a duplicate
-			// @TODO: determine why we have to do this. Was it to prevent duplicate boxes accumulating when reimporting?
-			for (int32 ElementIndex = 0; ElementIndex < AggGeo.SphylElems.Num() - 1; ++ElementIndex)
-			{
-				FKSphylElem& CurrentElem = AggGeo.SphylElems[ElementIndex];
-				if (CurrentElem == NewElem)
-				{
-					// The new element is a duplicate, remove it
-					AggGeo.SphylElems.RemoveAt(AggGeo.SphylElems.Num() - 1);
-					break;
-				}
-			}
-		}
-	}
-
-	return bResult;
-}
-
-
-bool UInterchangeStaticMeshFactory::ImportSphereCollision(const FImportAssetObjectParams& Arguments, UStaticMesh* StaticMesh)
-{
-	using namespace UE::Interchange;
-
-	bool bResult = false;
-
-	TMap<FInterchangeMeshPayLoadKey, FMeshPayload> SphereCollisionPayloads = PayloadsPerLodIndex.FindChecked(0).CollisionSpherePayloadPerKey;
-
-	FKAggregateGeom& AggGeo = StaticMesh->GetBodySetup()->AggGeom;
-
-	for (TPair<FInterchangeMeshPayLoadKey, FMeshPayload>& SphereCollisionPayload : SphereCollisionPayloads)
-	{
-		const FTransform& Transform = FTransform::Identity;
-		const TOptional<FMeshPayloadData>& PayloadData = SphereCollisionPayload.Value.PayloadData;
-
-		if (!PayloadData.IsSet())
-		{
-			// warning here
-			continue;
-		}
-
-		if (AddSphereGeomFromVertices(Arguments, PayloadData->MeshDescription, Transform, AggGeo))
-		{
-			bResult = true;
-
-			FKSphereElem& NewElem = AggGeo.SphereElems.Last();
-
-			// Now test the last element in the AggGeo list and remove it if its a duplicate
-			// @TODO: determine why we have to do this. Was it to prevent duplicate boxes accumulating when reimporting?
-			for (int32 ElementIndex = 0; ElementIndex < AggGeo.SphereElems.Num() - 1; ++ElementIndex)
-			{
-				FKSphereElem& CurrentElem = AggGeo.SphereElems[ElementIndex];
-
-				if (CurrentElem == NewElem)
-				{
-					// The new element is a duplicate, remove it
-					AggGeo.SphereElems.RemoveAt(AggGeo.SphereElems.Num() - 1);
-					break;
-				}
-			}
-		}
-	}
-
-	return bResult;
-}
-
-
-bool UInterchangeStaticMeshFactory::ImportConvexCollision(const FImportAssetObjectParams& Arguments, UStaticMesh* StaticMesh, const UInterchangeStaticMeshLodDataNode* LodDataNode)
-{
-	using namespace UE::Interchange;
-
-	bool bResult = false;
-
-	TMap<FInterchangeMeshPayLoadKey, FMeshPayload> ConvexCollisionPayloads = PayloadsPerLodIndex.FindChecked(0).CollisionConvexPayloadPerKey;
-
-	bool bOneConvexHullPerUCX;
-	if (!LodDataNode->GetOneConvexHullPerUCX(bOneConvexHullPerUCX) || !bOneConvexHullPerUCX)
-	{
-		for (TPair<FInterchangeMeshPayLoadKey, FMeshPayload>& ConvexCollisionPayload : ConvexCollisionPayloads)
-		{
-			const FTransform& Transform = FTransform::Identity;
-			const TOptional<FMeshPayloadData>& PayloadData = ConvexCollisionPayload.Value.PayloadData;
-
-			if (!PayloadData.IsSet())
-			{
-				// warning here
-				continue;
-			}
-
-			if (!DecomposeConvexMesh(Arguments, PayloadData->MeshDescription, Transform, StaticMesh->GetBodySetup()))
-			{
-				// error: could not decompose mesh
-			}
-			else
-			{
-				bResult = true;
-			}
-		}
-	}
-	else
-	{
-		FKAggregateGeom& AggGeo = StaticMesh->GetBodySetup()->AggGeom;
-
-		for (TPair<FInterchangeMeshPayLoadKey, FMeshPayload>& ConvexCollisionPayload : ConvexCollisionPayloads)
-		{
-			const FTransform& Transform = FTransform::Identity;
-			TOptional<FMeshPayloadData> PayloadData = ConvexCollisionPayload.Value.PayloadData;
-
-			if (!PayloadData.IsSet())
-			{
-				// warning here
-				continue;
-			}
-
-			if (AddConvexGeomFromVertices(Arguments, PayloadData->MeshDescription, Transform, AggGeo))
-			{
-				bResult = true;
-
-				FKConvexElem& NewElem = AggGeo.ConvexElems.Last();
-
-				// Now test the late element in the AggGeo list and remove it if its a duplicate
-				// @TODO: determine why the importer used to do this. Was it something to do with reimport not adding extra collision or something?
-				for (int32 ElementIndex = 0; ElementIndex < AggGeo.ConvexElems.Num() - 1; ++ElementIndex)
-				{
-					FKConvexElem& CurrentElem = AggGeo.ConvexElems[ElementIndex];
-
-					if (CurrentElem.VertexData.Num() == NewElem.VertexData.Num())
-					{
-						bool bFoundDifference = false;
-						for (int32 VertexIndex = 0; VertexIndex < NewElem.VertexData.Num(); ++VertexIndex)
-						{
-							if (CurrentElem.VertexData[VertexIndex] != NewElem.VertexData[VertexIndex])
-							{
-								bFoundDifference = true;
-								break;
-							}
-						}
-
-						if (!bFoundDifference)
-						{
-							// The new collision geo is a duplicate, delete it
-							AggGeo.ConvexElems.RemoveAt(AggGeo.ConvexElems.Num() - 1);
-							break;
-						}
-					}
-				}
-			}
-		}
-	}
-
-	return bResult;
-}
-
-bool UInterchangeStaticMeshFactory::ImportSockets(const FImportAssetObjectParams& Arguments, UStaticMesh* StaticMesh, const UInterchangeStaticMeshFactoryNode* FactoryNode)
-{
-	bool bImportSockets = false;
-	FactoryNode->GetCustomImportSockets(bImportSockets);
-	if (!bImportSockets)
-	{
-		//Skip Import sockets
-		return true;
-	}
-
-	TArray<FString> SocketUids;
-	FactoryNode->GetSocketUids(SocketUids);
-
-	TSet<FName> ImportedSocketNames;
-
-	FTransform GlobalOffsetTransform = FTransform::Identity;
-
-	bool bBakeMeshes = false;
-	bool bBakePivotMeshes = false;
-	if (UInterchangeCommonPipelineDataFactoryNode* CommonPipelineDataFactoryNode = UInterchangeCommonPipelineDataFactoryNode::GetUniqueInstance(Arguments.NodeContainer))
-	{
-		CommonPipelineDataFactoryNode->GetCustomGlobalOffsetTransform(GlobalOffsetTransform);
-		CommonPipelineDataFactoryNode->GetBakeMeshes(bBakeMeshes);
-		if (!bBakeMeshes)
-		{
-			CommonPipelineDataFactoryNode->GetBakePivotMeshes(bBakePivotMeshes);
-		}
-	}
-
-	for (const FString& SocketUid : SocketUids)
-	{
-		if (const UInterchangeSceneNode* SceneNode = Cast<UInterchangeSceneNode>(Arguments.NodeContainer->GetNode(SocketUid)))
-		{
-			FString NodeDisplayName = SceneNode->GetDisplayLabel();
-			if (NodeDisplayName.StartsWith(UInterchangeMeshFactoryNode::GetMeshSocketPrefix()))
-			{
-				NodeDisplayName.RightChopInline(UInterchangeMeshFactoryNode::GetMeshSocketPrefix().Len(), EAllowShrinking::No);
-			}
-			FName SocketName = FName(NodeDisplayName);
-			ImportedSocketNames.Add(SocketName);
-
-			FTransform Transform;
-			if (bBakeMeshes)
-			{
-				SceneNode->GetCustomGlobalTransform(Arguments.NodeContainer, GlobalOffsetTransform, Transform);
-			}
-
-			UE::Interchange::Private::MeshHelper::AddSceneNodeGeometricAndPivotToGlobalTransform(Transform, SceneNode, bBakeMeshes, bBakePivotMeshes);
-
-			//Apply axis transformation inverse to get correct Socket Transform:
-			const UInterchangeSourceNode* SourceNode = UInterchangeSourceNode::GetUniqueInstance(Arguments.NodeContainer);
-			FTransform AxisConversionInverseTransform;
-			if (SourceNode->GetCustomAxisConversionInverseTransform(AxisConversionInverseTransform))
-			{
-				Transform = AxisConversionInverseTransform * Transform;
-			}
-
-			UStaticMeshSocket* Socket = StaticMesh->FindSocket(SocketName);
-			if (!Socket)
-			{
-				// If the socket didn't exist create a new one now
-				Socket = NewObject<UStaticMeshSocket>(StaticMesh);
-#if WITH_EDITORONLY_DATA
-				Socket->bSocketCreatedAtImport = true;
-#endif
-				Socket->SocketName = SocketName;
-				StaticMesh->AddSocket(Socket);
-			}
-
-			Socket->RelativeLocation = Transform.GetLocation();
-			Socket->RelativeRotation = Transform.GetRotation().Rotator();
-			Socket->RelativeScale = Transform.GetScale3D();
-		}
-	}
-
-	// Delete any sockets which were previously imported but which no longer exist in the imported scene
-	for (TArray<TObjectPtr<UStaticMeshSocket>>::TIterator It = StaticMesh->Sockets.CreateIterator(); It; ++It)
-	{
-		UStaticMeshSocket* Socket = *It;
-		if (
-#if WITH_EDITORONLY_DATA
-			Socket->bSocketCreatedAtImport &&
-#endif
-			!ImportedSocketNames.Contains(Socket->SocketName))
-		{
-			It.RemoveCurrent();
-		}
-	}
-
-	return true;
-}
-
 bool UInterchangeStaticMeshFactory::GetSourceFilenames(const UObject* Object, TArray<FString>& OutSourceFilenames) const
 {
 #if WITH_EDITORONLY_DATA
diff --git a/Engine/Plugins/Interchange/Runtime/Source/Import/Public/Mesh/InterchangeMeshHelper.h b/Engine/Plugins/Interchange/Runtime/Source/Import/Public/Mesh/InterchangeMeshHelper.h
--- a/Engine/Plugins/Interchange/Runtime/Source/Import/Public/Mesh/InterchangeMeshHelper.h
+++ b/Engine/Plugins/Interchange/Runtime/Source/Import/Public/Mesh/InterchangeMeshHelper.h
@@ -1,22 +1,104 @@
 // Copyright Epic Games, Inc. All Rights Reserved. 
+
+#pragma once
+
 #include "MeshDescription.h"
 #include "StaticMeshOperations.h"
 
-struct FSkeletalMaterial;
-struct FStaticMaterial;
+#include "InterchangeFactoryBase.h"
+
+class UBodySetup;
 class UInterchangeBaseNodeContainer;
 class UInterchangeSceneNode;
+class UInterchangeStaticMeshFactoryNode;
+class UInterchangeStaticMeshLodDataNode;
 class UMaterialInterface;
 class USkeletalMesh;
+class UStaticMesh;
+struct FKAggregateGeom;
+struct FSkeletalMaterial;
+struct FStaticMaterial;
+struct FInterchangeMeshPayLoadKey;
 
 namespace UE::Interchange
 {
 	struct FMeshPayloadData;
+	struct FMeshPayload;
 }
 
 namespace UE::Interchange::Private::MeshHelper
 {
-	INTERCHANGEIMPORT_API void RemapPolygonGroups(const FMeshDescription& SourceMesh, FMeshDescription& TargetMesh, PolygonGroupMap& RemapPolygonGroup);
+	/** These are largely copied from StaticMeshImportUtils.h */
+
+	INTERCHANGEIMPORT_API bool AddConvexGeomFromVertices(
+		const UInterchangeFactoryBase::FImportAssetObjectParams& Arguments,
+		const FMeshDescription& MeshDescription,
+		FKAggregateGeom& AggGeom
+	);
+
+	INTERCHANGEIMPORT_API bool DecomposeConvexMesh(
+		const UInterchangeFactoryBase::FImportAssetObjectParams& Arguments,
+		const FMeshDescription& MeshDescription,
+		UBodySetup* BodySetup
+	);
+
+	INTERCHANGEIMPORT_API bool AddBoxGeomFromTris(
+		const FMeshDescription& MeshDescription, 
+		FKAggregateGeom& AggGeom,
+		bool bForcePrimitiveGeneration = true
+	);
+
+	INTERCHANGEIMPORT_API bool AddSphereGeomFromVertices(
+		const UInterchangeFactoryBase::FImportAssetObjectParams& Arguments,
+		const FMeshDescription& MeshDescription,
+		FKAggregateGeom& AggGeom,
+		bool bForcePrimitiveGeneration = true
+	);
+
+	INTERCHANGEIMPORT_API bool AddCapsuleGeomFromVertices(
+		const UInterchangeFactoryBase::FImportAssetObjectParams& Arguments,
+		const FMeshDescription& MeshDescription,
+		FKAggregateGeom& AggGeom
+	);
+
+	INTERCHANGEIMPORT_API bool ImportBoxCollision(
+		const UInterchangeFactoryBase::FImportAssetObjectParams& Arguments,
+		const TMap<FInterchangeMeshPayLoadKey, FMeshPayload>& BoxCollisionPayloads,
+		UStaticMesh* StaticMesh,
+		bool bForcePrimitiveGeneration = true
+	);
+
+	INTERCHANGEIMPORT_API bool ImportCapsuleCollision(
+		const UInterchangeFactoryBase::FImportAssetObjectParams& Arguments,
+		const TMap<FInterchangeMeshPayLoadKey, FMeshPayload>& CapsuleCollisionPayloads,
+		UStaticMesh* StaticMesh
+	);
+
+	INTERCHANGEIMPORT_API bool ImportSphereCollision(
+		const UInterchangeFactoryBase::FImportAssetObjectParams& Arguments,
+		const TMap<FInterchangeMeshPayLoadKey, FMeshPayload>& SphereCollisionPayloads, 
+		UStaticMesh* StaticMesh,
+		bool bForcePrimitiveGeneration = true
+	);
+
+	INTERCHANGEIMPORT_API bool ImportConvexCollision(
+		const UInterchangeFactoryBase::FImportAssetObjectParams& Arguments,
+		const TMap<FInterchangeMeshPayLoadKey, FMeshPayload>& ConvexCollisionPayloads,
+		UStaticMesh* StaticMesh,
+		const UInterchangeStaticMeshLodDataNode* LodDataNode
+	);
+
+	INTERCHANGEIMPORT_API bool ImportSockets(
+		const UInterchangeFactoryBase::FImportAssetObjectParams& Arguments,
+		UStaticMesh* StaticMesh,
+		const UInterchangeStaticMeshFactoryNode* FactoryNode
+	);
+
+	INTERCHANGEIMPORT_API void RemapPolygonGroups(
+		const FMeshDescription& SourceMesh,
+		FMeshDescription& TargetMesh,
+		PolygonGroupMap& RemapPolygonGroup
+	);
 
 	/* return the result of the global transform with the geometric and pivot transform of the scene node. */
 	INTERCHANGEIMPORT_API void AddSceneNodeGeometricAndPivotToGlobalTransform(FTransform& GlobalTransform, const UInterchangeSceneNode* SceneNode, const bool bBakeMeshes, const bool bBakePivotMeshes);
@@ -46,4 +128,4 @@ namespace UE::Interchange::Private::MeshHelper
 	INTERCHANGEIMPORT_API void RemapSkeletalMeshVertexColorToMeshDescription(const USkeletalMesh* SkeletalMesh, const int32 LODIndex, FMeshDescription& MeshDescription);
 #endif
 
-} //ns UE::Interchange::Private::MeshHelper
\ No newline at end of file
+} //ns UE::Interchange::Private::MeshHelper
diff --git a/Engine/Plugins/Interchange/Runtime/Source/Import/Public/Mesh/InterchangeMeshPayload.h b/Engine/Plugins/Interchange/Runtime/Source/Import/Public/Mesh/InterchangeMeshPayload.h
--- a/Engine/Plugins/Interchange/Runtime/Source/Import/Public/Mesh/InterchangeMeshPayload.h
+++ b/Engine/Plugins/Interchange/Runtime/Source/Import/Public/Mesh/InterchangeMeshPayload.h
@@ -26,5 +26,12 @@ namespace UE
 			//The name of the morph target
 			FString MorphTargetName;
 		};
+
+		struct FMeshPayload
+		{
+			FString MeshName;
+			TOptional<UE::Interchange::FMeshPayloadData> PayloadData;
+			FTransform Transform = FTransform::Identity;
+		};
 	}//ns Interchange
 }//ns UE
diff --git a/Engine/Plugins/Interchange/Runtime/Source/Import/Public/Mesh/InterchangeStaticMeshFactory.h b/Engine/Plugins/Interchange/Runtime/Source/Import/Public/Mesh/InterchangeStaticMeshFactory.h
--- a/Engine/Plugins/Interchange/Runtime/Source/Import/Public/Mesh/InterchangeStaticMeshFactory.h
+++ b/Engine/Plugins/Interchange/Runtime/Source/Import/Public/Mesh/InterchangeStaticMeshFactory.h
@@ -49,37 +49,18 @@ public:
 	// Interchange factory base interface end
 	//////////////////////////////////////////////////////////////////////////
 
-private:
-	struct FMeshPayload
-	{
-		FString MeshName;
-		TOptional<UE::Interchange::FMeshPayloadData> PayloadData;
-		FTransform Transform = FTransform::Identity;
-	};
-
+private: 
 	struct FLodPayloads
 	{
-		TMap<FInterchangeMeshPayLoadKey, FMeshPayload> MeshPayloadPerKey;
-		TMap<FInterchangeMeshPayLoadKey, FMeshPayload> CollisionBoxPayloadPerKey;
-		TMap<FInterchangeMeshPayLoadKey, FMeshPayload> CollisionCapsulePayloadPerKey;
-		TMap<FInterchangeMeshPayLoadKey, FMeshPayload> CollisionSpherePayloadPerKey;
-		TMap<FInterchangeMeshPayLoadKey, FMeshPayload> CollisionConvexPayloadPerKey;
+		TMap<FInterchangeMeshPayLoadKey, UE::Interchange::FMeshPayload> MeshPayloadPerKey;
+		TMap<FInterchangeMeshPayLoadKey, UE::Interchange::FMeshPayload> CollisionBoxPayloadPerKey;
+		TMap<FInterchangeMeshPayLoadKey, UE::Interchange::FMeshPayload> CollisionCapsulePayloadPerKey;
+		TMap<FInterchangeMeshPayLoadKey, UE::Interchange::FMeshPayload> CollisionSpherePayloadPerKey;
+		TMap<FInterchangeMeshPayLoadKey, UE::Interchange::FMeshPayload> CollisionConvexPayloadPerKey;
 	};
 
 	TMap<int32, FLodPayloads> PayloadsPerLodIndex;
 
-	UE_API bool AddConvexGeomFromVertices(const FImportAssetObjectParams& Arguments, const FMeshDescription& MeshDescription, const FTransform& Transform, FKAggregateGeom& AggGeom);
-	UE_API bool DecomposeConvexMesh(const FImportAssetObjectParams& Arguments, const FMeshDescription& MeshDescription, const FTransform& Transform, UBodySetup* BodySetup);
-	UE_API bool AddBoxGeomFromTris(const FMeshDescription& MeshDescription, const FTransform& Transform, FKAggregateGeom& AggGeom);
-	UE_API bool AddSphereGeomFromVertices(const FImportAssetObjectParams& Arguments, const FMeshDescription& MeshDescription, const FTransform& Transform, FKAggregateGeom& AggGeom);
-	UE_API bool AddCapsuleGeomFromVertices(const FImportAssetObjectParams& Arguments, const FMeshDescription& MeshDescription, const FTransform& Transform, FKAggregateGeom& AggGeom);
-
-	UE_API bool ImportBoxCollision(const FImportAssetObjectParams& Arguments, UStaticMesh* StaticMesh);
-	UE_API bool ImportCapsuleCollision(const FImportAssetObjectParams& Arguments, UStaticMesh* StaticMesh);
-	UE_API bool ImportSphereCollision(const FImportAssetObjectParams& Arguments, UStaticMesh* StaticMesh);
-	UE_API bool ImportConvexCollision(const FImportAssetObjectParams& Arguments, UStaticMesh* StaticMesh, const UInterchangeStaticMeshLodDataNode* LodDataNode);
-	UE_API bool ImportSockets(const FImportAssetObjectParams& Arguments, UStaticMesh* StaticMesh, const UInterchangeStaticMeshFactoryNode* FactoryNode);
-
 	UE_API void CommitMeshDescriptions(UStaticMesh& StaticMesh);
 	UE_API void BuildFromMeshDescriptions(UStaticMesh& StaticMesh);
 
diff --git a/Engine/Plugins/Interchange/Runtime/Source/Pipelines/Private/InterchangeGenericStaticMeshPipeline.cpp b/Engine/Plugins/Interchange/Runtime/Source/Pipelines/Private/InterchangeGenericStaticMeshPipeline.cpp
--- a/Engine/Plugins/Interchange/Runtime/Source/Pipelines/Private/InterchangeGenericStaticMeshPipeline.cpp
+++ b/Engine/Plugins/Interchange/Runtime/Source/Pipelines/Private/InterchangeGenericStaticMeshPipeline.cpp
@@ -546,7 +546,8 @@ UInterchangeStaticMeshLodDataNode* UInterchangeGenericMeshPipeline::CreateStatic
 	StaticMeshLodDataNode->SetOneConvexHullPerUCX(bOneConvexHullPerUCX);
 	StaticMeshLodDataNode->SetImportCollision(bCollision);
 	StaticMeshLodDataNode->SetImportCollisionType(Collision);
-	
+	StaticMeshLodDataNode->SetForceCollisionPrimitiveGeneration(bForceCollisionPrimitiveGeneration);
+
 	return StaticMeshLodDataNode;
 }
 
diff --git a/Engine/Plugins/Interchange/Runtime/Source/Pipelines/Public/InterchangeGenericMeshPipeline.h b/Engine/Plugins/Interchange/Runtime/Source/Pipelines/Public/InterchangeGenericMeshPipeline.h
--- a/Engine/Plugins/Interchange/Runtime/Source/Pipelines/Public/InterchangeGenericMeshPipeline.h
+++ b/Engine/Plugins/Interchange/Runtime/Source/Pipelines/Public/InterchangeGenericMeshPipeline.h
@@ -95,6 +95,10 @@ public:
 	/** Type used to generate a collision when no custom collisions are present in the file. */
 	UPROPERTY(EditAnywhere, BlueprintReadWrite, Category = "Static Meshes", meta = (ScriptName = "FallbackCollisionType", DisplayName = "Fallback Collision Type", SubCategory = "Collision", editcondition = "bCollision == true"))
 	EInterchangeMeshCollision Collision = EInterchangeMeshCollision::Convex18DOP;
+	
+	/** Sets whether to generate collision shapes even if the provided mesh data doesn't match the requested collision shape very well */
+	UPROPERTY(EditAnywhere, BlueprintReadWrite, Category = "Static Meshes", meta = (SubCategory = "Collision", editcondition = "bCollision == true", PipelineInternalEditionData = "True"))
+	bool bForceCollisionPrimitiveGeneration = false;
 
 	//////	Static Meshes Build settings Properties //////
 
