Interchange: Bake collider mesh transforms for prefix-style collision even when importing into level.

diff --git a/Engine/Plugins/Interchange/Runtime/Source/FactoryNodes/Private/InterchangeStaticMeshLodDataNode.cpp b/Engine/Plugins/Interchange/Runtime/Source/FactoryNodes/Private/InterchangeStaticMeshLodDataNode.cpp
--- a/Engine/Plugins/Interchange/Runtime/Source/FactoryNodes/Private/InterchangeStaticMeshLodDataNode.cpp
+++ b/Engine/Plugins/Interchange/Runtime/Source/FactoryNodes/Private/InterchangeStaticMeshLodDataNode.cpp
@@ -44,11 +44,13 @@ namespace UE
 
 UInterchangeStaticMeshLodDataNode::UInterchangeStaticMeshLodDataNode()
 {
-	MeshUids.Initialize(Attributes, UE::Interchange::FStaticMeshNodeLodDataStaticData::GetMeshUidsBaseKey().ToString());
-	BoxCollisionMeshUids.Initialize(Attributes, UE::Interchange::FStaticMeshNodeLodDataStaticData::GetBoxCollisionMeshUidsBaseKey().ToString());
-	CapsuleCollisionMeshUids.Initialize(Attributes, UE::Interchange::FStaticMeshNodeLodDataStaticData::GetCapsuleCollisionMeshUidsBaseKey().ToString());
-	SphereCollisionMeshUids.Initialize(Attributes, UE::Interchange::FStaticMeshNodeLodDataStaticData::GetSphereCollisionMeshUidsBaseKey().ToString());
-	ConvexCollisionMeshUids.Initialize(Attributes, UE::Interchange::FStaticMeshNodeLodDataStaticData::GetConvexCollisionMeshUidsBaseKey().ToString());
+	using namespace UE::Interchange;
+
+	MeshUids.Initialize(Attributes, FStaticMeshNodeLodDataStaticData::GetMeshUidsBaseKey().ToString());
+	BoxCollisionMeshUids.Initialize(Attributes.ToSharedRef(), FStaticMeshNodeLodDataStaticData::GetBoxCollisionMeshUidsBaseKey().ToString());
+	CapsuleCollisionMeshUids.Initialize(Attributes.ToSharedRef(), FStaticMeshNodeLodDataStaticData::GetCapsuleCollisionMeshUidsBaseKey().ToString());
+	SphereCollisionMeshUids.Initialize(Attributes.ToSharedRef(), FStaticMeshNodeLodDataStaticData::GetSphereCollisionMeshUidsBaseKey().ToString());
+	ConvexCollisionMeshUids.Initialize(Attributes.ToSharedRef(), FStaticMeshNodeLodDataStaticData::GetConvexCollisionMeshUidsBaseKey().ToString());
 }
 
 /**
@@ -217,24 +219,43 @@ int32 UInterchangeStaticMeshLodDataNode::GetBoxCollisionMeshUidsCount() const
 	return BoxCollisionMeshUids.GetCount();
 }
 
-void UInterchangeStaticMeshLodDataNode::GetBoxCollisionMeshUids(TArray<FString>& OutMeshNames) const
+TMap<FString, FString> UInterchangeStaticMeshLodDataNode::GetBoxCollisionMeshMap() const
+{
+	return BoxCollisionMeshUids.ToMap();
+}
+
+// Deprecated
+void UInterchangeStaticMeshLodDataNode::GetBoxCollisionMeshUids(TArray<FString>& OutColliderUids) const
 {
-	BoxCollisionMeshUids.GetItems(OutMeshNames);
+	TMap<FString, FString> Map = GetBoxCollisionMeshMap();
+	Map.GetKeys(OutColliderUids);
+}
+
+void UInterchangeStaticMeshLodDataNode::GetBoxColliderRenderMeshUid(const FString& InColliderMeshUid, FString& OutRenderMeshUid) const
+{
+	BoxCollisionMeshUids.GetValue(InColliderMeshUid, OutRenderMeshUid);
+}
+
+bool UInterchangeStaticMeshLodDataNode::AddBoxCollisionMeshUid(const FString& ColliderMeshUid)
+{
+	return false;
 }
 
-bool UInterchangeStaticMeshLodDataNode::AddBoxCollisionMeshUid(const FString& MeshName)
+bool UInterchangeStaticMeshLodDataNode::AddBoxCollisionMeshUids(const FString& ColliderMeshUid, const FString& RenderMeshUid)
 {
-	return BoxCollisionMeshUids.AddItem(MeshName);
+	return BoxCollisionMeshUids.SetKeyValue(ColliderMeshUid, RenderMeshUid);
 }
 
-bool UInterchangeStaticMeshLodDataNode::RemoveBoxCollisionMeshUid(const FString& MeshName)
+bool UInterchangeStaticMeshLodDataNode::RemoveBoxCollisionMeshUid(const FString& ColliderMeshUid)
 {
-	return BoxCollisionMeshUids.RemoveItem(MeshName);
+	return BoxCollisionMeshUids.RemoveKey(ColliderMeshUid);
 }
 
 bool UInterchangeStaticMeshLodDataNode::RemoveAllBoxCollisionMeshes()
 {
-	return BoxCollisionMeshUids.RemoveAllItems();
+	bool bHadItems = GetBoxCollisionMeshUidsCount() > 0;
+	BoxCollisionMeshUids.Empty();
+	return bHadItems;
 }
 
 int32 UInterchangeStaticMeshLodDataNode::GetCapsuleCollisionMeshUidsCount() const
@@ -242,24 +263,43 @@ int32 UInterchangeStaticMeshLodDataNode::GetCapsuleCollisionMeshUidsCount() cons
 	return CapsuleCollisionMeshUids.GetCount();
 }
 
-void UInterchangeStaticMeshLodDataNode::GetCapsuleCollisionMeshUids(TArray<FString>& OutMeshNames) const
+TMap<FString, FString> UInterchangeStaticMeshLodDataNode::GetCapsuleCollisionMeshMap() const
 {
-	CapsuleCollisionMeshUids.GetItems(OutMeshNames);
+	return CapsuleCollisionMeshUids.ToMap();
+}
+
+// Deprecated
+void UInterchangeStaticMeshLodDataNode::GetCapsuleCollisionMeshUids(TArray<FString>& OutColliderUids) const
+{
+	TMap<FString, FString> Map = GetCapsuleCollisionMeshMap();
+	Map.GetKeys(OutColliderUids);
+}
+
+void UInterchangeStaticMeshLodDataNode::GetCapsuleColliderRenderMeshUid(const FString& InColliderMeshUid, FString& OutRenderMeshUid) const
+{
+	CapsuleCollisionMeshUids.GetValue(InColliderMeshUid, OutRenderMeshUid);
+}
+
+bool UInterchangeStaticMeshLodDataNode::AddCapsuleCollisionMeshUid(const FString& ColliderMeshUid)
+{
+	return false;
 }
 
-bool UInterchangeStaticMeshLodDataNode::AddCapsuleCollisionMeshUid(const FString& MeshName)
+bool UInterchangeStaticMeshLodDataNode::AddCapsuleCollisionMeshUids(const FString& ColliderMeshUid, const FString& RenderMeshUid)
 {
-	return CapsuleCollisionMeshUids.AddItem(MeshName);
+	return CapsuleCollisionMeshUids.SetKeyValue(ColliderMeshUid, RenderMeshUid);
 }
 
-bool UInterchangeStaticMeshLodDataNode::RemoveCapsuleCollisionMeshUid(const FString& MeshName)
+bool UInterchangeStaticMeshLodDataNode::RemoveCapsuleCollisionMeshUid(const FString& ColliderMeshUid)
 {
-	return CapsuleCollisionMeshUids.RemoveItem(MeshName);
+	return CapsuleCollisionMeshUids.RemoveKey(ColliderMeshUid);
 }
 
 bool UInterchangeStaticMeshLodDataNode::RemoveAllCapsuleCollisionMeshes()
 {
-	return CapsuleCollisionMeshUids.RemoveAllItems();
+	bool bHadItems = GetCapsuleCollisionMeshUidsCount() > 0;
+	CapsuleCollisionMeshUids.Empty();
+	return bHadItems;
 }
 
 int32 UInterchangeStaticMeshLodDataNode::GetSphereCollisionMeshUidsCount() const
@@ -267,24 +307,43 @@ int32 UInterchangeStaticMeshLodDataNode::GetSphereCollisionMeshUidsCount() const
 	return SphereCollisionMeshUids.GetCount();
 }
 
-void UInterchangeStaticMeshLodDataNode::GetSphereCollisionMeshUids(TArray<FString>& OutMeshNames) const
+TMap<FString, FString> UInterchangeStaticMeshLodDataNode::GetSphereCollisionMeshMap() const
 {
-	SphereCollisionMeshUids.GetItems(OutMeshNames);
+	return SphereCollisionMeshUids.ToMap();
 }
 
-bool UInterchangeStaticMeshLodDataNode::AddSphereCollisionMeshUid(const FString& MeshName)
+// Deprecated
+void UInterchangeStaticMeshLodDataNode::GetSphereCollisionMeshUids(TArray<FString>& OutColliderUids) const
 {
-	return SphereCollisionMeshUids.AddItem(MeshName);
+	TMap<FString, FString> Map = GetSphereCollisionMeshMap();
+	Map.GetKeys(OutColliderUids);
 }
 
-bool UInterchangeStaticMeshLodDataNode::RemoveSphereCollisionMeshUid(const FString& MeshName)
+void UInterchangeStaticMeshLodDataNode::GetSphereColliderRenderMeshUid(const FString& InColliderMeshUid, FString& OutRenderMeshUid) const
 {
-	return SphereCollisionMeshUids.RemoveItem(MeshName);
+	SphereCollisionMeshUids.GetValue(InColliderMeshUid, OutRenderMeshUid);
+}
+
+bool UInterchangeStaticMeshLodDataNode::AddSphereCollisionMeshUid(const FString& ColliderMeshUid)
+{
+	return false;
+}
+
+bool UInterchangeStaticMeshLodDataNode::AddSphereCollisionMeshUids(const FString& ColliderMeshUid, const FString& RenderMeshUid)
+{
+	return SphereCollisionMeshUids.SetKeyValue(ColliderMeshUid, RenderMeshUid);
+}
+
+bool UInterchangeStaticMeshLodDataNode::RemoveSphereCollisionMeshUid(const FString& ColliderMeshUid)
+{
+	return SphereCollisionMeshUids.RemoveKey(ColliderMeshUid);
 }
 
 bool UInterchangeStaticMeshLodDataNode::RemoveAllSphereCollisionMeshes()
 {
-	return SphereCollisionMeshUids.RemoveAllItems();
+	bool bHadItems = GetSphereCollisionMeshUidsCount() > 0;
+	SphereCollisionMeshUids.Empty();
+	return bHadItems;
 }
 
 int32 UInterchangeStaticMeshLodDataNode::GetConvexCollisionMeshUidsCount() const
@@ -292,24 +351,43 @@ int32 UInterchangeStaticMeshLodDataNode::GetConvexCollisionMeshUidsCount() const
 	return ConvexCollisionMeshUids.GetCount();
 }
 
-void UInterchangeStaticMeshLodDataNode::GetConvexCollisionMeshUids(TArray<FString>& OutMeshNames) const
+TMap<FString, FString> UInterchangeStaticMeshLodDataNode::GetConvexCollisionMeshMap() const
 {
-	ConvexCollisionMeshUids.GetItems(OutMeshNames);
+	return ConvexCollisionMeshUids.ToMap();
+}
+
+// Deprecated
+void UInterchangeStaticMeshLodDataNode::GetConvexCollisionMeshUids(TArray<FString>& OutColliderUids) const
+{
+	TMap<FString, FString> Map = GetConvexCollisionMeshMap();
+	Map.GetKeys(OutColliderUids);
+}
+
+void UInterchangeStaticMeshLodDataNode::GetConvexColliderRenderMeshUid(const FString& InColliderMeshUid, FString& OutRenderMeshUid) const
+{
+	ConvexCollisionMeshUids.GetValue(InColliderMeshUid, OutRenderMeshUid);
+}
+
+bool UInterchangeStaticMeshLodDataNode::AddConvexCollisionMeshUid(const FString& ColliderMeshUid)
+{
+	return false;
 }
 
-bool UInterchangeStaticMeshLodDataNode::AddConvexCollisionMeshUid(const FString& MeshName)
+bool UInterchangeStaticMeshLodDataNode::AddConvexCollisionMeshUids(const FString& ColliderMeshUid, const FString& RenderMeshUid)
 {
-	return ConvexCollisionMeshUids.AddItem(MeshName);
+	return ConvexCollisionMeshUids.SetKeyValue(ColliderMeshUid, RenderMeshUid);
 }
 
-bool UInterchangeStaticMeshLodDataNode::RemoveConvexCollisionMeshUid(const FString& MeshName)
+bool UInterchangeStaticMeshLodDataNode::RemoveConvexCollisionMeshUid(const FString& ColliderMeshUid)
 {
-	return ConvexCollisionMeshUids.RemoveItem(MeshName);
+	return ConvexCollisionMeshUids.RemoveKey(ColliderMeshUid);
 }
 
 bool UInterchangeStaticMeshLodDataNode::RemoveAllConvexCollisionMeshes()
 {
-	return ConvexCollisionMeshUids.RemoveAllItems();
+	bool bHadItems = GetConvexCollisionMeshUidsCount() > 0;
+	ConvexCollisionMeshUids.Empty();
+	return bHadItems;
 }
 
 bool UInterchangeStaticMeshLodDataNode::GetOneConvexHullPerUCX(bool& AttributeValue) const
diff --git a/Engine/Plugins/Interchange/Runtime/Source/FactoryNodes/Public/InterchangeStaticMeshLodDataNode.h b/Engine/Plugins/Interchange/Runtime/Source/FactoryNodes/Public/InterchangeStaticMeshLodDataNode.h
--- a/Engine/Plugins/Interchange/Runtime/Source/FactoryNodes/Public/InterchangeStaticMeshLodDataNode.h
+++ b/Engine/Plugins/Interchange/Runtime/Source/FactoryNodes/Public/InterchangeStaticMeshLodDataNode.h
@@ -67,14 +67,25 @@ public:
 	UFUNCTION(BlueprintCallable, Category = "Interchange | Node | StaticMeshLodData")
 	UE_API int32 GetBoxCollisionMeshUidsCount() const;
 
+	UFUNCTION(BlueprintCallable, Category = "Interchange | Node | StaticMeshLodData")
+	UE_API TMap<FString, FString> GetBoxCollisionMeshMap() const;
+
+	UE_DEPRECATED(5.6, "No longer used: Collect the keys from GetBoxCollisionMeshMap() instead.")
 	UFUNCTION(BlueprintCallable, Category = "Interchange | Node | StaticMeshLodData")
 	UE_API void GetBoxCollisionMeshUids(TArray<FString>& OutMeshNames) const;
 
 	UFUNCTION(BlueprintCallable, Category = "Interchange | Node | StaticMeshLodData")
-	UE_API bool AddBoxCollisionMeshUid(const FString& MeshName);
+	UE_API void GetBoxColliderRenderMeshUid(const FString& InColliderMeshUid, FString& OutRenderMeshUid) const;
+
+	UE_DEPRECATED(5.6, "No longer used: Use the other overload that specifies the RenderMeshUid.")
+	UFUNCTION(BlueprintCallable, Category = "Interchange | Node | StaticMeshLodData")
+	UE_API bool AddBoxCollisionMeshUid(const FString& ColliderMeshUid);
 
 	UFUNCTION(BlueprintCallable, Category = "Interchange | Node | StaticMeshLodData")
-	UE_API bool RemoveBoxCollisionMeshUid(const FString& MeshName);
+	UE_API bool AddBoxCollisionMeshUids(const FString& ColliderMeshUid, const FString& RenderMeshUid);
+
+	UFUNCTION(BlueprintCallable, Category = "Interchange | Node | StaticMeshLodData")
+	UE_API bool RemoveBoxCollisionMeshUid(const FString& ColliderMeshUid);
 
 	UFUNCTION(BlueprintCallable, Category = "Interchange | Node | StaticMeshLodData")
 	UE_API bool RemoveAllBoxCollisionMeshes();
@@ -82,14 +93,25 @@ public:
 	UFUNCTION(BlueprintCallable, Category = "Interchange | Node | StaticMeshLodData")
 	UE_API int32 GetCapsuleCollisionMeshUidsCount() const;
 
+	UFUNCTION(BlueprintCallable, Category = "Interchange | Node | StaticMeshLodData")
+	UE_API TMap<FString, FString> GetCapsuleCollisionMeshMap() const;
+
+	UE_DEPRECATED(5.6, "No longer used: Collect the keys from GetCapsuleCollisionMeshMap() instead.")
 	UFUNCTION(BlueprintCallable, Category = "Interchange | Node | StaticMeshLodData")
 	UE_API void GetCapsuleCollisionMeshUids(TArray<FString>& OutMeshNames) const;
 
 	UFUNCTION(BlueprintCallable, Category = "Interchange | Node | StaticMeshLodData")
-	UE_API bool AddCapsuleCollisionMeshUid(const FString& MeshName);
+	UE_API void GetCapsuleColliderRenderMeshUid(const FString& InColliderMeshUid, FString& OutRenderMeshUid) const;
+
+	UE_DEPRECATED(5.6, "No longer used: Use the other overload that specifies the RenderMeshUid.")
+	UFUNCTION(BlueprintCallable, Category = "Interchange | Node | StaticMeshLodData")
+	UE_API bool AddCapsuleCollisionMeshUid(const FString& ColliderMeshUid);
 
 	UFUNCTION(BlueprintCallable, Category = "Interchange | Node | StaticMeshLodData")
-	UE_API bool RemoveCapsuleCollisionMeshUid(const FString& MeshName);
+	UE_API bool AddCapsuleCollisionMeshUids(const FString& ColliderMeshUid, const FString& RenderMeshUid);
+
+	UFUNCTION(BlueprintCallable, Category = "Interchange | Node | StaticMeshLodData")
+	UE_API bool RemoveCapsuleCollisionMeshUid(const FString& ColliderMeshUid);
 
 	UFUNCTION(BlueprintCallable, Category = "Interchange | Node | StaticMeshLodData")
 	UE_API bool RemoveAllCapsuleCollisionMeshes();
@@ -97,14 +119,25 @@ public:
 	UFUNCTION(BlueprintCallable, Category = "Interchange | Node | StaticMeshLodData")
 	UE_API int32 GetSphereCollisionMeshUidsCount() const;
 
+	UFUNCTION(BlueprintCallable, Category = "Interchange | Node | StaticMeshLodData")
+	UE_API TMap<FString, FString> GetSphereCollisionMeshMap() const;
+
+	UE_DEPRECATED(5.6, "No longer used: Collect the keys from GetSphereCollisionMeshMap() instead.")
 	UFUNCTION(BlueprintCallable, Category = "Interchange | Node | StaticMeshLodData")
 	UE_API void GetSphereCollisionMeshUids(TArray<FString>& OutMeshNames) const;
 
 	UFUNCTION(BlueprintCallable, Category = "Interchange | Node | StaticMeshLodData")
-	UE_API bool AddSphereCollisionMeshUid(const FString& MeshName);
+	UE_API void GetSphereColliderRenderMeshUid(const FString& InColliderMeshUid, FString& OutRenderMeshUid) const;
+
+	UE_DEPRECATED(5.6, "No longer used: Use the other overload that specifies the RenderMeshUid.")
+	UFUNCTION(BlueprintCallable, Category = "Interchange | Node | StaticMeshLodData")
+	UE_API bool AddSphereCollisionMeshUid(const FString& ColliderMeshUid);
 
 	UFUNCTION(BlueprintCallable, Category = "Interchange | Node | StaticMeshLodData")
-	UE_API bool RemoveSphereCollisionMeshUid(const FString& MeshName);
+	UE_API bool AddSphereCollisionMeshUids(const FString& ColliderMeshUid, const FString& RenderMeshUid);
+
+	UFUNCTION(BlueprintCallable, Category = "Interchange | Node | StaticMeshLodData")
+	UE_API bool RemoveSphereCollisionMeshUid(const FString& ColliderMeshUid);
 
 	UFUNCTION(BlueprintCallable, Category = "Interchange | Node | StaticMeshLodData")
 	UE_API bool RemoveAllSphereCollisionMeshes();
@@ -112,11 +145,22 @@ public:
 	UFUNCTION(BlueprintCallable, Category = "Interchange | Node | StaticMeshLodData")
 	UE_API int32 GetConvexCollisionMeshUidsCount() const;
 
+	UFUNCTION(BlueprintCallable, Category = "Interchange | Node | StaticMeshLodData")
+	UE_API TMap<FString, FString> GetConvexCollisionMeshMap() const;
+
+	UE_DEPRECATED(5.6, "No longer used: Collect the keys from GetConvexCollisionMeshMap() instead.")
 	UFUNCTION(BlueprintCallable, Category = "Interchange | Node | StaticMeshLodData")
 	UE_API void GetConvexCollisionMeshUids(TArray<FString>& OutMeshNames) const;
 
 	UFUNCTION(BlueprintCallable, Category = "Interchange | Node | StaticMeshLodData")
-	UE_API bool AddConvexCollisionMeshUid(const FString& MeshName);
+	UE_API void GetConvexColliderRenderMeshUid(const FString& InColliderMeshUid, FString& OutRenderMeshUid) const;
+
+	UE_DEPRECATED(5.6, "No longer used: Use the other overload that specifies the RenderMeshUid.")
+	UFUNCTION(BlueprintCallable, Category = "Interchange | Node | StaticMeshLodData")
+	UE_API bool AddConvexCollisionMeshUid(const FString& ColliderMeshUid);
+
+	UFUNCTION(BlueprintCallable, Category = "Interchange | Node | StaticMeshLodData")
+	UE_API bool AddConvexCollisionMeshUids(const FString& ColliderMeshUid, const FString& RenderMeshUid);
 
 	UFUNCTION(BlueprintCallable, Category = "Interchange | Node | StaticMeshLodData")
 	UE_API bool RemoveConvexCollisionMeshUid(const FString& MeshName);
@@ -162,11 +206,11 @@ private:
 	bool IsEditorOnlyDataDefined();
 
 	UE::Interchange::TArrayAttributeHelper<FString> MeshUids;
-	UE::Interchange::TArrayAttributeHelper<FString> BoxCollisionMeshUids;
-	UE::Interchange::TArrayAttributeHelper<FString> CapsuleCollisionMeshUids;
-	UE::Interchange::TArrayAttributeHelper<FString> SphereCollisionMeshUids;
-	UE::Interchange::TArrayAttributeHelper<FString> ConvexCollisionMeshUids;
-		
+	UE::Interchange::TMapAttributeHelper<FString, FString> BoxCollisionMeshUids;
+	UE::Interchange::TMapAttributeHelper<FString, FString> CapsuleCollisionMeshUids;
+	UE::Interchange::TMapAttributeHelper<FString, FString> SphereCollisionMeshUids;
+	UE::Interchange::TMapAttributeHelper<FString, FString> ConvexCollisionMeshUids;
+
 	IMPLEMENT_NODE_ATTRIBUTE_KEY(OneConvexHullPerUCX)
 	IMPLEMENT_NODE_ATTRIBUTE_KEY(ForceCollisionPrimitiveGeneration)
 	IMPLEMENT_NODE_ATTRIBUTE_KEY(ImportCollision)
diff --git a/Engine/Plugins/Interchange/Runtime/Source/Import/Private/Mesh/InterchangeStaticMeshFactory.cpp b/Engine/Plugins/Interchange/Runtime/Source/Import/Private/Mesh/InterchangeStaticMeshFactory.cpp
--- a/Engine/Plugins/Interchange/Runtime/Source/Import/Private/Mesh/InterchangeStaticMeshFactory.cpp
+++ b/Engine/Plugins/Interchange/Runtime/Source/Import/Private/Mesh/InterchangeStaticMeshFactory.cpp
@@ -103,43 +103,76 @@ void UInterchangeStaticMeshFactory::CreatePayloadTasks(const FImportAssetObjectP
 
 		FLodPayloads& LodPayloads = PayloadsPerLodIndex.FindOrAdd(LodIndex);
 
-		auto AddMeshPayloads = [&Arguments, MeshTranslatorPayloadInterface, bBakeMeshes, bBakePivotMeshes, &GlobalOffsetTransform, &PayloadTasks, bAsync, &PayloadAttributes]
-			(const TArray<FString>& MeshUids, TMap<FInterchangeMeshPayLoadKey, UE::Interchange::FMeshPayload>& PayloadPerKey)
+		auto GetNodeAndGlobalTransform = [&Arguments, &GlobalOffsetTransform](
+			const FString& InMeshOrSceneNodeUid,
+			bool bBakeMeshes,
+			bool bBakePivotMeshes,
+			const UInterchangeMeshNode*& OutMeshNode,
+			FTransform& OutGlobalTransform
+		)
+		{
+			const UInterchangeBaseNode* Node = Arguments.NodeContainer->GetNode(InMeshOrSceneNodeUid);
+			OutMeshNode = Cast<const UInterchangeMeshNode>(Node);
+			if (OutMeshNode == nullptr)
+			{
+				// MeshUid must refer to a scene node
+				const UInterchangeSceneNode* SceneNode = Cast<const UInterchangeSceneNode>(Node);
+				if (!ensure(SceneNode))
+				{
+					UE_LOG(LogInterchangeImport, Warning, TEXT("Invalid LOD mesh reference when importing StaticMesh asset %s."), *Arguments.AssetName);
+					return;
+				}
+
+				if (bBakeMeshes)
+				{
+					// Get the transform from the scene node
+					FTransform SceneNodeGlobalTransform;
+					if (SceneNode->GetCustomGlobalTransform(Arguments.NodeContainer, GlobalOffsetTransform, SceneNodeGlobalTransform))
+					{
+						OutGlobalTransform = SceneNodeGlobalTransform;
+					}
+				}
+				UE::Interchange::Private::MeshHelper::AddSceneNodeGeometricAndPivotToGlobalTransform(OutGlobalTransform, SceneNode, bBakeMeshes, bBakePivotMeshes);
+				// And get the mesh node which it references
+				FString MeshDependencyUid;
+				SceneNode->GetCustomAssetInstanceUid(MeshDependencyUid);
+				OutMeshNode = Cast<UInterchangeMeshNode>(Arguments.NodeContainer->GetNode(MeshDependencyUid));
+			}
+			else if (bBakeMeshes)
 			{
+				//If we have a mesh that is not reference by a scene node, we must apply the global offset.
+				OutGlobalTransform = GlobalOffsetTransform;
+			}
+		};
+
+		auto AddMeshPayloads = [&Arguments, MeshTranslatorPayloadInterface, bBakeMeshes, bBakePivotMeshes, &GlobalOffsetTransform, &PayloadTasks, bAsync, &PayloadAttributes, &GetNodeAndGlobalTransform]
+			(const TArray<FString>& MeshUids, TMap<FInterchangeMeshPayLoadKey, UE::Interchange::FMeshPayload>& PayloadPerKey, const TMap<FString, FString>* ColliderToRenderUid = nullptr)
+			{
+				PayloadPerKey.Reserve(MeshUids.Num());
+
+				// Collider meshes are always "baked into their render meshes" in some way, whether we're importing assets or levels
+				bool bBakeTheseMeshes = bBakeMeshes || ColliderToRenderUid;
+				bool bBakeThesePivotMeshes = bBakePivotMeshes || ColliderToRenderUid;
+
 				for (const FString& MeshUid : MeshUids)
 				{
 					FTransform GlobalMeshTransform;
-					const UInterchangeBaseNode* Node = Arguments.NodeContainer->GetNode(MeshUid);
-					const UInterchangeMeshNode* MeshNode = Cast<const UInterchangeMeshNode>(Node);
-					if (MeshNode == nullptr)
+					const UInterchangeMeshNode* MeshNode = nullptr;
+					GetNodeAndGlobalTransform(MeshUid, bBakeTheseMeshes, bBakeThesePivotMeshes, MeshNode, GlobalMeshTransform);
+
+					if (!bBakeMeshes && ColliderToRenderUid)
 					{
-						// MeshUid must refer to a scene node
-						const UInterchangeSceneNode* SceneNode = Cast<const UInterchangeSceneNode>(Node);
-						if (!ensure(SceneNode))
-						{
-							UE_LOG(LogInterchangeImport, Warning, TEXT("Invalid LOD mesh reference when importing StaticMesh asset %s."), *Arguments.AssetName);
-							continue;
-						}
+						const FString& ColliderMeshUid = MeshUid;
+						const FTransform& ColliderLocalToGlobal = GlobalMeshTransform;
 
-						if (bBakeMeshes)
+						if (const FString* RenderMeshUid = ColliderToRenderUid->Find(ColliderMeshUid))
 						{
-							// Get the transform from the scene node
-							FTransform SceneNodeGlobalTransform;
-							if (SceneNode->GetCustomGlobalTransform(Arguments.NodeContainer, GlobalOffsetTransform, SceneNodeGlobalTransform))
-							{
-								GlobalMeshTransform = SceneNodeGlobalTransform;
-							}
+							FTransform RenderLocalToGlobal;
+							const UInterchangeMeshNode* RenderMeshNode = nullptr;
+							GetNodeAndGlobalTransform(*RenderMeshUid, bBakeTheseMeshes, bBakeThesePivotMeshes, RenderMeshNode, RenderLocalToGlobal);
+
+							GlobalMeshTransform = ColliderLocalToGlobal * RenderLocalToGlobal.Inverse();
 						}
-						UE::Interchange::Private::MeshHelper::AddSceneNodeGeometricAndPivotToGlobalTransform(GlobalMeshTransform, SceneNode, bBakeMeshes, bBakePivotMeshes);
-						// And get the mesh node which it references
-						FString MeshDependencyUid;
-						SceneNode->GetCustomAssetInstanceUid(MeshDependencyUid);
-						MeshNode = Cast<UInterchangeMeshNode>(Arguments.NodeContainer->GetNode(MeshDependencyUid));
-					}
-					else if (bBakeMeshes)
-					{
-						//If we have a mesh that is not reference by a scene node, we must apply the global offset.
-						GlobalMeshTransform = GlobalOffsetTransform;
 					}
 
 					if (!ensure(MeshNode))
@@ -181,30 +214,29 @@ void UInterchangeStaticMeshFactory::CreatePayloadTasks(const FImportAssetObjectP
 
 		TArray<FString> MeshUids;
 		LodDataNode->GetMeshUids(MeshUids);
-		LodPayloads.MeshPayloadPerKey.Reserve(MeshUids.Num());
 		AddMeshPayloads(MeshUids, LodPayloads.MeshPayloadPerKey);
 
 		if (LodIndex == 0)
 		{
+			TMap<FString, FString> BoxColliderToRenderUid = LodDataNode->GetBoxCollisionMeshMap();
 			TArray<FString> BoxCollisionMeshUids;
-			LodDataNode->GetBoxCollisionMeshUids(BoxCollisionMeshUids);
-			LodPayloads.CollisionBoxPayloadPerKey.Reserve(BoxCollisionMeshUids.Num());
-			AddMeshPayloads(BoxCollisionMeshUids, LodPayloads.CollisionBoxPayloadPerKey);
+			BoxColliderToRenderUid.GetKeys(BoxCollisionMeshUids);
+			AddMeshPayloads(BoxCollisionMeshUids, LodPayloads.CollisionBoxPayloadPerKey, &BoxColliderToRenderUid);
 
+			TMap<FString, FString> CapsuleColliderToRenderUid = LodDataNode->GetCapsuleCollisionMeshMap();
 			TArray<FString> CapsuleCollisionMeshUids;
-			LodDataNode->GetCapsuleCollisionMeshUids(CapsuleCollisionMeshUids);
-			LodPayloads.CollisionCapsulePayloadPerKey.Reserve(CapsuleCollisionMeshUids.Num());
-			AddMeshPayloads(CapsuleCollisionMeshUids, LodPayloads.CollisionCapsulePayloadPerKey);
+			CapsuleColliderToRenderUid.GetKeys(CapsuleCollisionMeshUids);
+			AddMeshPayloads(CapsuleCollisionMeshUids, LodPayloads.CollisionCapsulePayloadPerKey, &CapsuleColliderToRenderUid);
 
+			TMap<FString, FString> SphereColliderToRenderUid = LodDataNode->GetSphereCollisionMeshMap();
 			TArray<FString> SphereCollisionMeshUids;
-			LodDataNode->GetSphereCollisionMeshUids(SphereCollisionMeshUids);
-			LodPayloads.CollisionSpherePayloadPerKey.Reserve(SphereCollisionMeshUids.Num());
-			AddMeshPayloads(SphereCollisionMeshUids, LodPayloads.CollisionSpherePayloadPerKey);
+			SphereColliderToRenderUid.GetKeys(SphereCollisionMeshUids);
+			AddMeshPayloads(SphereCollisionMeshUids, LodPayloads.CollisionSpherePayloadPerKey, &SphereColliderToRenderUid);
 
+			TMap<FString, FString> ConvexColliderToRenderUid = LodDataNode->GetConvexCollisionMeshMap();
 			TArray<FString> ConvexCollisionMeshUids;
-			LodDataNode->GetConvexCollisionMeshUids(ConvexCollisionMeshUids);
-			LodPayloads.CollisionConvexPayloadPerKey.Reserve(ConvexCollisionMeshUids.Num());
-			AddMeshPayloads(ConvexCollisionMeshUids, LodPayloads.CollisionConvexPayloadPerKey);
+			ConvexColliderToRenderUid.GetKeys(ConvexCollisionMeshUids);
+			AddMeshPayloads(ConvexCollisionMeshUids, LodPayloads.CollisionConvexPayloadPerKey, &ConvexColliderToRenderUid);
 		}
 	}
 }
diff --git a/Engine/Plugins/Interchange/Runtime/Source/Pipelines/Private/InterchangeGenericScenesPipeline.cpp b/Engine/Plugins/Interchange/Runtime/Source/Pipelines/Private/InterchangeGenericScenesPipeline.cpp
--- a/Engine/Plugins/Interchange/Runtime/Source/Pipelines/Private/InterchangeGenericScenesPipeline.cpp
+++ b/Engine/Plugins/Interchange/Runtime/Source/Pipelines/Private/InterchangeGenericScenesPipeline.cpp
@@ -547,6 +547,24 @@ void UInterchangeGenericLevelPipeline::ExecuteSceneNodePreImport(const FTransfor
 		}
 	}
 
+	// Don't generate actor factory nodes for static mesh scene nodes that haven't produced any static mesh factory nodes,
+	// as we're not going to have any static mesh to assign to the component anyway!
+	//
+	// We currently use this to prevent the generation of static mesh components for collision mesh nodes, when importing into level.
+	// In that case we won't produce a factory node for the collider mesh node, which we can detect here and early out
+	if (const UInterchangeMeshNode* MeshNode = Cast<UInterchangeMeshNode>(TranslatedAssetNode))
+	{
+		if (!MeshNode->IsSkinnedMesh())
+		{
+			const FString MeshFactoryNodeUid = UInterchangeFactoryBaseNode::BuildFactoryNodeUid(MeshNode->GetUniqueID());
+			const UInterchangeStaticMeshFactoryNode* MeshFactoryNode = Cast<UInterchangeStaticMeshFactoryNode>(BaseNodeContainer->GetFactoryNode(MeshFactoryNodeUid));
+			if (!MeshFactoryNode)
+			{
+				return;
+			}
+		}
+	}
+
 	const bool bIsLodGroup = [&SceneNode]()
 		{
 			TArray<FString> SpecializeTypes;
diff --git a/Engine/Plugins/Interchange/Runtime/Source/Pipelines/Private/InterchangeGenericStaticMeshPipeline.cpp b/Engine/Plugins/Interchange/Runtime/Source/Pipelines/Private/InterchangeGenericStaticMeshPipeline.cpp
--- a/Engine/Plugins/Interchange/Runtime/Source/Pipelines/Private/InterchangeGenericStaticMeshPipeline.cpp
+++ b/Engine/Plugins/Interchange/Runtime/Source/Pipelines/Private/InterchangeGenericStaticMeshPipeline.cpp
@@ -629,23 +629,26 @@ void UInterchangeGenericMeshPipeline::AddLodDataToStaticMesh(UInterchangeStaticM
 			if (bImportCollisionAccordingToMeshName)
 			{
 				TTuple<EInterchangeMeshCollision, FString> MeshType = GetCollisionMeshType(PipelineMeshesUtilities, *BaseNodeContainer, NodeUid, NodeUids, bImportCollisionAccordingToMeshName);
+				EInterchangeMeshCollision CollisionType = MeshType.Key;
+				const FString& RenderMeshUid = MeshType.Value;
+				const FString& ColliderMeshUid = NodeUid;
 
-				switch (MeshType.Get<0>())
+				switch (CollisionType)
 				{
 				case EInterchangeMeshCollision::None:
 					LodDataNode->AddMeshUid(NodeUid);
 					break;
 
 				case EInterchangeMeshCollision::Box:
-					LodDataNode->AddBoxCollisionMeshUid(NodeUid);
+					LodDataNode->AddBoxCollisionMeshUids(ColliderMeshUid, RenderMeshUid);
 					break;
 
 				case EInterchangeMeshCollision::Sphere:
-					LodDataNode->AddSphereCollisionMeshUid(NodeUid);
+					LodDataNode->AddSphereCollisionMeshUids(ColliderMeshUid, RenderMeshUid);
 					break;
 
 				case EInterchangeMeshCollision::Capsule:
-					LodDataNode->AddCapsuleCollisionMeshUid(NodeUid);
+					LodDataNode->AddCapsuleCollisionMeshUids(ColliderMeshUid, RenderMeshUid);
 					break;
 
 				case EInterchangeMeshCollision::Convex10DOP_X:
@@ -653,13 +656,13 @@ void UInterchangeGenericMeshPipeline::AddLodDataToStaticMesh(UInterchangeStaticM
 				case EInterchangeMeshCollision::Convex10DOP_Z:
 				case EInterchangeMeshCollision::Convex18DOP:
 				case EInterchangeMeshCollision::Convex26DOP:
-					LodDataNode->AddConvexCollisionMeshUid(NodeUid);
+					LodDataNode->AddConvexCollisionMeshUids(ColliderMeshUid, RenderMeshUid);
 					break;
 				}
 
 				// If the Mesh is both render AND collision mesh let's also add it as a render mesh,
 				// as the above switch will have only added the collision mesh to the LodDataNode
-				if (MeshType.Get<1>() == NodeUid && MeshType.Get<0>() != EInterchangeMeshCollision::None)
+				if (RenderMeshUid == NodeUid && CollisionType != EInterchangeMeshCollision::None)
 				{
 					LodDataNode->AddMeshUid(NodeUid);
 				}
diff --git a/Engine/Source/Runtime/Interchange/Core/Public/Nodes/InterchangeBaseNodeUtilities.h b/Engine/Source/Runtime/Interchange/Core/Public/Nodes/InterchangeBaseNodeUtilities.h
--- a/Engine/Source/Runtime/Interchange/Core/Public/Nodes/InterchangeBaseNodeUtilities.h
+++ b/Engine/Source/Runtime/Interchange/Core/Public/Nodes/InterchangeBaseNodeUtilities.h
@@ -450,6 +450,11 @@ public:
 		EmptyInternal(NumOfExpectedElements);
 	}
 
+	int32 GetCount() const
+	{
+		return CachedKeysAndValues.Num();
+	}
+
 	TMapAttributeHelper& operator=(const TMap<KeyType, ValueType>& InMap)
 	{
 		TSharedPtr<FAttributeStorage, ESPMode::ThreadSafe> AttributesPtr = Attributes.Pin();
