Interchange:
- Fix Material Instances created even though material import is disabe in the pipelien

diff --git a/Engine/Plugins/Interchange/Runtime/Source/Pipelines/Private/InterchangeGenericMaterialPipeline.cpp b/Engine/Plugins/Interchange/Runtime/Source/Pipelines/Private/InterchangeGenericMaterialPipeline.cpp
--- a/Engine/Plugins/Interchange/Runtime/Source/Pipelines/Private/InterchangeGenericMaterialPipeline.cpp
+++ b/Engine/Plugins/Interchange/Runtime/Source/Pipelines/Private/InterchangeGenericMaterialPipeline.cpp
@@ -919,6 +919,7 @@ void UInterchangeGenericMaterialPipeline::ExecutePipeline(UInterchangeBaseNodeCo
 		{
 			if (UInterchangeMaterialInstanceFactoryNode* MaterialInstanceFactoryNode = CreateMaterialInstanceFactoryNode(ShaderGraphNode))
 			{
+				BaseMaterialFactoryNodes.Add(MaterialInstanceFactoryNode);
 				MaterialInstanceFactoryNode->SetEnabled(bImportUnusedMaterial);
 			}
 		}
