Interchange MaterialX:
- Add support for mix of surface shaders either directly in a <surfacematerial> node or in the shader graph itself
- Add support of surfaceshaders inside the shader graph, previously it was only possible to create a surfaceshader (StandardSurface, OpenPBR, etc.) as a direct child of a <surfacematerial> now it is also possible to have these nodes in the graph itsef just like any other nodes
- Use Active outputs when retrieving the output of a node instead of just "output", because the latter function can fail and return no output, for example, it happend for the standardsurface node, which has different versions, with the active output we ensure to take the last version with the correct output

USD MaterialX:
- Take the input type from the nodedef when we couldn't find a valid candidate, happens for MaterialX type "surfaceshader" written as "token" in USD
- Remove the input of a node, if this input has no connection and should take the "defaultgeomprop" (hence has no value nor nodename)

diff --git a/Engine/Plugins/Interchange/Runtime/Source/Common/Public/MaterialX/InterchangeMaterialXDefinitions.h b/Engine/Plugins/Interchange/Runtime/Source/Common/Public/MaterialX/InterchangeMaterialXDefinitions.h
--- a/Engine/Plugins/Interchange/Runtime/Source/Common/Public/MaterialX/InterchangeMaterialXDefinitions.h
+++ b/Engine/Plugins/Interchange/Runtime/Source/Common/Public/MaterialX/InterchangeMaterialXDefinitions.h
@@ -846,6 +846,7 @@ MATERIALX_NAMESPACE_BEGIN
 		static constexpr const char* Surface = "ND_surface";
 		static constexpr const char* DisplacementFloat = "ND_displacement_float";
 		static constexpr const char* DisplacementVector3 = "ND_displacement_vector3";
+		static constexpr const char* MixSurfaceShader = "ND_mix_surfaceshader";
 	}
 
 	namespace Library
@@ -883,6 +884,10 @@ MATERIALX_NAMESPACE_BEGIN
 		static constexpr const char* Vector4Array = "vector4array";
 		static constexpr const char* StringArray = "stringarray";
 		static constexpr const char* GeomNameArray = "geomnamearray";
+		static constexpr const char* BSDF = "BSDF";
+		static constexpr const char* EDF = "EDF";
+		static constexpr const char* VDF = "VDF";
+
 	}
 
 	namespace NodeGroup
@@ -910,6 +915,7 @@ MATERIALX_NAMESPACE_BEGIN
 			static const T NeutralOne = T{ 1 };
 		}
 
+		static constexpr const char* Compositing = "compositing";
 		static constexpr const char* PBR = "pbr";
 	}
 
diff --git a/Engine/Plugins/Interchange/Runtime/Source/Import/Private/MaterialX/MaterialXUtils/MaterialXManager.cpp b/Engine/Plugins/Interchange/Runtime/Source/Import/Private/MaterialX/MaterialXUtils/MaterialXManager.cpp
--- a/Engine/Plugins/Interchange/Runtime/Source/Import/Private/MaterialX/MaterialXUtils/MaterialXManager.cpp
+++ b/Engine/Plugins/Interchange/Runtime/Source/Import/Private/MaterialX/MaterialXUtils/MaterialXManager.cpp
@@ -19,6 +19,7 @@
 #include "MaterialX/MaterialXUtils/MaterialXSpotLightShader.h"
 #include "MaterialX/MaterialXUtils/MaterialXSurfaceMaterial.h"
 #include "MaterialX/MaterialXUtils/MaterialXDisplacementShader.h"
+#include "MaterialX/MaterialXUtils/MaterialXMixShader.h"
 
 #include "InterchangeMaterialDefinitions.h"
 #include "InterchangeImportLog.h"
@@ -201,12 +202,23 @@ FMaterialXManager::FMaterialXManager()
 		{mx::Category::StandardSurface, FMaterialXManager::FOnGetMaterialXInstance::CreateStatic(&FMaterialXStandardSurfaceShader::MakeInstance)},
 		{mx::Category::UsdPreviewSurface, FMaterialXManager::FOnGetMaterialXInstance::CreateStatic(&FMaterialXUsdPreviewSurfaceShader::MakeInstance)},
 		{mx::Category::Displacement, FMaterialXManager::FOnGetMaterialXInstance::CreateStatic(&FMaterialXDisplacementShader::MakeInstance)},
+		{mx::Category::Mix, FMaterialXManager::FOnGetMaterialXInstance::CreateStatic(&FMaterialXMixShader::MakeInstance)},
 		{mx::Category::PointLight, FMaterialXManager::FOnGetMaterialXInstance::CreateStatic(&FMaterialXPointLightShader::MakeInstance)},
 		{mx::Category::DirectionalLight, FMaterialXManager::FOnGetMaterialXInstance::CreateStatic(&FMaterialXDirectionalLightShader::MakeInstance)},
 		{mx::Category::SpotLight, FMaterialXManager::FOnGetMaterialXInstance::CreateStatic(&FMaterialXSpotLightShader::MakeInstance)},
 		{mx::SURFACE_MATERIAL_NODE_STRING.c_str(), FMaterialXManager::FOnGetMaterialXInstance::CreateStatic(&FMaterialXSurfaceMaterial::MakeInstance)}
 	}
 	, MatchingMaterialFunctions {
+
+		// In case of StandardSurface and OpenPBR we set by default the value to the opaque shader, we decide later if we need to change for the transmission shader
+		// Surface Shader nodes
+		{mx::Category::StandardSurface,		    FMaterialXMaterialFunction{TInPlaceType<EInterchangeMaterialXShaders>{}, EInterchangeMaterialXShaders::StandardSurface}},
+		{mx::Category::OpenPBRSurface,		    FMaterialXMaterialFunction{TInPlaceType<EInterchangeMaterialXShaders>{}, EInterchangeMaterialXShaders::OpenPBRSurface}},
+		{mx::Category::UsdPreviewSurface,		FMaterialXMaterialFunction{TInPlaceType<EInterchangeMaterialXShaders>{}, EInterchangeMaterialXShaders::UsdPreviewSurface}},
+		{mx::Category::Surface,					FMaterialXMaterialFunction{TInPlaceType<EInterchangeMaterialXShaders>{}, EInterchangeMaterialXShaders::Surface}},
+		{mx::Category::SurfaceUnlit,			FMaterialXMaterialFunction{TInPlaceType<EInterchangeMaterialXShaders>{}, EInterchangeMaterialXShaders::SurfaceUnlit}},
+		{mx::Category::Displacement,			FMaterialXMaterialFunction{TInPlaceType<EInterchangeMaterialXShaders>{}, EInterchangeMaterialXShaders::Displacement}},
+
 		// BSDF Nodes
 		{mx::Category::BurleyDiffuseBSDF,		FMaterialXMaterialFunction{TInPlaceType<EInterchangeMaterialXBSDF>{}, EInterchangeMaterialXBSDF::BurleyDiffuse}},
 		{mx::Category::ConductorBSDF,			FMaterialXMaterialFunction{TInPlaceType<EInterchangeMaterialXBSDF>{}, EInterchangeMaterialXBSDF::Conductor}},
@@ -317,17 +329,22 @@ FMaterialXManager::FMaterialXManager()
 		MatchingMaterialExpressions.Add({ mx::Category::Layer }, UE::Expressions::Names::SubstrateVerticalLayering);
 
 		// Horizontal Layering, MaterialX names it mix so we have to be careful not to take the lerp expression
-		MatchingInputNames.Add({ FKeyCategoryNodegroup{mx::Category::Mix, mx::NodeGroup::PBR}, TEXT("bg") }, ExpressionInput(UE::Expressions::Inputs::Background));
-		MatchingInputNames.Add({ FKeyCategoryNodegroup{mx::Category::Mix, mx::NodeGroup::PBR}, TEXT("fg") }, ExpressionInput(UE::Expressions::Inputs::Foreground));
-		MatchingInputNames.Add({ FKeyCategoryNodegroup{mx::Category::Mix, mx::NodeGroup::PBR}, TEXT("mix") }, ExpressionInput(UE::Expressions::Inputs::Mix));
-		MatchingMaterialExpressions.Add({ mx::Category::Mix, mx::NodeGroup::PBR }, UE::Expressions::Names::SubstrateHorizontalMixing);
+		MatchingInputNames.Add({ FKeyExpression{mx::Category::Mix, mx::NodeGroup::PBR, mx::Type::BSDF}, TEXT("bg") }, ExpressionInput(UE::Expressions::Inputs::Background));
+		MatchingInputNames.Add({ FKeyExpression{mx::Category::Mix, mx::NodeGroup::PBR, mx::Type::BSDF}, TEXT("fg") }, ExpressionInput(UE::Expressions::Inputs::Foreground));
+		MatchingInputNames.Add({ FKeyExpression{mx::Category::Mix, mx::NodeGroup::PBR, mx::Type::BSDF}, TEXT("mix") }, ExpressionInput(UE::Expressions::Inputs::Mix));
+		MatchingMaterialExpressions.Add({ mx::Category::Mix, mx::NodeGroup::PBR, mx::Type::BSDF}, UE::Expressions::Names::SubstrateHorizontalMixing);
+
+		MatchingInputNames.Add({ FKeyExpression{mx::Category::Mix, mx::NodeGroup::Compositing, mx::Type::SurfaceShader}, TEXT("bg") }, ExpressionInput(UE::Expressions::Inputs::Background));
+		MatchingInputNames.Add({ FKeyExpression{mx::Category::Mix, mx::NodeGroup::Compositing, mx::Type::SurfaceShader}, TEXT("fg") }, ExpressionInput(UE::Expressions::Inputs::Foreground));
+		MatchingInputNames.Add({ FKeyExpression{mx::Category::Mix, mx::NodeGroup::Compositing, mx::Type::SurfaceShader}, TEXT("mix") }, ExpressionInput(UE::Expressions::Inputs::Mix));
+		MatchingMaterialExpressions.Add({ mx::Category::Mix, mx::NodeGroup::Compositing, mx::Type::SurfaceShader }, UE::Expressions::Names::SubstrateHorizontalMixing);
 
 		// Add
-		MatchingMaterialExpressions.Add({ mx::Category::Add, mx::NodeGroup::PBR }, UE::Expressions::Names::SubstrateAdd);
+		MatchingMaterialExpressions.Add({ mx::Category::Add, mx::NodeGroup::PBR, mx::Type::BSDF}, UE::Expressions::Names::SubstrateAdd);
 
 		// Multiply
-		MatchingInputNames.Add({ FKeyCategoryNodegroup{mx::Category::Multiply, mx::NodeGroup::PBR}, TEXT("in2") }, ExpressionInput(UE::Expressions::Inputs::Weight));
-		MatchingMaterialExpressions.Add({ mx::Category::Multiply, mx::NodeGroup::PBR }, UE::Expressions::Names::SubstrateWeight);
+		MatchingInputNames.Add({ FKeyExpression{mx::Category::Multiply, mx::NodeGroup::PBR, mx::Type::BSDF}, TEXT("in2") }, ExpressionInput(UE::Expressions::Inputs::Weight));
+		MatchingMaterialExpressions.Add({ mx::Category::Multiply, mx::NodeGroup::PBR, mx::Type::BSDF}, UE::Expressions::Names::SubstrateWeight);
 	}
 }
 
@@ -540,9 +557,9 @@ bool FMaterialXManager::Translate(MaterialX::DocumentPtr Document, UInterchangeB
 	return bIsDocumentValid;
 }
 
-const FString* FMaterialXManager::FindMatchingInput(const FString& CategoryKey, const FString& InputKey, const FString& NodeGroup) const
+const FString* FMaterialXManager::FindMatchingInput(const FString& CategoryKey, const FString& InputKey, const FString& NodeGroup, const FString& Type) const
 {
-	return MatchingInputNames.Find({ FKeyCategoryNodegroup{ CategoryKey, NodeGroup}, InputKey });
+	return MatchingInputNames.Find({ FKeyExpression{ CategoryKey, NodeGroup, Type}, InputKey });
 }
 
 const FString* FMaterialXManager::FindMaterialExpressionInput(const FString& InputKey) const
@@ -550,9 +567,9 @@ const FString* FMaterialXManager::FindMaterialExpressionInput(const FString& Inp
 	return MaterialExpressionInputs.Find(InputKey);
 }
 
-const FString* FMaterialXManager::FindMatchingMaterialExpression(const FString& CategoryKey, const FString& NodeGroup) const
+const FString* FMaterialXManager::FindMatchingMaterialExpression(const FString& CategoryKey, const FString& NodeGroup, const FString& Type) const
 {
-	return MatchingMaterialExpressions.Find({ CategoryKey, NodeGroup });
+	return MatchingMaterialExpressions.Find({ CategoryKey, NodeGroup, Type });
 }
 
 bool FMaterialXManager::FindMatchingMaterialFunction(const FString& CategoryKey, const FString*& MaterialFunctionPath, uint8& EnumType, uint8& EnumValue) const
diff --git a/Engine/Plugins/Interchange/Runtime/Source/Import/Private/MaterialX/MaterialXUtils/MaterialXMixShader.cpp b/Engine/Plugins/Interchange/Runtime/Source/Import/Private/MaterialX/MaterialXUtils/MaterialXMixShader.cpp
new file mode 100644
--- /dev/null
+++ b/Engine/Plugins/Interchange/Runtime/Source/Import/Private/MaterialX/MaterialXUtils/MaterialXMixShader.cpp
@@ -0,0 +1,59 @@
+// Copyright Epic Games, Inc. All Rights Reserved.
+
+#if WITH_EDITOR
+#include "MaterialXMixShader.h"
+#include "InterchangeImportLog.h"
+
+#include "Engine/EngineTypes.h"
+
+namespace mx = MaterialX;
+
+FMaterialXMixShader::FMaterialXMixShader(UInterchangeBaseNodeContainer& BaseNodeContainer)
+	: FMaterialXSurfaceShaderAbstract(BaseNodeContainer)
+{
+	// TODO: When we'll support volume materials we should also handle this case
+	NodeDefinition = mx::NodeDefinition::MixSurfaceShader;
+}
+
+TSharedRef<FMaterialXBase> FMaterialXMixShader::MakeInstance(UInterchangeBaseNodeContainer& BaseNodeContainer)
+{
+	TSharedRef<FMaterialXMixShader> Result = MakeShared<FMaterialXMixShader>(BaseNodeContainer);
+	Result->RegisterConnectNodeOutputToInputDelegates();
+	return Result;
+}
+
+UInterchangeBaseNode* FMaterialXMixShader::Translate(mx::NodePtr MixNode)
+{
+	using namespace UE::Interchange::Materials;
+	this->SurfaceShaderNode = MixNode;
+
+	// we need to rename the input of this mix node to the ones used by Substrate Mix
+	if (mx::InputPtr Input = MixNode->getInput("fg"))
+	{
+		SetAttributeNewName(Input, TCHAR_TO_ANSI(UE::Expressions::Inputs::Foreground));
+	}
+	if (mx::InputPtr Input = MixNode->getInput("bg"))
+	{
+		SetAttributeNewName(Input, TCHAR_TO_ANSI(UE::Expressions::Inputs::Background));
+	}
+	if (mx::InputPtr Input = MixNode->getInput("mix"))
+	{
+		SetAttributeNewName(Input, TCHAR_TO_ANSI(UE::Expressions::Inputs::Mix));
+	}
+
+	UInterchangeShaderNode* MixShaderNode = CreateShaderNode(MixNode, MixNode->getName().c_str(), UE::Expressions::Names::SubstrateHorizontalMixing);
+
+	MixShaderNode = FMaterialXSurfaceShaderAbstract::Translate(MixShaderNode);
+
+	if(bIsSubstrateEnabled)
+	{
+		UInterchangeShaderPortsAPI::ConnectDefaultOuputToInput(ShaderGraphNode, SubstrateMaterial::Parameters::FrontMaterial.ToString(), MixShaderNode->GetUniqueID());
+	}
+	else
+	{
+		UE_LOG(LogInterchangeImport, Error, TEXT("<%s> : mix of surface shaders is only available with Substrate."), UTF8_TO_TCHAR(MixNode->getName().c_str()));
+	}
+
+	return MixShaderNode;
+}
+#endif
\ No newline at end of file
diff --git a/Engine/Plugins/Interchange/Runtime/Source/Import/Private/MaterialX/MaterialXUtils/MaterialXMixShader.h b/Engine/Plugins/Interchange/Runtime/Source/Import/Private/MaterialX/MaterialXUtils/MaterialXMixShader.h
new file mode 100644
--- /dev/null
+++ b/Engine/Plugins/Interchange/Runtime/Source/Import/Private/MaterialX/MaterialXUtils/MaterialXMixShader.h
@@ -0,0 +1,23 @@
+// Copyright Epic Games, Inc. All Rights Reserved.
+
+#pragma once
+
+#if WITH_EDITOR
+#include "MaterialX/MaterialXUtils/MaterialXSurfaceShaderAbstract.h"
+
+class FMaterialXMixShader : public FMaterialXSurfaceShaderAbstract
+{
+protected:
+	
+	template<typename T, ESPMode>
+	friend class SharedPointerInternals::TIntrusiveReferenceController;
+
+	FMaterialXMixShader(UInterchangeBaseNodeContainer & BaseNodeContainer);
+
+public:
+
+	static TSharedRef<FMaterialXBase> MakeInstance(UInterchangeBaseNodeContainer& BaseNodeContainer);
+
+	virtual UInterchangeBaseNode* Translate(MaterialX::NodePtr MixNode) override;
+};
+#endif
\ No newline at end of file
diff --git a/Engine/Plugins/Interchange/Runtime/Source/Import/Private/MaterialX/MaterialXUtils/MaterialXSurfaceMaterial.cpp b/Engine/Plugins/Interchange/Runtime/Source/Import/Private/MaterialX/MaterialXUtils/MaterialXSurfaceMaterial.cpp
--- a/Engine/Plugins/Interchange/Runtime/Source/Import/Private/MaterialX/MaterialXUtils/MaterialXSurfaceMaterial.cpp
+++ b/Engine/Plugins/Interchange/Runtime/Source/Import/Private/MaterialX/MaterialXUtils/MaterialXSurfaceMaterial.cpp
@@ -99,7 +99,7 @@ UInterchangeBaseNode* FMaterialXSurfaceMaterial::Translate(MaterialX::NodePtr Su
 	UInterchangeShaderNode* DisplacementShaderNode = DisplacementShaderTranslator ? Cast<UInterchangeShaderNode>(DisplacementShaderTranslator->Translate(ConnectedDisplacementShaderNode)) : nullptr;
 
 	//Now connect the displacementshader material function normal output to the input of the surfaceshader material function
-	if(SurfaceShaderNode && DisplacementShaderNode)
+	if(SurfaceShaderNode && DisplacementShaderNode && InputNormal)
 	{
 		UInterchangeShaderPortsAPI::ConnectOuputToInputByName(SurfaceShaderNode, InputNormalName, DisplacementShaderNode->GetUniqueID(), PBRMR::Parameters::Normal.ToString());
 	}
diff --git a/Engine/Plugins/Interchange/Runtime/Source/Import/Private/MaterialX/MaterialXUtils/MaterialXSurfaceShaderAbstract.cpp b/Engine/Plugins/Interchange/Runtime/Source/Import/Private/MaterialX/MaterialXUtils/MaterialXSurfaceShaderAbstract.cpp
--- a/Engine/Plugins/Interchange/Runtime/Source/Import/Private/MaterialX/MaterialXUtils/MaterialXSurfaceShaderAbstract.cpp
+++ b/Engine/Plugins/Interchange/Runtime/Source/Import/Private/MaterialX/MaterialXUtils/MaterialXSurfaceShaderAbstract.cpp
@@ -328,7 +328,7 @@ bool FMaterialXSurfaceShaderAbstract::ConnectMatchingNodeOutputToInput(const FCo
 
 	// First search a matching Material Expression
 	// search for a Material Expression based on the node group (essentially used for Substrate Mix)
-	if(const FString* ShaderType = Manager.FindMatchingMaterialExpression(Connect.UpstreamNode->getCategory().c_str(), Connect.UpstreamNode->getNodeDef(mx::EMPTY_STRING, true)->getNodeGroup().c_str()))
+	if(const FString* ShaderType = Manager.FindMatchingMaterialExpression(Connect.UpstreamNode->getCategory().c_str(), Connect.UpstreamNode->getNodeDef(mx::EMPTY_STRING, true)->getNodeGroup().c_str(), Connect.UpstreamNode->getType().c_str()))
 	{
 		ConnectOutputToInput(ShaderType, &FMaterialXSurfaceShaderAbstract::CreateShaderNode);
 	}
@@ -342,6 +342,28 @@ bool FMaterialXSurfaceShaderAbstract::ConnectMatchingNodeOutputToInput(const FCo
 	}
 	else if(uint8 EnumType, EnumValue; Manager.FindMatchingMaterialFunction(Connect.UpstreamNode->getCategory().c_str(), MaterialFunctionPath, EnumType, EnumValue))
 	{
+		// In case of a surfaceshader node, especialy openpbr and standard_surface we need to check if we have to take the transmission surface shader
+		if (EnumType == UE::Interchange::MaterialX::IndexSurfaceShaders)
+		{
+			bool bIsTransmittance = false;
+			//check if the surfaceshader has a transmission input, only for openpbr and standardsurface
+			if (EnumValue == uint8(EInterchangeMaterialXShaders::OpenPBRSurface) && Connect.UpstreamNode->getInput(mx::OpenPBRSurface::Input::TransmissionWeight))
+			{
+				EnumValue = uint8(EInterchangeMaterialXShaders::OpenPBRSurfaceTransmission);
+				bIsTransmittance = true;
+			}
+			else if (EnumValue == uint8(EInterchangeMaterialXShaders::StandardSurface) && Connect.UpstreamNode->getInput(mx::StandardSurface::Input::Transmission))
+			{
+				EnumValue = uint8(EInterchangeMaterialXShaders::StandardSurfaceTransmission);
+				bIsTransmittance = true;
+			}
+
+			if (bIsTransmittance)
+			{
+				ShaderGraphNode->SetCustomBlendMode(EBlendMode::BLEND_TranslucentColoredTransmittance);
+			}
+		}
+
 		MaterialFunctionPath ?
 			ConnectOutputToInput(MaterialFunctionPath, &FMaterialXSurfaceShaderAbstract::CreateFunctionCallShaderNode) :
 			ConnectFunctionShaderNodeOutputToInput(EnumType, EnumValue, &FMaterialXSurfaceShaderAbstract::CreateFunctionCallShaderNode);
@@ -392,10 +414,10 @@ void FMaterialXSurfaceShaderAbstract::ConnectNodeCategoryOutputToInput(const Mat
 					}
 				}
 
-				std::vector<mx::OutputPtr> Outputs{ DownstreamNode->getOutputs()};
+				std::vector<mx::OutputPtr> Outputs{ DownstreamNode->getActiveOutputs()};
 				if(Outputs.empty())
 				{
-					Outputs = DownstreamNode->getNodeDef(mx::EMPTY_STRING, true)->getOutputs();
+					Outputs = DownstreamNode->getNodeDef(mx::EMPTY_STRING, true)->getActiveOutputs();
 				}
 
 				ParentShaderNodeOutputs.Empty();
@@ -1614,7 +1636,7 @@ const FString& FMaterialXSurfaceShaderAbstract::GetMatchedInputName(MaterialX::N
 		const FString NodeCategory{ Node->getCategory().c_str() };
 		const FString InputName{ GetInputName(Input) };
 
-		if(const FString* Result = Manager.FindMatchingInput(NodeCategory, InputName, Node->getNodeDef(mx::EMPTY_STRING, true)->getNodeGroup().c_str()))
+		if(const FString* Result = Manager.FindMatchingInput(NodeCategory, InputName, Node->getNodeDef(mx::EMPTY_STRING, true)->getNodeGroup().c_str(), Node->getType().c_str()))
 		{
 			return *Result;
 		}
@@ -1835,92 +1857,98 @@ namespace
 		}
 	}
 }
+
 UInterchangeShaderNode* FMaterialXSurfaceShaderAbstract::Translate(EInterchangeMaterialXShaders ShaderType)
 {
-	constexpr bool bInputInTangentSpace = true;
-
 	UInterchangeFunctionCallShaderNode* FunctionCallShaderNode = CreateFunctionCallShaderNode(SurfaceShaderNode, SurfaceMaterialName + TEXT("_") + SurfaceShaderNode->getName().c_str(), UE::Interchange::MaterialX::IndexSurfaceShaders, uint8(ShaderType));
 
-	for(mx::InputPtr Input : SurfaceShaderNode->getInputs())
+	return Translate(FunctionCallShaderNode);
+}
+
+UInterchangeShaderNode* FMaterialXSurfaceShaderAbstract::Translate(UInterchangeShaderNode* ShaderNode)
+{
+	constexpr bool bInputInTangentSpace = true;
+
+	for (mx::InputPtr Input : SurfaceShaderNode->getInputs())
 	{
 		mx::ValuePtr DefaultValue = Input->getDefaultValue();
 
 		//Let's create default values in case there is none
-		if(!DefaultValue)
+		if (!DefaultValue)
 		{
-			if(Input->getType() == mx::Type::Float)
+			if (Input->getType() == mx::Type::Float)
 			{
 				DefaultValue = mx::Value::createValueFromStrings(ValueToString<float>, mx::Type::Float);
 			}
-			else if(Input->getType() == mx::Type::Color3)
+			else if (Input->getType() == mx::Type::Color3)
 			{
-				DefaultValue = mx::Value::createValueFromStrings(ValueToString<float> + "," + ValueToString<float> + "," + ValueToString<float>, mx::Type::Color3);
+				DefaultValue = mx::Value::createValueFromStrings(ValueToString<float> +"," + ValueToString<float> +"," + ValueToString<float>, mx::Type::Color3);
 			}
-			else if(Input->getType() == mx::Type::Color4)
+			else if (Input->getType() == mx::Type::Color4)
 			{
-				DefaultValue = mx::Value::createValueFromStrings(ValueToString<float> + "," + ValueToString<float> + "," + ValueToString<float> + "," + ValueToString<float>, mx::Type::Color4);
+				DefaultValue = mx::Value::createValueFromStrings(ValueToString<float> +"," + ValueToString<float> +"," + ValueToString<float> +"," + ValueToString<float>, mx::Type::Color4);
 			}
-			else if(Input->getType() == mx::Type::Boolean)
+			else if (Input->getType() == mx::Type::Boolean)
 			{
 				DefaultValue = mx::Value::createValueFromStrings("false", mx::Type::Boolean);
 			}
-			else if(Input->getType() == mx::Type::Integer)
+			else if (Input->getType() == mx::Type::Integer)
 			{
 				DefaultValue = mx::Value::createValueFromStrings(ValueToString<int32>, mx::Type::Integer);
 			}
-			else if(Input->getType() == mx::Type::Vector2)
+			else if (Input->getType() == mx::Type::Vector2)
 			{
-				DefaultValue = mx::Value::createValueFromStrings(ValueToString<float> + "," + ValueToString<float>, mx::Type::Vector2);
+				DefaultValue = mx::Value::createValueFromStrings(ValueToString<float> +"," + ValueToString<float>, mx::Type::Vector2);
 			}
-			else if(Input->getType() == mx::Type::Vector3)
+			else if (Input->getType() == mx::Type::Vector3)
 			{
-				DefaultValue = mx::Value::createValueFromStrings(ValueToString<float> + "," + ValueToString<float> + "," + ValueToString<float>, mx::Type::Vector3);
+				DefaultValue = mx::Value::createValueFromStrings(ValueToString<float> +"," + ValueToString<float> +"," + ValueToString<float>, mx::Type::Vector3);
 			}
-			else if(Input->getType() == mx::Type::Vector4)
+			else if (Input->getType() == mx::Type::Vector4)
 			{
-				DefaultValue = mx::Value::createValueFromStrings(ValueToString<float> + "," + ValueToString<float> + "," + ValueToString<float> + "," + ValueToString<float>, mx::Type::Vector4);
+				DefaultValue = mx::Value::createValueFromStrings(ValueToString<float> +"," + ValueToString<float> +"," + ValueToString<float> +"," + ValueToString<float>, mx::Type::Vector4);
 			}
 		}
 
-		if(DefaultValue->getTypeString() == mx::Type::Float)
+		if (DefaultValue->getTypeString() == mx::Type::Float)
 		{
-			ConnectNodeOutputToInput(Input->getName().c_str(), FunctionCallShaderNode, Input->getName().c_str(), mx::fromValueString<float>(DefaultValue->getValueString()));
+			ConnectNodeOutputToInput(Input->getName().c_str(), ShaderNode, Input->getName().c_str(), mx::fromValueString<float>(DefaultValue->getValueString()));
 		}
-		else if(DefaultValue->getTypeString() == mx::Type::Color3)
+		else if (DefaultValue->getTypeString() == mx::Type::Color3)
 		{
-			ConnectNodeOutputToInput(Input->getName().c_str(), FunctionCallShaderNode, Input->getName().c_str(), GetValue<FLinearColor>(mx::fromValueString<mx::Color3>(DefaultValue->getValueString())));
+			ConnectNodeOutputToInput(Input->getName().c_str(), ShaderNode, Input->getName().c_str(), GetValue<FLinearColor>(mx::fromValueString<mx::Color3>(DefaultValue->getValueString())));
 		}
-		else if(DefaultValue->getTypeString() == mx::Type::Color4)
+		else if (DefaultValue->getTypeString() == mx::Type::Color4)
 		{
-			ConnectNodeOutputToInput(Input->getName().c_str(), FunctionCallShaderNode, Input->getName().c_str(), GetValue<FLinearColor>(mx::fromValueString<mx::Color4>(DefaultValue->getValueString())));
+			ConnectNodeOutputToInput(Input->getName().c_str(), ShaderNode, Input->getName().c_str(), GetValue<FLinearColor>(mx::fromValueString<mx::Color4>(DefaultValue->getValueString())));
 		}
-		else if(DefaultValue->getTypeString() == mx::Type::Boolean)
+		else if (DefaultValue->getTypeString() == mx::Type::Boolean)
 		{
-			ConnectNodeOutputToInput(Input->getName().c_str(), FunctionCallShaderNode, Input->getName().c_str(), mx::fromValueString<bool>(DefaultValue->getValueString()));
+			ConnectNodeOutputToInput(Input->getName().c_str(), ShaderNode, Input->getName().c_str(), mx::fromValueString<bool>(DefaultValue->getValueString()));
 		}
-		else if(DefaultValue->getTypeString() == mx::Type::Vector2)
+		else if (DefaultValue->getTypeString() == mx::Type::Vector2)
 		{
-			ConnectNodeOutputToInput(Input->getName().c_str(), FunctionCallShaderNode, Input->getName().c_str(), FVector4f{ GetValue<FVector2f>(mx::fromValueString<mx::Vector2>(DefaultValue->getValueString())), FVector2f{0,0} }, bInputInTangentSpace);
+			ConnectNodeOutputToInput(Input->getName().c_str(), ShaderNode, Input->getName().c_str(), FVector4f{ GetValue<FVector2f>(mx::fromValueString<mx::Vector2>(DefaultValue->getValueString())), FVector2f{0,0} }, bInputInTangentSpace);
 		}
-		else if(DefaultValue->getTypeString() == mx::Type::Vector3)
+		else if (DefaultValue->getTypeString() == mx::Type::Vector3)
 		{
-			ConnectNodeOutputToInput(Input->getName().c_str(), FunctionCallShaderNode, Input->getName().c_str(), GetValue<FVector3f>(mx::fromValueString<mx::Vector3>(DefaultValue->getValueString())), bInputInTangentSpace);
+			ConnectNodeOutputToInput(Input->getName().c_str(), ShaderNode, Input->getName().c_str(), GetValue<FVector3f>(mx::fromValueString<mx::Vector3>(DefaultValue->getValueString())), bInputInTangentSpace);
 		}
-		else if(DefaultValue->getTypeString() == mx::Type::Vector4)
+		else if (DefaultValue->getTypeString() == mx::Type::Vector4)
 		{
-			ConnectNodeOutputToInput(Input->getName().c_str(), FunctionCallShaderNode, Input->getName().c_str(), GetValue<FVector4f>(mx::fromValueString<mx::Vector4>(DefaultValue->getValueString())), bInputInTangentSpace);
+			ConnectNodeOutputToInput(Input->getName().c_str(), ShaderNode, Input->getName().c_str(), GetValue<FVector4f>(mx::fromValueString<mx::Vector4>(DefaultValue->getValueString())), bInputInTangentSpace);
 		}
-		else if(DefaultValue->getTypeString() == mx::Type::Integer)
+		else if (DefaultValue->getTypeString() == mx::Type::Integer)
 		{
-			ConnectNodeOutputToInput(Input->getName().c_str(), FunctionCallShaderNode, Input->getName().c_str(), mx::fromValueString<int32>(DefaultValue->getValueString()));
+			ConnectNodeOutputToInput(Input->getName().c_str(), ShaderNode, Input->getName().c_str(), mx::fromValueString<int32>(DefaultValue->getValueString()));
 		}
 		else // this is most likely a BSDF
 		{
-			ConnectNodeOutputToInput(Input->getName().c_str(), FunctionCallShaderNode, Input->getName().c_str(), nullptr);
+			ConnectNodeOutputToInput(Input->getName().c_str(), ShaderNode, GetInputName(Input), nullptr); // for BSDF/Surface shaders, we need to take a matching input just in case
 		}
 	}
 
-	return FunctionCallShaderNode;
+	return ShaderNode;
 }
 
 mx::InputPtr FMaterialXSurfaceShaderAbstract::GetInputNormal(mx::NodePtr Node, const char*& InputNormal) const
diff --git a/Engine/Plugins/Interchange/Runtime/Source/Import/Public/MaterialX/MaterialXUtils/MaterialXManager.h b/Engine/Plugins/Interchange/Runtime/Source/Import/Public/MaterialX/MaterialXUtils/MaterialXManager.h
--- a/Engine/Plugins/Interchange/Runtime/Source/Import/Public/MaterialX/MaterialXUtils/MaterialXManager.h
+++ b/Engine/Plugins/Interchange/Runtime/Source/Import/Public/MaterialX/MaterialXUtils/MaterialXManager.h
@@ -38,14 +38,14 @@ public:
 	/** Translate a MaterialX Document */
 	INTERCHANGEIMPORT_API bool Translate(MaterialX::DocumentPtr Document, UInterchangeBaseNodeContainer& BaseNodeContainer, const UInterchangeTranslatorBase* Translator = nullptr);
 
-	/** Find a corresponding Material Expression Input given a (Category [NodeGroup], Input) pair*/
-	INTERCHANGEIMPORT_API const FString* FindMatchingInput(const FString& CategoryKey, const FString& InputKey, const FString& NodeGroup = {}) const;
+	/** Find a corresponding Material Expression Input given a (Category [NodeGroup] [Type], Input) pair*/
+	INTERCHANGEIMPORT_API const FString* FindMatchingInput(const FString& CategoryKey, const FString& InputKey, const FString& NodeGroup = {}, const FString& Type = {}) const;
 
 	/** Find a stored Material Expression Input */
 	INTERCHANGEIMPORT_API const FString* FindMaterialExpressionInput(const FString& InputKey) const;
 
-	/** Find a matching Material Expression given a MaterialX category*/
-	INTERCHANGEIMPORT_API const FString* FindMatchingMaterialExpression(const FString& CategoryKey, const FString& NodeGroup = {}) const;
+	/** Find a matching Material Expression given a MaterialX category [nodegroup] [type]*/
+	INTERCHANGEIMPORT_API const FString* FindMatchingMaterialExpression(const FString& CategoryKey, const FString& NodeGroup = {}, const FString& Type= {}) const;
 
 	/** Find a matching Material Function given a MaterialX category*/
 	INTERCHANGEIMPORT_API bool FindMatchingMaterialFunction(const FString& CategoryKey, const FString*& MaterialFunctionPath, uint8& EnumType, uint8& EnumValue) const;
@@ -85,31 +85,33 @@ private:
 	*/
 	using FMaterialXMaterialFunction = TVariant<FString, EInterchangeMaterialXShaders, EInterchangeMaterialXBSDF, EInterchangeMaterialXEDF, EInterchangeMaterialXVDF>;
 
-	struct FKeyCategoryNodegroup
+	struct FKeyExpression
 	{
 		template<typename CategoryString>
-		FKeyCategoryNodegroup(CategoryString&& Category)
+		FKeyExpression(CategoryString&& Category)
 			: Category{ std::forward<CategoryString>(Category) }
 		{}
 
-		template<typename CategoryString, typename NodeGroupString>
-		FKeyCategoryNodegroup(CategoryString&& Category, NodeGroupString&& NodeGroup)
+		template<typename CategoryString, typename NodeGroupString, typename TypeString>
+		FKeyExpression(CategoryString&& Category, NodeGroupString&& NodeGroup, TypeString&& Type)
 			: Category{ std::forward<CategoryString>(Category) }
-			, NodeGroup{ std::forward<CategoryString>(NodeGroup) }
+			, NodeGroup{ std::forward<NodeGroupString>(NodeGroup) }
+			, Type{ std::forward<TypeString>(Type) }
 		{}
 
-		[[nodiscard]] FORCEINLINE bool operator==(const FKeyCategoryNodegroup& Rhs) const
+		[[nodiscard]] FORCEINLINE bool operator==(const FKeyExpression& Rhs) const
 		{
-			return Category == Rhs.Category && NodeGroup == Rhs.NodeGroup;
+			return Category == Rhs.Category && NodeGroup == Rhs.NodeGroup && Type == Rhs.Type;
 		}
 
-		friend FORCEINLINE uint32 GetTypeHash(const FKeyCategoryNodegroup& Key)
+		friend FORCEINLINE uint32 GetTypeHash(const FKeyExpression& Key)
 		{
-			return HashCombine(GetTypeHash(Key.Category), GetTypeHash(Key.NodeGroup));
+			return HashCombine(HashCombine(GetTypeHash(Key.Category), GetTypeHash(Key.NodeGroup)), GetTypeHash(Key.Type));
 		}
 
 		FString Category;
-		FString NodeGroup; // node group is optional, some nodes in MaterialX have different material expressions in UE
+		FString NodeGroup; // node group is optional, some nodes in MaterialX have different material expressions in UE depending on the nodegroup
+		FString Type; // Type is optional, some nodes in MaterialX have different material expressions in UE depending on the type
 	};
 
 	INTERCHANGEIMPORT_API FMaterialXManager();
@@ -121,10 +123,10 @@ private:
 	TSet<FString> MaterialExpressionInputs;
 
 	/** Given a MaterialX node (category (optionally a nodegroup) - input), return the UE/Interchange input name.*/
-	TMap<TPair<FKeyCategoryNodegroup, FString>, FString> MatchingInputNames;
+	TMap<TPair<FKeyExpression, FString>, FString> MatchingInputNames;
 
 	/** Given a MaterialX node category, optionally with a node group, return the UE material expression class name*/
-	TMap<FKeyCategoryNodegroup, FString> MatchingMaterialExpressions;
+	TMap<FKeyExpression, FString> MatchingMaterialExpressions;
 
 	/** Container of a MaterialX document, to translate the different nodes based on their category*/
 	TMap<FString, FOnGetMaterialXInstance> MaterialXContainerDelegates;
diff --git a/Engine/Plugins/Interchange/Runtime/Source/Import/Public/MaterialX/MaterialXUtils/MaterialXSurfaceShaderAbstract.h b/Engine/Plugins/Interchange/Runtime/Source/Import/Public/MaterialX/MaterialXUtils/MaterialXSurfaceShaderAbstract.h
--- a/Engine/Plugins/Interchange/Runtime/Source/Import/Public/MaterialX/MaterialXUtils/MaterialXSurfaceShaderAbstract.h
+++ b/Engine/Plugins/Interchange/Runtime/Source/Import/Public/MaterialX/MaterialXUtils/MaterialXSurfaceShaderAbstract.h
@@ -551,6 +551,14 @@ protected:
 	 */
 	UInterchangeShaderNode* Translate(EInterchangeMaterialXShaders ShaderType);
 
+	/**
+	 * This function should be called first by the Translate method of derived class, SurfaceShaderNode should initialized first by the derived class
+	 *
+	 * @param ShaderNode - the shader node that was created, can be either a shader node or a function call shader node
+	 * @return The shader node created with all inputs connected to it
+	 */
+	UInterchangeShaderNode* Translate(UInterchangeShaderNode* ShaderNode);
+
 	virtual UInterchangeBaseNode* Translate(MaterialX::NodePtr ShaderNode) override = 0;
 
 private:
diff --git a/Engine/Plugins/Runtime/USDCore/Source/USDUtilities/Private/USDMaterialXShaderGraph.cpp b/Engine/Plugins/Runtime/USDCore/Source/USDUtilities/Private/USDMaterialXShaderGraph.cpp
--- a/Engine/Plugins/Runtime/USDCore/Source/USDUtilities/Private/USDMaterialXShaderGraph.cpp
+++ b/Engine/Plugins/Runtime/USDCore/Source/USDUtilities/Private/USDMaterialXShaderGraph.cpp
@@ -287,8 +287,15 @@ void FUsdMaterialXShaderGraph::SetInputValue(mx::InputPtr Input, const pxr::UsdS
 		{
 			mx::NodeDefPtr NodeDef = Node->getNodeDef();
 			mx::InputPtr ActiveInput = NodeDef->getActiveInput(Input->getName());
-			MtlxType DefaultValue = ActiveInput->getDefaultValue()->asA<MtlxType>();
-			Input->setValue(DefaultValue);
+			if (ActiveInput->hasDefaultGeomPropString())
+			{
+				Node->removeInput(Input->getName());
+			}
+			else
+			{
+				MtlxType DefaultValue = ActiveInput->getDefaultValue()->asA<MtlxType>();
+				Input->setValue(DefaultValue);
+			}
 		}
 	};
 
@@ -436,6 +443,12 @@ mx::InputPtr FUsdMaterialXShaderGraph::GetInput(mx::NodePtr Node, const pxr::Usd
 			}
 		}
 	}
+	else // take the input given by the nodedef
+	{
+		mx::NodeDefPtr NodeDef = Node->getNodeDef();
+		mx::InputPtr Input = NodeDef->getActiveInput(InputName);
+		InputType = Input->getType();
+	}
 
 	return Node->addInput(InputName, InputType);
 }
