[Interchange-glTF] Fixing MorphTargetAnimation imports mixed with SkeletlaMesh Imports.

#rnx

diff --git a/Engine/Plugins/Interchange/Runtime/Source/Import/Private/Gltf/InterchangeGltfAnimation.cpp b/Engine/Plugins/Interchange/Runtime/Source/Import/Private/Gltf/InterchangeGltfAnimation.cpp
--- a/Engine/Plugins/Interchange/Runtime/Source/Import/Private/Gltf/InterchangeGltfAnimation.cpp
+++ b/Engine/Plugins/Interchange/Runtime/Source/Import/Private/Gltf/InterchangeGltfAnimation.cpp
@@ -712,6 +712,7 @@ namespace UE::Interchange::Gltf::Private
 			int32 AnimationIndex;
 			const GLTF::FAnimation& GLTFAnimation;
 			const TArray<GLTF::FNode>& GLTFNodes;
+			const TArray<GLTF::FSkinInfo>& GLTFSkins;
 			const TMap<const GLTF::FNode*, FString>& GLTFNodeToInterchangeUidMap; //UInterchangeGLTFTranslator::NodeUidMap
 
 			TMap<FString, UInterchangeSkeletalAnimationTrackNode*> RootJointIndexToTrackNodeMap;
@@ -719,11 +720,13 @@ namespace UE::Interchange::Gltf::Private
 			FAnimationHandler(UInterchangeBaseNodeContainer& InNodeContainer, 
 				int32 InAnimationIndex, const GLTF::FAnimation& InGLTFAnimation, 
 				const TArray<GLTF::FNode>& InGLTFNodes,
+				const TArray<GLTF::FSkinInfo>& InGLTFSkins,
 				const TMap<const GLTF::FNode*, FString>& InGLTFNodeToInterchangeUidMap)
 				: NodeContainer(InNodeContainer)
 				, AnimationIndex(InAnimationIndex)
 				, GLTFAnimation(InGLTFAnimation)
 				, GLTFNodes(InGLTFNodes)
+				, GLTFSkins(InGLTFSkins)
 				, GLTFNodeToInterchangeUidMap(InGLTFNodeToInterchangeUidMap)
 			{}
 
@@ -1060,7 +1063,7 @@ namespace UE::Interchange::Gltf::Private
 			{
 				TMap<int32, TSet<int32>> RigidAnimation;							// (AnimatedNodeIndex, [Channels])
 				TMap<int32, TMap<int32, TSet<int32>>> RiggedAnimations;				// (SkeletonRootIndex, (AnimatedNodeIndex, [Channels])
-				TMap<int32, TMap<int32, TSet<int32>>> MorphTargetAnimations;		// (AnimatedNodeIndex, (AnimatedNodeIndex, [Channels]) //Channel is a Target.Path==Weight channel.
+				TMap<int32, TMap<int32, TSet<int32>>> MorphTargetAnimations;		// (UEAnimatedNodeIndex, (AnimatedNodeIndex, [Channels]) //Channel is a Target.Path==Weight channel.
 																					// glTF allows MorphTargets on StaticMeshes as well.
 																					// However UE only does MorphTargets on Skeletals, for this reason we presume Skeletals and Skeletal Animations for Morph Target Animations.
 
@@ -1078,7 +1081,23 @@ namespace UE::Interchange::Gltf::Private
 
 					if (bMorphTargetAnimation)
 					{
-						TMap<int32, TSet<int32>>& AnimatedNodeIndicesToChannelIndices = MorphTargetAnimations.FindOrAdd(AnimatedNodeIndex);
+						int32 UEAnimatedNodeIndex = AnimatedNodeIndex;
+						if (AnimatedNode.Type == GLTF::FNode::EType::MeshSkinned)
+						{
+							const GLTF::FSkinInfo& Skin = GLTFSkins[AnimatedNode.Skindex];
+							if (GLTFSkins.IsValidIndex(AnimatedNode.Skindex)
+								&& GLTFSkins[AnimatedNode.Skindex].Joints.Num()
+								&& GLTFNodes.IsValidIndex(GLTFSkins[AnimatedNode.Skindex].Joints[0]))
+							{
+								UEAnimatedNodeIndex = GLTFNodes[Skin.Joints[0]].RootJointIndex;
+							}
+						}
+						else if (AnimatedNode.Type == GLTF::FNode::EType::Joint)
+						{
+							UEAnimatedNodeIndex = AnimatedNode.RootJointIndex;
+						}
+
+						TMap<int32, TSet<int32>>& AnimatedNodeIndicesToChannelIndices = MorphTargetAnimations.FindOrAdd(UEAnimatedNodeIndex);
 						TSet<int32>& ChannelIndices = AnimatedNodeIndicesToChannelIndices.FindOrAdd(AnimatedNodeIndex);
 						ChannelIndices.Add(ChannelIndex);
 					}
@@ -1112,13 +1131,14 @@ namespace UE::Interchange::Gltf::Private
 	void HandleGLTFAnimations(UInterchangeBaseNodeContainer& NodeContainer,
 		TArray<GLTF::FAnimation> Animations,
 		const TArray<GLTF::FNode>& GLTFNodes,
+		const TArray<GLTF::FSkinInfo>& GLTFSkins,
 		const TMap<const GLTF::FNode*, FString>& GLTFNodeToInterchangeUidMap)
 	{
 		using namespace AnimationHelpers;
 
 		for (int32 AnimationIndex = 0; AnimationIndex < Animations.Num(); AnimationIndex++)
 		{
-			FAnimationHandler AnimationHadler(NodeContainer, AnimationIndex, Animations[AnimationIndex], GLTFNodes, GLTFNodeToInterchangeUidMap);
+			FAnimationHandler AnimationHadler(NodeContainer, AnimationIndex, Animations[AnimationIndex], GLTFNodes, GLTFSkins, GLTFNodeToInterchangeUidMap);
 			AnimationHadler.Process();
 		}
 	}
diff --git a/Engine/Plugins/Interchange/Runtime/Source/Import/Private/Gltf/InterchangeGltfPrivate.h b/Engine/Plugins/Interchange/Runtime/Source/Import/Private/Gltf/InterchangeGltfPrivate.h
--- a/Engine/Plugins/Interchange/Runtime/Source/Import/Private/Gltf/InterchangeGltfPrivate.h
+++ b/Engine/Plugins/Interchange/Runtime/Source/Import/Private/Gltf/InterchangeGltfPrivate.h
@@ -11,6 +11,7 @@ namespace GLTF
 	struct FAsset;
 	struct FAnimation;
 	struct FNode;
+	struct FSkinInfo;
 }
 
 namespace UE::Interchange
@@ -43,6 +44,7 @@ namespace UE::Interchange
 		void HandleGLTFAnimations(UInterchangeBaseNodeContainer& NodeContainer,
 			TArray<GLTF::FAnimation> Animations,
 			const TArray<GLTF::FNode>& GLTFNodes,
+			const TArray<GLTF::FSkinInfo>& GLTFSkins,
 			const TMap<const GLTF::FNode*, FString>& GLTFNodeToInterchangeUidMap);
 	}
 }
\ No newline at end of file
diff --git a/Engine/Plugins/Interchange/Runtime/Source/Import/Private/Gltf/InterchangeGltfTranslator.cpp b/Engine/Plugins/Interchange/Runtime/Source/Import/Private/Gltf/InterchangeGltfTranslator.cpp
--- a/Engine/Plugins/Interchange/Runtime/Source/Import/Private/Gltf/InterchangeGltfTranslator.cpp
+++ b/Engine/Plugins/Interchange/Runtime/Source/Import/Private/Gltf/InterchangeGltfTranslator.cpp
@@ -911,7 +911,7 @@ bool UInterchangeGLTFTranslator::Translate( UInterchangeBaseNodeContainer& NodeC
 
 	// Animations
 	{
-		UE::Interchange::Gltf::Private::HandleGLTFAnimations(NodeContainer, GltfAsset.Animations, GltfAsset.Nodes, NodeUidMap);
+		UE::Interchange::Gltf::Private::HandleGLTFAnimations(NodeContainer, GltfAsset.Animations, GltfAsset.Nodes, GltfAsset.Skins, NodeUidMap);
 	}
 
 	// Variants
