[Interchange][FBX] Extra animation asset is created fix.

#rnx

diff --git a/Engine/Plugins/Interchange/Runtime/Source/Pipelines/Private/InterchangeGenericAnimationPipeline.cpp b/Engine/Plugins/Interchange/Runtime/Source/Pipelines/Private/InterchangeGenericAnimationPipeline.cpp
--- a/Engine/Plugins/Interchange/Runtime/Source/Pipelines/Private/InterchangeGenericAnimationPipeline.cpp
+++ b/Engine/Plugins/Interchange/Runtime/Source/Pipelines/Private/InterchangeGenericAnimationPipeline.cpp
@@ -139,6 +139,76 @@ namespace UE::Interchange::Private
 		constexpr bool bIsJointFalse = false;
 		UE::Interchange::SanitizeName(DisplayString, bIsJointFalse);
 	}
+
+	//Check Compatibility: (as in does the animation target the skeleton at all)
+	bool DoesSkeletalAnimationTargetSkeleton(const UInterchangeBaseNodeContainer* BaseNodeContainer, const UInterchangeSkeletalAnimationTrackNode& TrackNode, USkeleton* Skeleton)
+	{
+		//Check if the animation is targeting any of the joints on said skeleton:
+		TMap<FString, FString> SceneNodeAnimationPayloadKeyUids;
+		TMap<FString, uint8> SceneNodeAnimationPayloadKeyTypes;
+		TrackNode.GetSceneNodeAnimationPayloadKeys(SceneNodeAnimationPayloadKeyUids, SceneNodeAnimationPayloadKeyTypes);
+
+		TSet<FString> SceneNodeUIDs;
+		SceneNodeAnimationPayloadKeyUids.GetKeys(SceneNodeUIDs);
+		const FReferenceSkeleton& SkeletonRef = Skeleton->GetReferenceSkeleton();
+
+		bool bSkeletonCompatible = false;
+
+		for (const FString& SceneNodeUID : SceneNodeUIDs)
+		{
+			if (const UInterchangeSceneNode* SceneNode = Cast<UInterchangeSceneNode>(BaseNodeContainer->GetNode(SceneNodeUID)))
+			{
+				FString DisplayName = SceneNode->GetDisplayLabel();
+				int32 BoneIndex = SkeletonRef.FindBoneIndex(*DisplayName);
+				if (BoneIndex != INDEX_NONE)
+				{
+					bSkeletonCompatible = true;
+					break;
+				}
+			}
+		}
+
+		if (!bSkeletonCompatible)
+		{
+			USkeletalMesh* PreviewMesh = Skeleton->GetPreviewMesh();
+			const TArray<TObjectPtr<UMorphTarget>>& SkeletalMeshMorphTargets = PreviewMesh->GetMorphTargets();
+			TSet<FString> MorphTargetNames;
+			for (const TObjectPtr<UMorphTarget>& MorphTarget : SkeletalMeshMorphTargets)
+			{
+				if (MorphTarget)
+				{
+					MorphTargetNames.Add(MorphTarget->GetName());
+				}
+			}
+
+			TMap<FString, FString> MorphTargetNodeAnimationPayloadKeyUids;
+			TMap<FString, uint8> MorphTargetNodeAnimationPayloadKeyTypes;
+			TrackNode.GetMorphTargetNodeAnimationPayloadKeys(MorphTargetNodeAnimationPayloadKeyUids, MorphTargetNodeAnimationPayloadKeyTypes);
+
+			TSet<FString> MorphTargetNodeUIDs;
+			MorphTargetNodeAnimationPayloadKeyUids.GetKeys(MorphTargetNodeUIDs);
+
+			for (const FString& MorphTargetNodeUID : MorphTargetNodeUIDs)
+			{
+				if (const UInterchangeMeshNode* MeshNode = Cast<UInterchangeMeshNode>(BaseNodeContainer->GetNode(MorphTargetNodeUID)))
+				{
+					FString MorphTargetName = MeshNode->GetDisplayLabel();
+					if (MorphTargetNames.Contains(MorphTargetName))
+					{
+						bSkeletonCompatible = true;
+						break;
+					}
+				}
+			}
+
+			if (!bSkeletonCompatible)
+			{
+				return false;
+			}
+		}
+
+		return true;
+	}
 }
 
 void UInterchangeGenericAnimationPipeline::Serialize(FArchive& Ar)
@@ -600,7 +670,8 @@ void UInterchangeGenericAnimationPipeline::CreateAnimSequenceFactoryNode(UInterc
 
 	//If we import anim only and we do not have meshes and skeleton. We need to create a skeleton factory node
 	//base on the specified skeleton
-	if(bImportOnlyAnimation && !SkeletonFactoryNode && CommonSkeletalMeshesAndAnimationsProperties->Skeleton.IsValid())
+	bool bSkeletonCompatible = true;
+	if (bImportOnlyAnimation && !SkeletonFactoryNode && CommonSkeletalMeshesAndAnimationsProperties->Skeleton.IsValid())
 	{
 		TWeakObjectPtr<USkeleton> Skeleton = CommonSkeletalMeshesAndAnimationsProperties->Skeleton;
 		TPair<int32, FString> SkeletonRootNodeUidAndBoneIndex = TPair<int32, FString>(INDEX_NONE, FString());
@@ -632,6 +703,17 @@ void UInterchangeGenericAnimationPipeline::CreateAnimSequenceFactoryNode(UInterc
 			//Create a skeleton node from all the joint in the translated nodes
 			SkeletonFactoryNode = CommonSkeletalMeshesAndAnimationsProperties->CreateSkeletonFactoryNode(BaseNodeContainer, SkeletonRootUid);
 		}
+		
+		bSkeletonCompatible = UE::Interchange::Private::DoesSkeletalAnimationTargetSkeleton(BaseNodeContainer, TrackNode, CommonSkeletalMeshesAndAnimationsProperties->Skeleton.Get());
+		if (!bSkeletonCompatible)
+		{
+			UInterchangeResultDisplay_Generic* Message = AddMessage<UInterchangeResultDisplay_Generic>();
+			Message->Text = FText::Format(NSLOCTEXT("UInterchangeGenericAnimationPipeline", "IncompatibleSkeleton", "Incompatible skeleton {0} when importing AnimSequence {1}."),
+				FText::FromString(CommonSkeletalMeshesAndAnimationsProperties->Skeleton->GetName()),
+				FText::FromString(TrackNode.GetDisplayLabel()));
+
+			return;
+		}
 	}
 
 	if (!SkeletonFactoryNode)
@@ -895,11 +977,9 @@ void UInterchangeGenericAnimationPipeline::CreateAnimSequenceFactoryNode(UInterc
 
 	if (CommonSkeletalMeshesAndAnimationsProperties->Skeleton.IsValid())
 	{
-		bool bSkeletonCompatible = true;
-
 		//TODO: support skeleton helper in runtime
 #if WITH_EDITOR
-		bSkeletonCompatible = UE::Interchange::Private::FSkeletonHelper::IsCompatibleSkeleton(CommonSkeletalMeshesAndAnimationsProperties->Skeleton.Get()
+		bSkeletonCompatible = bSkeletonCompatible && UE::Interchange::Private::FSkeletonHelper::IsCompatibleSkeleton(CommonSkeletalMeshesAndAnimationsProperties->Skeleton.Get()
 			, RootJointUid
 			, BaseNodeContainer
 			, CommonSkeletalMeshesAndAnimationsProperties->bConvertStaticsWithMorphTargetsToSkeletals || CommonMeshesProperties->ForceAllMeshAsType == EInterchangeForceMeshType::IFMT_SkeletalMesh
