[BUGFIX] Interchange test plan - Level import - Packed Level Actor Hierarchy test failed (Crash)

#rnx

diff --git a/Engine/Plugins/Interchange/Runtime/Source/Pipelines/Private/InterchangeGenericScenesPipeline.cpp b/Engine/Plugins/Interchange/Runtime/Source/Pipelines/Private/InterchangeGenericScenesPipeline.cpp
--- a/Engine/Plugins/Interchange/Runtime/Source/Pipelines/Private/InterchangeGenericScenesPipeline.cpp
+++ b/Engine/Plugins/Interchange/Runtime/Source/Pipelines/Private/InterchangeGenericScenesPipeline.cpp
@@ -1421,6 +1421,12 @@ void UInterchangeGenericLevelPipeline::ExecutePostBroadcastPipeline(const UInter
 				{
 					if (CurrentLevelInstanceActor->GetWorldAsset().Get() == WorldAsset)
 					{
+						const UWorld* CurrentLevelInstanceActorWorld = CurrentLevelInstanceActor->GetWorld();
+						if (!CurrentLevelInstanceActorWorld || CurrentLevelInstanceActorWorld->IsCleanedUp())
+						{
+							continue;
+						}
+
 						CurrentLevelInstanceActor->UpdateLevelInstanceFromWorldAsset();
 						if (!CurrentLevelInstanceActor->IsA<APackedLevelActor>())
 						{
