Interchange MaterialX: remove use of swizzle node and channels attribute during translation

The swizzle node along with the channels attribute were both removed
in MaterialX 1.39.0:

[REDACTED_PROFILE_URL]/MaterialX/pull/1793
[REDACTED_PROFILE_URL]/MaterialX/pull/1804

MaterialX itself takes care of translating legacy documents from 1.38.X
and earlier to use the extract, separate, and combine nodes.

This change allows transitioning from MaterialX version 1.38.10 to
version 1.39.3 by discontinuing use of the swizzle MaterialX node and
the channels attribute when building against MaterialX 1.39.X. This only
addresses compilation issues however, and further changes are likely needed
to ensure full UE support of MaterialX nodes newly added or changed between
1.38.10 and 1.39.3.

The swizzle material expression on the UE side remains available in
case it has been saved into Unreal materials. The swizzle node handling
on the Interchange pipeline side also remains for now since it does not
have any includes of MaterialX headers and does not know the MaterialX
version, but it can be removed after we have transitioned to MaterialX
1.39.3 and no swizzle Interchange nodes are being produced.

diff --git a/Engine/Plugins/Interchange/Runtime/Source/Common/Public/MaterialX/InterchangeMaterialXDefinitions.h b/Engine/Plugins/Interchange/Runtime/Source/Common/Public/MaterialX/InterchangeMaterialXDefinitions.h
--- a/Engine/Plugins/Interchange/Runtime/Source/Common/Public/MaterialX/InterchangeMaterialXDefinitions.h
+++ b/Engine/Plugins/Interchange/Runtime/Source/Common/Public/MaterialX/InterchangeMaterialXDefinitions.h
@@ -726,7 +726,9 @@ MATERIALX_NAMESPACE_BEGIN
 		static constexpr const char* Separate2 = "separate2";
 		static constexpr const char* Separate3 = "separate3";
 		static constexpr const char* Separate4 = "separate4";
+#if MATERIALX_MINOR_VERSION < 39
 		static constexpr const char* Swizzle = "swizzle";
+#endif // MATERIALX_MINOR_VERSION < 39
 		// Procedural nodes 
 		static constexpr const char* Constant = "constant";
 		// Procedural2D nodes
diff --git a/Engine/Plugins/Interchange/Runtime/Source/Import/Private/MaterialX/MaterialXUtils/MaterialXSurfaceShaderAbstract.cpp b/Engine/Plugins/Interchange/Runtime/Source/Import/Private/MaterialX/MaterialXUtils/MaterialXSurfaceShaderAbstract.cpp
--- a/Engine/Plugins/Interchange/Runtime/Source/Import/Private/MaterialX/MaterialXUtils/MaterialXSurfaceShaderAbstract.cpp
+++ b/Engine/Plugins/Interchange/Runtime/Source/Import/Private/MaterialX/MaterialXUtils/MaterialXSurfaceShaderAbstract.cpp
@@ -74,6 +74,7 @@ bool FMaterialXSurfaceShaderAbstract::AddAttributeFromValueOrInterface(MaterialX
 		UInterchangeShaderNode* ShaderNodeToConnectTo = ShaderNode;
 		FString InputToConnectTo = InputChannelName;
 
+#if MATERIALX_MINOR_VERSION < 39
 		if(Input->hasChannels())
 		{
 			using namespace UE::Interchange::Materials::Standard::Nodes;
@@ -86,6 +87,7 @@ bool FMaterialXSurfaceShaderAbstract::AddAttributeFromValueOrInterface(MaterialX
 			ShaderNodeToConnectTo = SwizzleNode;
 			InputToConnectTo = TEXT("Input");
 		}
+#endif // MATERIALX_MINOR_VERSION < 39
 
 		if(Input->hasValue())
 		{
@@ -390,8 +392,10 @@ void FMaterialXSurfaceShaderAbstract::ConnectNodeCategoryOutputToInput(const Mat
 
 		FString OutputChannelName = OutputName;
 
+#if MATERIALX_MINOR_VERSION < 39
 		// Swizzle node for the attribute 'channels'
 		UInterchangeShaderNode* ChannelsNode = nullptr;
+#endif // MATERIALX_MINOR_VERSION < 39
 
 		if(mx::ElementPtr DownstreamElement = Edge.getDownstreamElement())
 		{
@@ -406,12 +410,14 @@ void FMaterialXSurfaceShaderAbstract::ConnectNodeCategoryOutputToInput(const Mat
 						OutputChannelName = ConnectedInput->getOutputString().c_str();
 					}
 
+#if MATERIALX_MINOR_VERSION < 39
 					if(ConnectedInput->hasChannels())
 					{
 						using namespace UE::Interchange::Materials::Standard::Nodes;
 						ChannelsNode = CreateShaderNode(ConnectedInput, (UpstreamNode->getName() + ConnectedInput->getName()).c_str() + FString{ TEXT("_Channels") }, Swizzle::Name.ToString());
 						ChannelsNode->AddStringAttribute(Swizzle::Attributes::Channels.ToString(), ConnectedInput->getChannels().c_str());
 					}
+#endif // MATERIALX_MINOR_VERSION < 39
 				}
 
 				std::vector<mx::OutputPtr> Outputs{ DownstreamNode->getActiveOutputs()};
@@ -425,6 +431,7 @@ void FMaterialXSurfaceShaderAbstract::ConnectNodeCategoryOutputToInput(const Mat
 				{
 					if(UInterchangeShaderNode** FoundNode = ShaderNodes.Find({ GetAttributeParentName(DownstreamNode, ConnectedInput), Outputs[i]->getName().c_str()}))
 					{
+#if MATERIALX_MINOR_VERSION < 39
 						//Connect the swizzle node between the upstream and downstream node
 						if(ChannelsNode)
 						{
@@ -432,6 +439,7 @@ void FMaterialXSurfaceShaderAbstract::ConnectNodeCategoryOutputToInput(const Mat
 							InputChannelName = TEXT("Input");
 							FoundNode = &ChannelsNode;
 						}
+#endif // MATERIALX_MINOR_VERSION < 39
 
 						ParentShaderNodeOutputs.Emplace(*FoundNode);
 					}
@@ -1455,6 +1463,7 @@ void FMaterialXSurfaceShaderAbstract::ConnectSwitchInputToOutput(const FConnectN
 	UInterchangeShaderPortsAPI::ConnectDefaultOuputToInput(Connect.ParentShaderNode, Connect.InputChannelName, SwitchNode->GetUniqueID());
 }
 
+#if MATERIALX_MINOR_VERSION < 39
 void FMaterialXSurfaceShaderAbstract::ConnectSwizzleInputToOutput(const FConnectNode& Connect)
 {
 	using namespace UE::Interchange::Materials::Standard::Nodes;
@@ -1471,6 +1480,7 @@ void FMaterialXSurfaceShaderAbstract::ConnectSwizzleInputToOutput(const FConnect
 
 	UInterchangeShaderPortsAPI::ConnectDefaultOuputToInput(Connect.ParentShaderNode, Connect.InputChannelName, SwizzleNode->GetUniqueID());
 }
+#endif // MATERIALX_MINOR_VERSION < 39
 
 void FMaterialXSurfaceShaderAbstract::ConnectNormalMapInputToOutput(const FConnectNode& Connect)
 {
@@ -1792,7 +1802,9 @@ void FMaterialXSurfaceShaderAbstract::RegisterConnectNodeOutputToInputDelegates(
 	MatchingConnectNodeDelegates.Add(mx::Category::Separate3,		FOnConnectNodeOutputToInput::CreateSP(this, &FMaterialXSurfaceShaderAbstract::ConnectSeparateInputToOutput));
 	MatchingConnectNodeDelegates.Add(mx::Category::Separate4,		FOnConnectNodeOutputToInput::CreateSP(this, &FMaterialXSurfaceShaderAbstract::ConnectSeparateInputToOutput));
 	MatchingConnectNodeDelegates.Add(mx::Category::TexCoord,		FOnConnectNodeOutputToInput::CreateSP(this, &FMaterialXSurfaceShaderAbstract::ConnectTexCoordInputToOutput));
+#if MATERIALX_MINOR_VERSION < 39
 	MatchingConnectNodeDelegates.Add(mx::Category::Swizzle,			FOnConnectNodeOutputToInput::CreateSP(this, &FMaterialXSurfaceShaderAbstract::ConnectSwizzleInputToOutput));
+#endif // MATERIALX_MINOR_VERSION < 39
 	MatchingConnectNodeDelegates.Add(mx::Category::Switch,			FOnConnectNodeOutputToInput::CreateSP(this, &FMaterialXSurfaceShaderAbstract::ConnectSwitchInputToOutput));
 	MatchingConnectNodeDelegates.Add(mx::Category::Refract,			FOnConnectNodeOutputToInput::CreateSP(this, &FMaterialXSurfaceShaderAbstract::ConnectRefractInputToOutput));
 	MatchingConnectNodeDelegates.Add(mx::Category::ViewDirection,	FOnConnectNodeOutputToInput::CreateSP(this, &FMaterialXSurfaceShaderAbstract::ConnectViewDirectionInputToOutput));
diff --git a/Engine/Plugins/Interchange/Runtime/Source/Import/Public/MaterialX/MaterialXUtils/MaterialXSurfaceShaderAbstract.h b/Engine/Plugins/Interchange/Runtime/Source/Import/Public/MaterialX/MaterialXUtils/MaterialXSurfaceShaderAbstract.h
--- a/Engine/Plugins/Interchange/Runtime/Source/Import/Public/MaterialX/MaterialXUtils/MaterialXSurfaceShaderAbstract.h
+++ b/Engine/Plugins/Interchange/Runtime/Source/Import/Public/MaterialX/MaterialXUtils/MaterialXSurfaceShaderAbstract.h
@@ -315,8 +315,10 @@ protected:
 	/** <switch> */
 	void ConnectSwitchInputToOutput(const FConnectNode& Connect);
 
+#if MATERIALX_MINOR_VERSION < 39
 	/** <swizzle> */
 	void ConnectSwizzleInputToOutput(const FConnectNode& Connect);
+#endif // MATERIALX_MINOR_VERSION < 39
 
 	/** <normalmap> */
 	void ConnectNormalMapInputToOutput(const FConnectNode& Connect);
