[Interchange-FBX] MorphTarget naming fix.

#rnx

diff --git a/Engine/Plugins/Interchange/Runtime/Source/Parsers/Fbx/Private/FbxMesh.cpp b/Engine/Plugins/Interchange/Runtime/Source/Parsers/Fbx/Private/FbxMesh.cpp
--- a/Engine/Plugins/Interchange/Runtime/Source/Parsers/Fbx/Private/FbxMesh.cpp
+++ b/Engine/Plugins/Interchange/Runtime/Source/Parsers/Fbx/Private/FbxMesh.cpp
@@ -1865,7 +1865,7 @@ void FFbxMesh::AddAllMeshes(FbxScene* SDKScene, FbxGeometryConverter* SDKGeometr
 				const int32 CurrentChannelMorphTargetCount = Channel->GetTargetShapeCount();
 				FString ChannelName = Parser.GetFbxHelper()->GetFbxObjectName(Channel);
 				// Maya adds the name of the MorphTarget and an underscore to the front of the channel name, so remove it
-				if (ChannelName.StartsWith(MorphTargetName) && (ChannelName.Compare(MorphTargetName, ESearchCase::IgnoreCase) == 0))
+				if (ChannelName.StartsWith(MorphTargetName) && (ChannelName.Compare(MorphTargetName, ESearchCase::IgnoreCase) != 0))
 				{
 					ChannelName.RightInline(ChannelName.Len() - (MorphTargetName.Len() + 1), EAllowShrinking::No);
 				}
