Interchange MaterialX:
- Fix displacement input normal not correctly plugged to the surfaceshader input normal

diff --git a/Engine/Plugins/Interchange/Runtime/Source/Import/Private/MaterialX/MaterialXUtils/MaterialXOpenPBRSurfaceShader.cpp b/Engine/Plugins/Interchange/Runtime/Source/Import/Private/MaterialX/MaterialXUtils/MaterialXOpenPBRSurfaceShader.cpp
--- a/Engine/Plugins/Interchange/Runtime/Source/Import/Private/MaterialX/MaterialXUtils/MaterialXOpenPBRSurfaceShader.cpp
+++ b/Engine/Plugins/Interchange/Runtime/Source/Import/Private/MaterialX/MaterialXUtils/MaterialXOpenPBRSurfaceShader.cpp
@@ -96,7 +96,15 @@ UInterchangeBaseNode* FMaterialXOpenPBRSurfaceShader::Translate(MaterialX::NodeP
 mx::InputPtr FMaterialXOpenPBRSurfaceShader::GetInputNormal(mx::NodePtr OpenPbrSurfaceShaderNode, const char*& InputNormal) const
 {
 	InputNormal = mx::OpenPBRSurface::Input::GeometryNormal;
-	return OpenPbrSurfaceShaderNode->getActiveInput(mx::OpenPBRSurface::Input::GeometryNormal);
+	mx::InputPtr Input = OpenPbrSurfaceShaderNode->getActiveInput(InputNormal);
+
+	// if no input is connected take the one from the nodedef
+	if (!Input)
+	{
+		Input = OpenPbrSurfaceShaderNode->getNodeDef(mx::EMPTY_STRING, true)->getActiveInput(InputNormal);
+	}
+
+	return Input;
 }
 
 #endif //WITH_EDITOR
diff --git a/Engine/Plugins/Interchange/Runtime/Source/Import/Private/MaterialX/MaterialXUtils/MaterialXStandardSurfaceShader.cpp b/Engine/Plugins/Interchange/Runtime/Source/Import/Private/MaterialX/MaterialXUtils/MaterialXStandardSurfaceShader.cpp
--- a/Engine/Plugins/Interchange/Runtime/Source/Import/Private/MaterialX/MaterialXUtils/MaterialXStandardSurfaceShader.cpp
+++ b/Engine/Plugins/Interchange/Runtime/Source/Import/Private/MaterialX/MaterialXUtils/MaterialXStandardSurfaceShader.cpp
@@ -114,7 +114,15 @@ UInterchangeBaseNode* FMaterialXStandardSurfaceShader::Translate(mx::NodePtr Sta
 mx::InputPtr FMaterialXStandardSurfaceShader::GetInputNormal(mx::NodePtr StandardSurfaceNode, const char*& InputNormal) const
 {
 	InputNormal = mx::StandardSurface::Input::Normal;
-	return StandardSurfaceNode->getActiveInput(mx::StandardSurface::Input::Normal);
+	
+	mx::InputPtr Input = StandardSurfaceNode->getActiveInput(InputNormal);
+		
+	if (!Input)
+	{
+		Input = StandardSurfaceNode->getNodeDef(mx::EMPTY_STRING, true)->getActiveInput(InputNormal);
+	}
+
+	return Input;
 }
 
 #endif
diff --git a/Engine/Plugins/Interchange/Runtime/Source/Import/Private/MaterialX/MaterialXUtils/MaterialXUsdPreviewSurfaceShader.cpp b/Engine/Plugins/Interchange/Runtime/Source/Import/Private/MaterialX/MaterialXUtils/MaterialXUsdPreviewSurfaceShader.cpp
--- a/Engine/Plugins/Interchange/Runtime/Source/Import/Private/MaterialX/MaterialXUtils/MaterialXUsdPreviewSurfaceShader.cpp
+++ b/Engine/Plugins/Interchange/Runtime/Source/Import/Private/MaterialX/MaterialXUtils/MaterialXUsdPreviewSurfaceShader.cpp
@@ -75,6 +75,13 @@ UInterchangeBaseNode* FMaterialXUsdPreviewSurfaceShader::Translate(MaterialX::No
 mx::InputPtr FMaterialXUsdPreviewSurfaceShader::GetInputNormal(mx::NodePtr UsdPreviewSurfaceNode, const char*& InputNormal) const
 {
 	InputNormal = mx::UsdPreviewSurface::Input::Normal;
-	return UsdPreviewSurfaceNode->getActiveInput(mx::UsdPreviewSurface::Input::Normal);
+	mx::InputPtr Input = UsdPreviewSurfaceNode->getActiveInput(InputNormal);
+	
+	if (!Input)
+	{
+		Input = UsdPreviewSurfaceNode->getNodeDef(mx::EMPTY_STRING, true)->getActiveInput(InputNormal);
+	}
+
+	return Input;
 }
 #endif
\ No newline at end of file
