Interchange USD: Fixed GeometryCache losing its materials after a reimport.

#rnx

diff --git a/Engine/Plugins/Interchange/Runtime/Source/Import/Private/Mesh/InterchangeGeometryCacheFactory.cpp b/Engine/Plugins/Interchange/Runtime/Source/Import/Private/Mesh/InterchangeGeometryCacheFactory.cpp
--- a/Engine/Plugins/Interchange/Runtime/Source/Import/Private/Mesh/InterchangeGeometryCacheFactory.cpp
+++ b/Engine/Plugins/Interchange/Runtime/Source/Import/Private/Mesh/InterchangeGeometryCacheFactory.cpp
@@ -231,7 +231,15 @@ UInterchangeFactoryBase::FImportAssetResult UInterchangeGeometryCacheFactory::Be
 		GeometryCache = Cast<UGeometryCache>(ExistingAsset);
 		if (GeometryCache)
 		{
+			// Backup the material assignment since it will be reset with the call to ClearForReimporting
+			TArray<TObjectPtr<UMaterialInterface>> ExistingMaterials = GeometryCache->Materials;
+			TArray<FName> ExistingMaterialSlotNames = GeometryCache->MaterialSlotNames;
+
 			GeometryCache->ClearForReimporting();
+
+			// Re-apply the material assignment
+			GeometryCache->Materials = ExistingMaterials;
+			GeometryCache->MaterialSlotNames = ExistingMaterialSlotNames;
 		}
 	}
 
