Interchange USD: Fixed crash when reimporting a geometry cache while it is playing back in PIE.

#rnx

diff --git a/Engine/Plugins/Interchange/Runtime/Source/Import/Private/Mesh/InterchangeGeometryCacheFactory.cpp b/Engine/Plugins/Interchange/Runtime/Source/Import/Private/Mesh/InterchangeGeometryCacheFactory.cpp
--- a/Engine/Plugins/Interchange/Runtime/Source/Import/Private/Mesh/InterchangeGeometryCacheFactory.cpp
+++ b/Engine/Plugins/Interchange/Runtime/Source/Import/Private/Mesh/InterchangeGeometryCacheFactory.cpp
@@ -5,6 +5,7 @@
 #include "Async/ParallelFor.h"
 #include "GeometryCache.h"
 #include "GeometryCacheCodecV1.h"
+#include "GeometryCacheComponent.h"
 #include "GeometryCacheHelpers.h"
 #include "GeometryCacheMeshData.h"
 #include "GeometryCacheTrackStreamable.h"
@@ -21,6 +22,7 @@
 #include "Misc/ScopedSlowTask.h"
 #include "Nodes/InterchangeBaseNode.h"
 #include "Nodes/InterchangeBaseNodeContainer.h"
+#include "UObject/UObjectIterator.h"
 
 #include UE_INLINE_GENERATED_CPP_BY_NAME(InterchangeGeometryCacheFactory)
 
@@ -35,6 +37,45 @@ static FAutoConsoleVariableRef CVarGeometryCacheParallelFrameReads(
 	TEXT("Maximum number of frames to read in parallel")
 );
 
+class FGeometryCacheComponentResetAsset
+{
+public:
+	/* 
+	 * Clear the given geometry cache from any geometry cache component using it.
+	 * This stops the geometry cache from playing while it is re-imported.
+	 */
+	FGeometryCacheComponentResetAsset(UGeometryCache* InGeometryCache)
+	: GeometryCache(InGeometryCache)
+	{
+		for (TObjectIterator<UGeometryCacheComponent> GeometryCacheComponentIt; GeometryCacheComponentIt; ++GeometryCacheComponentIt)
+		{
+			if (GeometryCacheComponentIt->GeometryCache == InGeometryCache)
+			{
+				GeometryCacheComponentIt->SetGeometryCache(nullptr);
+				GeometryCacheComponents.Add(*GeometryCacheComponentIt);
+			}
+		}
+	}
+
+	/** Restore the geometry cache on the geometry cache components that were previously using it */
+	~FGeometryCacheComponentResetAsset()
+	{
+		const int32 NumComponents = GeometryCacheComponents.Num();
+		for (UGeometryCacheComponent* GeometryCacheComponent : GeometryCacheComponents)
+		{
+			GeometryCacheComponent->SetGeometryCache(GeometryCache);
+		}
+	}
+
+	FGeometryCacheComponentResetAsset(const FGeometryCacheComponentResetAsset& Other) = delete;
+	FGeometryCacheComponentResetAsset(const FGeometryCacheComponentResetAsset&& Other) = delete;
+	FGeometryCacheComponentResetAsset& operator=(const FGeometryCacheComponentResetAsset& Other) = delete;
+
+private:
+	UGeometryCache* GeometryCache;
+	TArray<UGeometryCacheComponent*> GeometryCacheComponents;
+};
+
 UClass* UInterchangeGeometryCacheFactory::GetFactoryClass() const
 {
 	return UGeometryCache::StaticClass();
@@ -231,6 +272,8 @@ UInterchangeFactoryBase::FImportAssetResult UInterchangeGeometryCacheFactory::Be
 		GeometryCache = Cast<UGeometryCache>(ExistingAsset);
 		if (GeometryCache)
 		{
+			ResetAssetOnReimport = MakePimpl<FGeometryCacheComponentResetAsset>(GeometryCache);
+
 			// Backup the material assignment since it will be reset with the call to ClearForReimporting
 			TArray<TObjectPtr<UMaterialInterface>> ExistingMaterials = GeometryCache->Materials;
 			TArray<FName> ExistingMaterialSlotNames = GeometryCache->MaterialSlotNames;
@@ -633,6 +676,9 @@ UInterchangeFactoryBase::FImportAssetResult UInterchangeGeometryCacheFactory::En
 #endif
 
 	ImportAssetResult.ImportedObject = GeometryCache;
+
+	ResetAssetOnReimport.Reset();
+
 	return ImportAssetResult;
 }
 
diff --git a/Engine/Plugins/Interchange/Runtime/Source/Import/Public/Mesh/InterchangeGeometryCacheFactory.h b/Engine/Plugins/Interchange/Runtime/Source/Import/Public/Mesh/InterchangeGeometryCacheFactory.h
--- a/Engine/Plugins/Interchange/Runtime/Source/Import/Public/Mesh/InterchangeGeometryCacheFactory.h
+++ b/Engine/Plugins/Interchange/Runtime/Source/Import/Public/Mesh/InterchangeGeometryCacheFactory.h
@@ -3,6 +3,7 @@
 #pragma once
 
 #include "InterchangeFactoryBase.h"
+#include "Templates/PimplPtr.h"
 
 #include "InterchangeGeometryCacheFactory.generated.h"
 
@@ -36,6 +37,7 @@ public:
 
 private:
 	TArray<TObjectPtr<UGeometryCacheTrack>> Tracks;
+	TPimplPtr<class FGeometryCacheComponentResetAsset> ResetAssetOnReimport;
 };
 
 #undef UE_API
