Interchange FBX: Fix incorrect skeletons produced in very edge cases when TransformLinks seem to represent local transforms.

diff --git a/Engine/Plugins/Interchange/Runtime/Source/Parsers/Fbx/Private/FbxMesh.cpp b/Engine/Plugins/Interchange/Runtime/Source/Parsers/Fbx/Private/FbxMesh.cpp
--- a/Engine/Plugins/Interchange/Runtime/Source/Parsers/Fbx/Private/FbxMesh.cpp
+++ b/Engine/Plugins/Interchange/Runtime/Source/Parsers/Fbx/Private/FbxMesh.cpp
@@ -2005,122 +2005,139 @@ void FFbxMesh::AddAllMeshes(FbxScene* SDKScene, FbxGeometryConverter* SDKGeometr
 	} // for GeometryCount
 }
 
-bool FFbxMesh::GetGlobalJointBindPoseTransform(FFbxParser* Parser, FbxScene* SDKScene, FbxNode* Joint, FFbxJointMeshBindPoseGenerator& FbxJointMeshBindPoseGenerator,
-	FbxAMatrix& GlobalBindPoseJointMatrix, TMap<FString, FMatrix>& MeshIdToGlobalBindPoseJointMap, TMap<FString, FMatrix>& MeshIdToGlobalBindPoseReferenceMap, 
-	bool& bBadBindPoseMessageDisplay)
+bool AcquireBindPoseMatrix(FbxPose* CurrentPose, FbxAMatrix& GlobalBindPoseJointMatrix, FbxNode* Joint)
 {
-	FbxManager* SDKManager = SDKScene->GetFbxManager();
-	
-	GlobalBindPoseJointMatrix = FbxJointMeshBindPoseGenerator.GetBindPoseFbxAMatrix(Joint, GlobalBindPoseJointMatrix);
-	MeshIdToGlobalBindPoseReferenceMap = FbxJointMeshBindPoseGenerator.GetMeshIdToGlobalBindPoseReferenceMap(Joint);
-	MeshIdToGlobalBindPoseJointMap = FbxJointMeshBindPoseGenerator.GetMeshIdToGlobalBindPoseJointMap(Joint);
-	
-	if (MeshIdToGlobalBindPoseReferenceMap.Num() > 0)
+	if (!CurrentPose || !Joint)
 	{
-		return true;
+		return false;
 	}
 
-	if (!(Joint->GetMesh() && Joint->GetMesh()->GetDeformerCount(FbxDeformer::eSkin) > 0))
+	int32 PoseLinkIndex = CurrentPose->Find(Joint);
+	if (PoseLinkIndex >= 0)
+	{
+		if (!CurrentPose->IsLocalMatrix(PoseLinkIndex))
+		{
+			FbxMatrix NoneAffineMatrix = CurrentPose->GetMatrix(PoseLinkIndex);
+			GlobalBindPoseJointMatrix = *(FbxAMatrix*)(double*)&NoneAffineMatrix;
+			return true;
+		}
+	}
+	return false;
+}
+
+bool RetrievePoseFromBindPose(FbxNode* Joint, FbxAMatrix& GlobalBindPoseJointMatrix)
+{
+	if (!Joint)
 	{
 		return false;
 	}
 
-	auto AcquireBindPoseMatrix = [](FbxPose* CurrentPose, FbxAMatrix& GlobalBindPoseJointMatrix, FbxNode* Joint)
+	FbxScene* SDKScene = Joint->GetScene();
+	if (!SDKScene)
 	{
-		int32 PoseLinkIndex = CurrentPose->Find(Joint);
-		if (PoseLinkIndex >= 0)
-		{
-			if (!CurrentPose->IsLocalMatrix(PoseLinkIndex))
-			{
-				FbxMatrix NoneAffineMatrix = CurrentPose->GetMatrix(PoseLinkIndex);
-				GlobalBindPoseJointMatrix = *(FbxAMatrix*)(double*)&NoneAffineMatrix;
-				return true;
-			}
-		}
 		return false;
-	};
+	}
 
-	auto RetrievePoseFromBindPose = [&]()
+	const int32 PoseCount = SDKScene->GetPoseCount();
+	for (int32 PoseIndex = 0; PoseIndex < PoseCount; PoseIndex++)
+	{
+		FbxPose* CurrentPose = SDKScene->GetPose(PoseIndex);
+
+		// current pose is bind pose, 
+		if (CurrentPose && CurrentPose->IsBindPose())
 		{
-			const int32 PoseCount = SDKScene->GetPoseCount();
-			for (int32 PoseIndex = 0; PoseIndex < PoseCount; PoseIndex++)
-			{
-				FbxPose* CurrentPose = SDKScene->GetPose(PoseIndex);
+			FString PoseName = CurrentPose->GetName();
+			// all error report status
+			FbxStatus Status;
+
+			FbxArray<FbxNode*> pMissingAncestors, pMissingDeformers, pMissingDeformersAncestors, pWrongMatrices;
 
-				// current pose is bind pose, 
-				if (CurrentPose && CurrentPose->IsBindPose())
+			if (CurrentPose->IsValidBindPoseVerbose(Joint, pMissingAncestors, pMissingDeformers, pMissingDeformersAncestors, pWrongMatrices, 0.0001, &Status))
+			{
+				if (AcquireBindPoseMatrix(CurrentPose, GlobalBindPoseJointMatrix, Joint))
 				{
-					FString PoseName = CurrentPose->GetName();
-					// all error report status
-					FbxStatus Status;
+					return true;
+				}
+			}
+			else
+			{
+				// first try to fix up
+				// add missing ancestors
+				for (int i = 0; i < pMissingAncestors.GetCount(); i++)
+				{
+					FbxAMatrix mat = pMissingAncestors.GetAt(i)->EvaluateGlobalTransform(FBXSDK_TIME_ZERO);
+					CurrentPose->Add(pMissingAncestors.GetAt(i), mat);
+				}
 
-					FbxArray<FbxNode*> pMissingAncestors, pMissingDeformers, pMissingDeformersAncestors, pWrongMatrices;
+				pMissingAncestors.Clear();
+				pMissingDeformers.Clear();
+				pMissingDeformersAncestors.Clear();
+				pWrongMatrices.Clear();
 
-					if (CurrentPose->IsValidBindPoseVerbose(Joint, pMissingAncestors, pMissingDeformers, pMissingDeformersAncestors, pWrongMatrices, 0.0001, &Status))
+				// check it again
+				if (CurrentPose->IsValidBindPose(Joint))
+				{
+					if (AcquireBindPoseMatrix(CurrentPose, GlobalBindPoseJointMatrix, Joint))
 					{
-						if (AcquireBindPoseMatrix(CurrentPose, GlobalBindPoseJointMatrix, Joint))
-						{
-							return true;
-						}
+						return true;
 					}
-					else
+				}
+				else
+				{
+					// first try to find parent who is null group and see if you can try test it again
+					FbxNode* ParentNode = Joint->GetParent();
+					while (ParentNode)
 					{
-						// first try to fix up
-						// add missing ancestors
-						for (int i = 0; i < pMissingAncestors.GetCount(); i++)
+						FbxNodeAttribute* Attr = ParentNode->GetNodeAttribute();
+						if (Attr && Attr->GetAttributeType() == FbxNodeAttribute::eNull)
 						{
-							FbxAMatrix mat = pMissingAncestors.GetAt(i)->EvaluateGlobalTransform(FBXSDK_TIME_ZERO);
-							CurrentPose->Add(pMissingAncestors.GetAt(i), mat);
+							// found it 
+							break;
 						}
 
-						pMissingAncestors.Clear();
-						pMissingDeformers.Clear();
-						pMissingDeformersAncestors.Clear();
-						pWrongMatrices.Clear();
+						// find next parent
+						ParentNode = ParentNode->GetParent();
+					}
 
-						// check it again
-						if (CurrentPose->IsValidBindPose(Joint))
-						{
-							if (AcquireBindPoseMatrix(CurrentPose, GlobalBindPoseJointMatrix, Joint))
-							{
-								return true;
-							}
-						}
-						else
+					if (ParentNode && CurrentPose->IsValidBindPose(ParentNode))
+					{
+						if (AcquireBindPoseMatrix(CurrentPose, GlobalBindPoseJointMatrix, Joint))
 						{
-							// first try to find parent who is null group and see if you can try test it again
-							FbxNode* ParentNode = Joint->GetParent();
-							while (ParentNode)
-							{
-								FbxNodeAttribute* Attr = ParentNode->GetNodeAttribute();
-								if (Attr && Attr->GetAttributeType() == FbxNodeAttribute::eNull)
-								{
-									// found it 
-									break;
-								}
-
-								// find next parent
-								ParentNode = ParentNode->GetParent();
-							}
-
-							if (ParentNode && CurrentPose->IsValidBindPose(ParentNode))
-							{
-								if (AcquireBindPoseMatrix(CurrentPose, GlobalBindPoseJointMatrix, Joint))
-								{
-									return true;
-								}
-							}
+							return true;
 						}
 					}
 				}
 			}
+		}
+	}
 
-			return false;
-		};
+	return false;
+} 
 
-	bool bRerieveBindPoseResult = RetrievePoseFromBindPose();
+bool FFbxMesh::GetGlobalJointBindPoseTransform(FFbxParser* Parser, FbxScene* SDKScene, FbxNode* Joint, FFbxJointMeshBindPoseGenerator& FbxJointMeshBindPoseGenerator,
+	FbxAMatrix& GlobalBindPoseJointMatrix, TMap<FString, FMatrix>& MeshIdToGlobalBindPoseJointMap, TMap<FString, FMatrix>& MeshIdToGlobalBindPoseReferenceMap, 
+	bool& bBadBindPoseMessageDisplay)
+{
+	FbxManager* SDKManager = SDKScene->GetFbxManager();
+	
+	GlobalBindPoseJointMatrix = FbxJointMeshBindPoseGenerator.GetBindPoseFbxAMatrix(Joint, GlobalBindPoseJointMatrix);
+	MeshIdToGlobalBindPoseReferenceMap = FbxJointMeshBindPoseGenerator.GetMeshIdToGlobalBindPoseReferenceMap(Joint);
+	MeshIdToGlobalBindPoseJointMap = FbxJointMeshBindPoseGenerator.GetMeshIdToGlobalBindPoseJointMap(Joint);
+	
+	if (MeshIdToGlobalBindPoseReferenceMap.Num() > 0)
+	{
+		return true;
+	}
+
+	// TODO: Do we expect all Joints to have meshes assigned to them?
+	if (!(Joint->GetMesh() && Joint->GetMesh()->GetDeformerCount(FbxDeformer::eSkin) > 0))
+	{
+		return false;
+	}
+
+	bool bRetrieveBindPoseResult = RetrievePoseFromBindPose(Joint, GlobalBindPoseJointMatrix);
 	// get bind pose
-	if (!bRerieveBindPoseResult)
+	if (!bRetrieveBindPoseResult)
 	{
 		// if failed, delete bind pose, and retry.
 		const int32 PoseCount = SDKScene->GetPoseCount();
@@ -2137,9 +2154,9 @@ bool FFbxMesh::GetGlobalJointBindPoseTransform(FFbxParser* Parser, FbxScene* SDK
 		}
 
 		SDKManager->CreateMissingBindPoses(SDKScene);
-		bRerieveBindPoseResult = RetrievePoseFromBindPose();
+		bRetrieveBindPoseResult = RetrievePoseFromBindPose(Joint, GlobalBindPoseJointMatrix);
 	}
-	return bRerieveBindPoseResult;
+	return bRetrieveBindPoseResult;
 }
 
 bool FFbxMesh::ExtractSkinnedMeshNodeJoints(FbxScene* SDKScene, UInterchangeBaseNodeContainer& NodeContainer, FbxMesh* Mesh, UInterchangeMeshNode* MeshNode)
@@ -2207,8 +2224,6 @@ UInterchangeMeshNode* FFbxMesh::CreateMeshNode(UInterchangeBaseNodeContainer& No
 
 FFbxJointMeshBindPoseGenerator::FFbxJointMeshBindPoseGenerator(FbxScene* SDKScene, FFbxParser& Parser)
 {
-	//First look for Cluster and then look in the bind pose if no cluster was found
-
 	//Search all skeletalmesh(FbxGeometry with valid deformer) using this joint and see if there is a valid FbxCluster
 	//containing a TransformLinkMatrix for this joint
 	const int32 GeometryCount = SDKScene->GetGeometryCount();
@@ -2224,7 +2239,7 @@ FFbxJointMeshBindPoseGenerator::FFbxJointMeshBindPoseGenerator(FbxScene* SDKScen
 		const int32 GeometryDeformerCount = Geometry->GetDeformerCount(FbxDeformer::eSkin);
 		for (int32 GeometryDeformerIndex = 0; GeometryDeformerIndex < GeometryDeformerCount; ++GeometryDeformerIndex)
 		{
-			const FbxSkin* Skin = (FbxSkin*)Geometry->GetDeformer(GeometryDeformerIndex, FbxDeformer::eSkin);
+			FbxSkin* Skin = (FbxSkin*)Geometry->GetDeformer(GeometryDeformerIndex, FbxDeformer::eSkin);
 			if (!ensure(Skin))
 			{
 				continue;
@@ -2232,7 +2247,7 @@ FFbxJointMeshBindPoseGenerator::FFbxJointMeshBindPoseGenerator(FbxScene* SDKScen
 			const int32 ClusterCount = Skin->GetClusterCount();
 			for (int32 ClusterIndex = 0; ClusterIndex < ClusterCount; ClusterIndex++)
 			{
-				const FbxCluster* Cluster = Skin->GetCluster(ClusterIndex);
+				FbxCluster* Cluster = Skin->GetCluster(ClusterIndex);
 				// When Maya plug-in exports rigid binding, it will generate "CompensationCluster" for each ancestor links.
 				// FBX writes these "CompensationCluster" out. The CompensationCluster also has weight 1 for vertices.
 				// Unreal importer should skip these clusters.
@@ -2241,13 +2256,42 @@ FFbxJointMeshBindPoseGenerator::FFbxJointMeshBindPoseGenerator(FbxScene* SDKScen
 					continue;
 				}
 
-				const FbxNode* Joint = Cluster->GetLink();
+				FbxNode* Joint = Cluster->GetLink();
 
 				uint64 ClusterId = Cluster->GetUniqueID();
 				uint64 MeshId = Geometry->GetUniqueID();
 
+				// Here we try fetching the bind pose only from actual FbxPose objects. Interchange traditionally preferred reading the
+				// TransformLink directly from the cluster, but it seems that in some edge cases (see UE-273136) that is not the correct thing to do?
+				// By checking the pose first, we mirror the FbxLegacy behavior shown at UnFbx::FFbxImporter::ImportBones, with how it first 
+				// checks the PoseArray when filling in GlobalsPerLink
 				FbxAMatrix GlobalBindPoseJointMatrixPerGeo;
-				Cluster->GetTransformLinkMatrix(GlobalBindPoseJointMatrixPerGeo);
+				bool bRetrieveBindPoseResult = RetrievePoseFromBindPose(Joint, GlobalBindPoseJointMatrixPerGeo);
+				if (!bRetrieveBindPoseResult)
+				{
+					// if failed, delete bind pose, and retry.
+					// This is likely not needed as FFbxScene::IsValidBindPose will already create this fallback pose, but just in case
+					const int32 PoseCount = SDKScene->GetPoseCount();
+					for (int32 PoseIndex = PoseCount - 1; PoseIndex >= 0; --PoseIndex)
+					{
+						FbxPose* CurrentPose = SDKScene->GetPose(PoseIndex);
+
+						// current pose is bind pose, 
+						if (CurrentPose && CurrentPose->IsBindPose())
+						{
+							SDKScene->RemovePose(PoseIndex);
+							CurrentPose->Destroy();
+						}
+					}
+
+					SDKScene->GetFbxManager()->CreateMissingBindPoses(SDKScene);
+					bRetrieveBindPoseResult = RetrievePoseFromBindPose(Joint, GlobalBindPoseJointMatrixPerGeo);
+				}
+
+				if (!bRetrieveBindPoseResult)
+				{
+					Cluster->GetTransformLinkMatrix(GlobalBindPoseJointMatrixPerGeo);
+				}
 
 				TMap<FString, FMatrix>& MeshIdToGlobalBindPoseJointMap = JointToMeshIdToGlobalBindPoseJointMap.FindOrAdd(Joint);
 				MeshIdToGlobalBindPoseJointMap.Add(MeshUniqueID, FFbxConvert::ConvertMatrix<FMatrix>(GlobalBindPoseJointMatrixPerGeo));
