Interchange MaterialX: remove unused code for MaterialX 1.38

With the engine's version of MaterialX now at 1.39.3, we can remove
code for earlier versions that is no longer used around handling of
swizzle nodes.

diff --git a/Engine/Plugins/Interchange/Runtime/Source/Common/Public/MaterialX/InterchangeMaterialXDefinitions.h b/Engine/Plugins/Interchange/Runtime/Source/Common/Public/MaterialX/InterchangeMaterialXDefinitions.h
--- a/Engine/Plugins/Interchange/Runtime/Source/Common/Public/MaterialX/InterchangeMaterialXDefinitions.h
+++ b/Engine/Plugins/Interchange/Runtime/Source/Common/Public/MaterialX/InterchangeMaterialXDefinitions.h
@@ -726,9 +726,6 @@ MATERIALX_NAMESPACE_BEGIN
 		static constexpr const char* Separate2 = "separate2";
 		static constexpr const char* Separate3 = "separate3";
 		static constexpr const char* Separate4 = "separate4";
-#if MATERIALX_MINOR_VERSION < 39
-		static constexpr const char* Swizzle = "swizzle";
-#endif // MATERIALX_MINOR_VERSION < 39
 		// Procedural nodes 
 		static constexpr const char* Constant = "constant";
 		// Procedural2D nodes
diff --git a/Engine/Plugins/Interchange/Runtime/Source/Import/Private/MaterialX/MaterialXUtils/MaterialXSurfaceShaderAbstract.cpp b/Engine/Plugins/Interchange/Runtime/Source/Import/Private/MaterialX/MaterialXUtils/MaterialXSurfaceShaderAbstract.cpp
--- a/Engine/Plugins/Interchange/Runtime/Source/Import/Private/MaterialX/MaterialXUtils/MaterialXSurfaceShaderAbstract.cpp
+++ b/Engine/Plugins/Interchange/Runtime/Source/Import/Private/MaterialX/MaterialXUtils/MaterialXSurfaceShaderAbstract.cpp
@@ -74,21 +74,6 @@ bool FMaterialXSurfaceShaderAbstract::AddAttributeFromValueOrInterface(MaterialX
 		UInterchangeShaderNode* ShaderNodeToConnectTo = ShaderNode;
 		FString InputToConnectTo = InputChannelName;
 
-#if MATERIALX_MINOR_VERSION < 39
-		if(Input->hasChannels())
-		{
-			using namespace UE::Interchange::Materials::Standard::Nodes;
-
-			UInterchangeShaderNode* SwizzleNode = CreateShaderNode(Input, Input->getParent()->asA<mx::Node>()->getName().c_str() + FString{ TEXT("_Channels_") } + Input->getName().c_str(), Swizzle::Name.ToString());
-			SwizzleNode->AddStringAttribute(Swizzle::Attributes::Channels.ToString(), Input->getChannels().c_str());
-
-			UInterchangeShaderPortsAPI::ConnectDefaultOuputToInput(ShaderNode, InputChannelName, SwizzleNode->GetUniqueID());
-
-			ShaderNodeToConnectTo = SwizzleNode;
-			InputToConnectTo = TEXT("Input");
-		}
-#endif // MATERIALX_MINOR_VERSION < 39
-
 		if(Input->hasValue())
 		{
 			bAttribute = AddAttribute(Input, InputToConnectTo, ShaderNodeToConnectTo, OutputIndex);
@@ -392,11 +377,6 @@ void FMaterialXSurfaceShaderAbstract::ConnectNodeCategoryOutputToInput(const Mat
 
 		FString OutputChannelName = OutputName;
 
-#if MATERIALX_MINOR_VERSION < 39
-		// Swizzle node for the attribute 'channels'
-		UInterchangeShaderNode* ChannelsNode = nullptr;
-#endif // MATERIALX_MINOR_VERSION < 39
-
 		if(mx::ElementPtr DownstreamElement = Edge.getDownstreamElement())
 		{
 			if(mx::NodePtr DownstreamNode = DownstreamElement->asA<mx::Node>())
@@ -409,15 +389,6 @@ void FMaterialXSurfaceShaderAbstract::ConnectNodeCategoryOutputToInput(const Mat
 					{
 						OutputChannelName = ConnectedInput->getOutputString().c_str();
 					}
-
-#if MATERIALX_MINOR_VERSION < 39
-					if(ConnectedInput->hasChannels())
-					{
-						using namespace UE::Interchange::Materials::Standard::Nodes;
-						ChannelsNode = CreateShaderNode(ConnectedInput, (UpstreamNode->getName() + ConnectedInput->getName()).c_str() + FString{ TEXT("_Channels") }, Swizzle::Name.ToString());
-						ChannelsNode->AddStringAttribute(Swizzle::Attributes::Channels.ToString(), ConnectedInput->getChannels().c_str());
-					}
-#endif // MATERIALX_MINOR_VERSION < 39
 				}
 
 				std::vector<mx::OutputPtr> Outputs{ DownstreamNode->getActiveOutputs()};
@@ -431,16 +402,6 @@ void FMaterialXSurfaceShaderAbstract::ConnectNodeCategoryOutputToInput(const Mat
 				{
 					if(UInterchangeShaderNode** FoundNode = ShaderNodes.Find({ GetAttributeParentName(DownstreamNode, ConnectedInput), Outputs[i]->getName().c_str()}))
 					{
-#if MATERIALX_MINOR_VERSION < 39
-						//Connect the swizzle node between the upstream and downstream node
-						if(ChannelsNode)
-						{
-							UInterchangeShaderPortsAPI::ConnectDefaultOuputToInput(*FoundNode, InputChannelName, ChannelsNode->GetUniqueID());
-							InputChannelName = TEXT("Input");
-							FoundNode = &ChannelsNode;
-						}
-#endif // MATERIALX_MINOR_VERSION < 39
-
 						ParentShaderNodeOutputs.Emplace(*FoundNode);
 					}
 				}
@@ -1463,25 +1424,6 @@ void FMaterialXSurfaceShaderAbstract::ConnectSwitchInputToOutput(const FConnectN
 	UInterchangeShaderPortsAPI::ConnectDefaultOuputToInput(Connect.ParentShaderNode, Connect.InputChannelName, SwitchNode->GetUniqueID());
 }
 
-#if MATERIALX_MINOR_VERSION < 39
-void FMaterialXSurfaceShaderAbstract::ConnectSwizzleInputToOutput(const FConnectNode& Connect)
-{
-	using namespace UE::Interchange::Materials::Standard::Nodes;
-
-	UInterchangeShaderNode* SwizzleNode = CreateShaderNode(Connect.UpstreamNode, Connect.UpstreamNode->getName().c_str(), Swizzle::Name.ToString());
-	if(mx::InputPtr Input = Connect.UpstreamNode->getInput("in"))
-	{
-		AddAttributeFromValueOrInterface(Input, GetInputName(Input), SwizzleNode);
-	}
-	if(mx::InputPtr Input = Connect.UpstreamNode->getInput("channels"))
-	{
-		SwizzleNode->AddStringAttribute(Swizzle::Attributes::Channels.ToString(), Input->getValueString().c_str());
-	}
-
-	UInterchangeShaderPortsAPI::ConnectDefaultOuputToInput(Connect.ParentShaderNode, Connect.InputChannelName, SwizzleNode->GetUniqueID());
-}
-#endif // MATERIALX_MINOR_VERSION < 39
-
 void FMaterialXSurfaceShaderAbstract::ConnectNormalMapInputToOutput(const FConnectNode& Connect)
 {
 	using namespace UE::Interchange::Materials::Standard::Nodes;
@@ -1802,9 +1744,6 @@ void FMaterialXSurfaceShaderAbstract::RegisterConnectNodeOutputToInputDelegates(
 	MatchingConnectNodeDelegates.Add(mx::Category::Separate3,		FOnConnectNodeOutputToInput::CreateSP(this, &FMaterialXSurfaceShaderAbstract::ConnectSeparateInputToOutput));
 	MatchingConnectNodeDelegates.Add(mx::Category::Separate4,		FOnConnectNodeOutputToInput::CreateSP(this, &FMaterialXSurfaceShaderAbstract::ConnectSeparateInputToOutput));
 	MatchingConnectNodeDelegates.Add(mx::Category::TexCoord,		FOnConnectNodeOutputToInput::CreateSP(this, &FMaterialXSurfaceShaderAbstract::ConnectTexCoordInputToOutput));
-#if MATERIALX_MINOR_VERSION < 39
-	MatchingConnectNodeDelegates.Add(mx::Category::Swizzle,			FOnConnectNodeOutputToInput::CreateSP(this, &FMaterialXSurfaceShaderAbstract::ConnectSwizzleInputToOutput));
-#endif // MATERIALX_MINOR_VERSION < 39
 	MatchingConnectNodeDelegates.Add(mx::Category::Switch,			FOnConnectNodeOutputToInput::CreateSP(this, &FMaterialXSurfaceShaderAbstract::ConnectSwitchInputToOutput));
 	MatchingConnectNodeDelegates.Add(mx::Category::Refract,			FOnConnectNodeOutputToInput::CreateSP(this, &FMaterialXSurfaceShaderAbstract::ConnectRefractInputToOutput));
 	MatchingConnectNodeDelegates.Add(mx::Category::ViewDirection,	FOnConnectNodeOutputToInput::CreateSP(this, &FMaterialXSurfaceShaderAbstract::ConnectViewDirectionInputToOutput));
diff --git a/Engine/Plugins/Interchange/Runtime/Source/Import/Public/MaterialX/MaterialXUtils/MaterialXSurfaceShaderAbstract.h b/Engine/Plugins/Interchange/Runtime/Source/Import/Public/MaterialX/MaterialXUtils/MaterialXSurfaceShaderAbstract.h
--- a/Engine/Plugins/Interchange/Runtime/Source/Import/Public/MaterialX/MaterialXUtils/MaterialXSurfaceShaderAbstract.h
+++ b/Engine/Plugins/Interchange/Runtime/Source/Import/Public/MaterialX/MaterialXUtils/MaterialXSurfaceShaderAbstract.h
@@ -315,11 +315,6 @@ protected:
 	/** <switch> */
 	void ConnectSwitchInputToOutput(const FConnectNode& Connect);
 
-#if MATERIALX_MINOR_VERSION < 39
-	/** <swizzle> */
-	void ConnectSwizzleInputToOutput(const FConnectNode& Connect);
-#endif // MATERIALX_MINOR_VERSION < 39
-
 	/** <normalmap> */
 	void ConnectNormalMapInputToOutput(const FConnectNode& Connect);
 
diff --git a/Engine/Plugins/Interchange/Runtime/Source/Nodes/Public/InterchangeMaterialDefinitions.h b/Engine/Plugins/Interchange/Runtime/Source/Nodes/Public/InterchangeMaterialDefinitions.h
--- a/Engine/Plugins/Interchange/Runtime/Source/Nodes/Public/InterchangeMaterialDefinitions.h
+++ b/Engine/Plugins/Interchange/Runtime/Source/Nodes/Public/InterchangeMaterialDefinitions.h
@@ -190,21 +190,6 @@ namespace Interchange
 					}
 				}
 
-				namespace Swizzle
-				{
-					const FName Name = TEXT("MaterialXSwizzle");
-
-					namespace Inputs
-					{
-						const FName Input = TEXT("Input");
-					}
-
-					namespace Attributes
-					{
-						const FName Channels = TEXT("Channels");
-					}
-				}
-
 				namespace TextureCoordinate
 				{
 					const FName Name = TEXT("TextureCoordinate");
diff --git a/Engine/Plugins/Interchange/Runtime/Source/Pipelines/Private/InterchangeGenericMaterialPipeline.cpp b/Engine/Plugins/Interchange/Runtime/Source/Pipelines/Private/InterchangeGenericMaterialPipeline.cpp
--- a/Engine/Plugins/Interchange/Runtime/Source/Pipelines/Private/InterchangeGenericMaterialPipeline.cpp
+++ b/Engine/Plugins/Interchange/Runtime/Source/Pipelines/Private/InterchangeGenericMaterialPipeline.cpp
@@ -55,7 +55,6 @@
 #include "Materials/MaterialExpressionVectorParameter.h"
 #include "Materials/MaterialInstanceConstant.h"
 #include "Materials/MaterialInstanceDynamic.h"
-#include "MaterialX/MaterialExpressions/MaterialExpressionSwizzle.h"
 #include "MaterialX/MaterialExpressions/MaterialExpressionTextureSampleParameterBlur.h"
 #include "Misc/CoreMisc.h"
 #include "Misc/PackageName.h"
@@ -2858,29 +2857,6 @@ void UInterchangeGenericMaterialPipeline::HandleVectorNoiseNode(const UInterchan
 	}
 }
 
-void UInterchangeGenericMaterialPipeline::HandleSwizzleNode(const UInterchangeShaderNode* ShaderNode, UInterchangeBaseMaterialFactoryNode* MaterialFactoryNode, UInterchangeMaterialExpressionFactoryNode* SwizzleFactoryNode)
-{
-	using namespace UE::Interchange::Materials::Standard::Nodes;
-	SwizzleFactoryNode->SetCustomExpressionClassName(UMaterialExpressionMaterialXSwizzle::StaticClass()->GetName());
-	// Input
-	{
-		TTuple<UInterchangeMaterialExpressionFactoryNode*, FString> InputExpression =
-			CreateMaterialExpressionForInput(MaterialFactoryNode, ShaderNode, Swizzle::Inputs::Input.ToString(), SwizzleFactoryNode->GetUniqueID());
-		if(InputExpression.Get<0>())
-		{
-			UInterchangeShaderPortsAPI::ConnectOuputToInputByName(SwizzleFactoryNode, GET_MEMBER_NAME_CHECKED(UMaterialExpressionMaterialXSwizzle, Input).ToString(),
-																  InputExpression.Get<0>()->GetUniqueID(), InputExpression.Get<1>());
-		}
-	}
-	// Channels
-	if(FString Channels; ShaderNode->GetStringAttribute(Swizzle::Attributes::Channels.ToString(), Channels))
-	{
-		const FName ChannelsMemberName = GET_MEMBER_NAME_CHECKED(UMaterialExpressionMaterialXSwizzle, Channels);
-		SwizzleFactoryNode->AddStringAttribute(ChannelsMemberName.ToString(), Channels);
-		SwizzleFactoryNode->AddApplyAndFillDelegates<FString>(ChannelsMemberName.ToString(), UMaterialExpressionMaterialXSwizzle::StaticClass(), ChannelsMemberName);
-	}
-}
-
 void UInterchangeGenericMaterialPipeline::HandleSwitchNode(const UInterchangeShaderNode* ShaderNode, UInterchangeBaseMaterialFactoryNode* MaterialFactoryNode, UInterchangeMaterialExpressionFactoryNode* SwitchFactoryNode)
 {
 	using namespace UE::Interchange::Materials::Standard::Nodes;
@@ -3142,10 +3118,6 @@ UInterchangeMaterialExpressionFactoryNode* UInterchangeGenericMaterialPipeline::
 	{
 		HandleSwitchNode(ShaderNode, MaterialFactoryNode, MaterialExpression);
 	}
-	else if(*ShaderType == Nodes::Swizzle::Name)
-	{
-		HandleSwizzleNode(ShaderNode, MaterialFactoryNode, MaterialExpression);
-	}
 	else if (*ShaderType == Nodes::TextureCoordinate::Name)
 	{
 		HandleTextureCoordinateNode(ShaderNode, MaterialFactoryNode, MaterialExpression);
diff --git a/Engine/Plugins/Interchange/Runtime/Source/Pipelines/Public/InterchangeGenericMaterialPipeline.h b/Engine/Plugins/Interchange/Runtime/Source/Pipelines/Public/InterchangeGenericMaterialPipeline.h
--- a/Engine/Plugins/Interchange/Runtime/Source/Pipelines/Public/InterchangeGenericMaterialPipeline.h
+++ b/Engine/Plugins/Interchange/Runtime/Source/Pipelines/Public/InterchangeGenericMaterialPipeline.h
@@ -228,7 +228,6 @@ private:
 	UE_API void HandleTransformVectorNode(const UInterchangeShaderNode* ShaderNode, UInterchangeBaseMaterialFactoryNode* MaterialFactoryNode, UInterchangeMaterialExpressionFactoryNode* TransformVectorFactoryNode);
 	UE_API void HandleNoiseNode(const UInterchangeShaderNode* ShaderNode, UInterchangeBaseMaterialFactoryNode* MaterialFactoryNode, UInterchangeMaterialExpressionFactoryNode* NoiseFactoryNode);
 	UE_API void HandleVectorNoiseNode(const UInterchangeShaderNode* ShaderNode, UInterchangeBaseMaterialFactoryNode* MaterialFactoryNode, UInterchangeMaterialExpressionFactoryNode* NoiseFactoryNode);
-	UE_API void HandleSwizzleNode(const UInterchangeShaderNode* ShaderNode, UInterchangeBaseMaterialFactoryNode* MaterialFactoryNode, UInterchangeMaterialExpressionFactoryNode* SwizzleFactoryNode);
 	UE_API void HandleSwitchNode(const UInterchangeShaderNode* ShaderNode, UInterchangeBaseMaterialFactoryNode* MaterialFactoryNode, UInterchangeMaterialExpressionFactoryNode* SwitchFactoryNode);
 	UE_API void HandleSlabBSDFNode(const UInterchangeShaderNode* ShaderNode, UInterchangeBaseMaterialFactoryNode* MaterialFactoryNode, UInterchangeMaterialExpressionFactoryNode* SlabBSDFFactoryNode);
 	UE_API void HandleTrigonometryNode(const UInterchangeShaderNode* ShaderNode, UClass* StaticClass, UInterchangeBaseMaterialFactoryNode* MaterialFactoryNode, UInterchangeMaterialExpressionFactoryNode* TrigonometryFactoryNode);
