Interchange: Fix finding materials by SearchLocation no longer working when import materials was disabled.

diff --git a/Engine/Plugins/Interchange/Runtime/Source/Pipelines/Private/InterchangeGenericMaterialPipeline.cpp b/Engine/Plugins/Interchange/Runtime/Source/Pipelines/Private/InterchangeGenericMaterialPipeline.cpp
--- a/Engine/Plugins/Interchange/Runtime/Source/Pipelines/Private/InterchangeGenericMaterialPipeline.cpp
+++ b/Engine/Plugins/Interchange/Runtime/Source/Pipelines/Private/InterchangeGenericMaterialPipeline.cpp
@@ -1057,33 +1057,32 @@ void UInterchangeGenericMaterialPipeline::ExecutePipeline(UInterchangeBaseNodeCo
 
 		bool bEnableFactoryNode = bCombinedImportMaterials;
 
-		if (bEnableFactoryNode)
+		// Note: Even if import materials is fully disabled we still want to search for existing materials to use,
+		// and even potentially re-enable the node if we're reimporting
+		FString MaterialName = BaseMaterialFactoryNode->GetDisplayLabel();
+		if (UMaterialInterface* ExistingMaterial = FindExistingMaterialFromSearchLocation(MaterialName, ContentBasePath, SearchLocation))
 		{
-			FString MaterialName = BaseMaterialFactoryNode->GetDisplayLabel();
-			if (UMaterialInterface* ExistingMaterial = FindExistingMaterialFromSearchLocation(MaterialName, ContentBasePath, SearchLocation))
-			{
-				BaseMaterialFactoryNode->SetCustomReferenceObject(ExistingMaterial);
+			BaseMaterialFactoryNode->SetCustomReferenceObject(ExistingMaterial);
 
-				// We're reimporting the existing material we found: Enable the factory node if we can.
-				// TODO: Can this potentially lead to issues, since the existing material to reimport is discovered by name and search location?
-				if (CacheContextParam.ReimportAsset && CacheContextParam.ReimportAsset == ExistingMaterial)
-				{
-					const bool bIsMaterial = BaseMaterialFactoryNode->IsA<UInterchangeMaterialFactoryNode>()
-											 && ExistingMaterial->IsA<UMaterial>();
+			// We're reimporting the existing material we found: Enable the factory node if we can.
+			// TODO: Can this potentially lead to issues, since the existing material to reimport is discovered by name and search location?
+			if (CacheContextParam.ReimportAsset && CacheContextParam.ReimportAsset == ExistingMaterial)
+			{
+				const bool bIsMaterial = BaseMaterialFactoryNode->IsA<UInterchangeMaterialFactoryNode>()
+											&& ExistingMaterial->IsA<UMaterial>();
 
-					const bool bIsMaterialinstance = BaseMaterialFactoryNode->IsA<UInterchangeMaterialInstanceFactoryNode>()
-													 && ExistingMaterial->IsA<UMaterialInstance>();
+				const bool bIsMaterialinstance = BaseMaterialFactoryNode->IsA<UInterchangeMaterialInstanceFactoryNode>()
+													&& ExistingMaterial->IsA<UMaterialInstance>();
 
-					// Reimport can only be done on material instances
-					bEnableFactoryNode = bEnableFactoryNode && !bIsMaterial && bIsMaterialinstance;
-				}
-				// Found an existing material on the search location during an import --> Just use it as-is
-				else
-				{
-					// Disabling the factory node here is important so that we don't overwrite the existing asset.
-					// The existing material will still actually be used for the import's material assignments and etc.
-					bEnableFactoryNode = false;
-				}
+				// Reimport can only be done on material instances
+				bEnableFactoryNode = !bIsMaterial && bIsMaterialinstance;
+			}
+			// Found an existing material on the search location during an import --> Just use it as-is
+			else
+			{
+				// Disabling the factory node here is important so that we don't overwrite the existing asset.
+				// The existing material will still actually be used for the import's material assignments and etc.
+				bEnableFactoryNode = false;
 			}
 		}
 
