[BUGFIX] Texture incorrectly set when importing glTF via Interchange.
- Same texture being used as a normal texture now creates a duplicate texture with all the normal map requirements for GLTF.

#rnx

diff --git a/Engine/Plugins/Interchange/Runtime/Source/Import/Private/Gltf/InterchangeGltfTranslator.cpp b/Engine/Plugins/Interchange/Runtime/Source/Import/Private/Gltf/InterchangeGltfTranslator.cpp
--- a/Engine/Plugins/Interchange/Runtime/Source/Import/Private/Gltf/InterchangeGltfTranslator.cpp
+++ b/Engine/Plugins/Interchange/Runtime/Source/Import/Private/Gltf/InterchangeGltfTranslator.cpp
@@ -621,6 +621,7 @@ bool UInterchangeGLTFTranslator::Translate( UInterchangeBaseNodeContainer& NodeC
 		{
 			AddTextureIndex(GltfMaterial.Normal.TextureIndex);
 			AddTextureIndex(GltfMaterial.ClearCoat.NormalMap.TextureIndex);
+			AddTextureIndex(GltfMaterial.Packing.NormalMap.TextureIndex);
 		}
 	}
 
diff --git a/Engine/Plugins/Interchange/Runtime/Source/Parsers/GLTFCore/Private/GLTF/GLTFAsset.cpp b/Engine/Plugins/Interchange/Runtime/Source/Parsers/GLTFCore/Private/GLTF/GLTFAsset.cpp
--- a/Engine/Plugins/Interchange/Runtime/Source/Parsers/GLTFCore/Private/GLTF/GLTFAsset.cpp
+++ b/Engine/Plugins/Interchange/Runtime/Source/Parsers/GLTFCore/Private/GLTF/GLTFAsset.cpp
@@ -55,6 +55,67 @@ namespace GLTF
 		}
 	}
 
+	// Prefixes
+	FString GLTFNodePrefixStaticHelper::GetScenePrefix(const FString& AssetName)
+	{
+		return AssetName + TEXT("_scene_");
+	}
+
+	FString GLTFNodePrefixStaticHelper::GetNodePrefix(const FString& AssetName)
+	{
+		return AssetName + TEXT("_node_");
+	}
+
+	FString GLTFNodePrefixStaticHelper::GetJointPrefix(const FString& AssetName)
+	{
+		return AssetName + TEXT("_joint_");
+	}
+
+	FString GLTFNodePrefixStaticHelper::GetTexturePrefix(const FString& AssetName)
+	{
+		return AssetName + TEXT("_texture_");
+	}
+
+	FString GLTFNodePrefixStaticHelper::GetCameraPrefix(const FString& AssetName)
+	{
+		return AssetName + TEXT("_camera_");
+	}
+
+	FString GLTFNodePrefixStaticHelper::GetLightPrefix(const FString& AssetName)
+	{
+		return AssetName + TEXT("_light_");
+	}
+
+	FString GLTFNodePrefixStaticHelper::GetMaterialPrefix(const FString& AssetName)
+	{
+		return AssetName + TEXT("_material_");
+	}
+
+	FString GLTFNodePrefixStaticHelper::GetSkinPrefix(const FString& AssetName)
+	{
+		return AssetName + TEXT("_skin_");
+	}
+
+	FString GLTFNodePrefixStaticHelper::GetAnimationPrefix(const FString& AssetName)
+	{
+		return AssetName + TEXT("_animation_");
+	}
+
+	FString GLTFNodePrefixStaticHelper::GetImagePrefix(const FString& AssetName)
+	{
+		return AssetName + TEXT("_image_");
+	}
+
+	FString GLTFNodePrefixStaticHelper::GetMeshPrefix(const FString& AssetName)
+	{
+		return AssetName + TEXT("_mesh_");
+	}
+
+	FString GLTFNodePrefixStaticHelper::GetIESPrefix(const FString& AssetName)
+	{
+		return AssetName + TEXT("_IES_");
+	}
+
 	const FMetadata::FExtraData* FMetadata::GetExtraData(const FString& Name) const
 	{
 		return Extras.FindByPredicate([&Name](const FMetadata::FExtraData& Data) { return Data.Name == Name; });
@@ -102,7 +163,7 @@ namespace GLTF
 	void FAsset::GenerateNames()
 	{
 		{
-			const FString ScenePrefix = Name + TEXT("_scene_");
+			const FString ScenePrefix = GLTFNodePrefixStaticHelper::GetScenePrefix(Name);
 			for (size_t SceneIndex = 0; SceneIndex < Scenes.Num(); SceneIndex++)
 			{
 				FScene& Scene = Scenes[SceneIndex];
@@ -125,8 +186,8 @@ namespace GLTF
 
 			TSet<FString> Names;
 
-			const FString NodePrefix = Name + TEXT("_node_");
-			const FString JointPrefix = Name + TEXT("_joint_");
+			const FString NodePrefix = GLTFNodePrefixStaticHelper::GetNodePrefix(Name);
+			const FString JointPrefix = GLTFNodePrefixStaticHelper::GetJointPrefix(Name);
 			for (int32 NodeIndex = 0; NodeIndex < Nodes.Num(); ++NodeIndex)
 			{
 				FNode& Node = Nodes[NodeIndex];
@@ -155,7 +216,7 @@ namespace GLTF
 		}
 
 		{
-			const FString TexPrefix = Name + TEXT("_texture_");
+			const FString TexPrefix = GLTFNodePrefixStaticHelper::GetTexturePrefix(Name);
 			for (int32 TextureIndex = 0; TextureIndex < Textures.Num(); ++TextureIndex)
 			{
 				FTexture& Tex = Textures[TextureIndex];
@@ -192,7 +253,7 @@ namespace GLTF
 		}
 
 		{
-			const FString CameraPrefix = Name + TEXT("_camera_");
+			const FString CameraPrefix = GLTFNodePrefixStaticHelper::GetCameraPrefix(Name);
 			for (int32 CameraIndex = 0; CameraIndex < Cameras.Num(); ++CameraIndex)
 			{
 				FCamera& Camera = Cameras[CameraIndex];
@@ -216,7 +277,7 @@ namespace GLTF
 		}
 
 		{
-			const FString LightPrefix = Name + TEXT("_light_");
+			const FString LightPrefix = GLTFNodePrefixStaticHelper::GetLightPrefix(Name);
 			for (int32 LightIndex = 0; LightIndex < Lights.Num(); ++LightIndex)
 			{
 				FLight& Light = Lights[LightIndex];
@@ -244,12 +305,12 @@ namespace GLTF
 		// UniqueId will consist of: [Prefix]+[_ElementType_]+[ElementIndex]
 		// Name will be the provided GLTF object.name, if present,
 		// if Name is not present the UniqueId will be used as the Name.
-		GLTF::GenerateNames(Name + TEXT("_material_"), Materials);
-		GLTF::GenerateNames(Name + TEXT("_skin_"), Skins);
-		GLTF::GenerateNames(Name + TEXT("_animation_"), Animations);
-		GLTF::GenerateNames(Name + TEXT("_image_"), Images);
-		GLTF::GenerateNames(Name + TEXT("_mesh_"), Meshes);
-		GLTF::GenerateNames(Name + TEXT("_IES_"), LightsIES);
+		GLTF::GenerateNames(GLTFNodePrefixStaticHelper::GetMaterialPrefix(Name), Materials);
+		GLTF::GenerateNames(GLTFNodePrefixStaticHelper::GetSkinPrefix(Name), Skins);
+		GLTF::GenerateNames(GLTFNodePrefixStaticHelper::GetAnimationPrefix(Name), Animations);
+		GLTF::GenerateNames(GLTFNodePrefixStaticHelper::GetImagePrefix(Name), Images);
+		GLTF::GenerateNames(GLTFNodePrefixStaticHelper::GetMeshPrefix(Name), Meshes);
+		GLTF::GenerateNames(GLTFNodePrefixStaticHelper::GetIESPrefix(Name), LightsIES);
 
 		// Validate FName criteria for Materials Textures Meshes Animations (assets)
 		GLTF::ValidateFNameCriteria(TEXT("Material_"), Materials);
diff --git a/Engine/Plugins/Interchange/Runtime/Source/Parsers/GLTFCore/Private/GLTF/GLTFReader.cpp b/Engine/Plugins/Interchange/Runtime/Source/Parsers/GLTFCore/Private/GLTF/GLTFReader.cpp
--- a/Engine/Plugins/Interchange/Runtime/Source/Parsers/GLTFCore/Private/GLTF/GLTFReader.cpp
+++ b/Engine/Plugins/Interchange/Runtime/Source/Parsers/GLTFCore/Private/GLTF/GLTFReader.cpp
@@ -980,6 +980,8 @@ namespace GLTF
 			}
 		}
 		OutAsset.GenerateNames();
+		
+		ResolveConflictingTextures();
 
 		if (OutAsset.ValidationCheck() != FAsset::Valid)
 		{
@@ -1707,6 +1709,100 @@ namespace GLTF
 		}
 	}
 
+	void FFileReader::ResolveConflictingTextures()
+	{
+		TSet<int32> DuplicateTexturesForNormalMap;
+
+		struct FAffectedTextureMaps
+		{
+			TArray<FTextureMap*> TextureMaps;
+		};
+
+		TMap<int32, FAffectedTextureMaps> TextureMapsRequiringUpdate;
+
+		auto AddTextureIndex = [&](TSet<int32>& IndicesSet, int32 TextureIndex)
+		{
+			if (Asset->Textures.IsValidIndex(TextureIndex))
+			{
+				IndicesSet.Add(TextureIndex);
+			}
+		};
+
+		auto AddTextureIndexToMap = [&](TMap<int32, GLTF::FTextureMap*>& IndexToTextureMap, GLTF::FTextureMap* TextureMap)
+		{
+			if (TextureMap && Asset->Textures.IsValidIndex(TextureMap->TextureIndex))
+			{
+				IndexToTextureMap.Emplace(TextureMap->TextureIndex, TextureMap);
+			}
+		};
+
+		for (GLTF::FMaterial& Material : Asset->Materials)
+		{
+			TSet<int32> NormalMapTextureIndices;
+			TSet<int32> ColoredTextureIndices;
+			TMap<int32, GLTF::FTextureMap*> IdToTextureMap;
+
+			// Textures that are used as normal maps
+			AddTextureIndex(NormalMapTextureIndices, Material.Normal.TextureIndex);
+			AddTextureIndex(NormalMapTextureIndices, Material.ClearCoat.NormalMap.TextureIndex);
+			AddTextureIndex(NormalMapTextureIndices, Material.Packing.NormalMap.TextureIndex);
+
+			// To quickly identify the FTextureMaps from the Texture Index
+			AddTextureIndexToMap(IdToTextureMap, &Material.Normal);
+			AddTextureIndexToMap(IdToTextureMap, &Material.ClearCoat.NormalMap);
+			AddTextureIndexToMap(IdToTextureMap, &Material.Packing.NormalMap);
+
+			// Textures that are not used as normal maps
+			AddTextureIndex(ColoredTextureIndices, Material.BaseColor.TextureIndex);
+			AddTextureIndex(ColoredTextureIndices, Material.MetallicRoughness.Map.TextureIndex);
+			AddTextureIndex(ColoredTextureIndices, Material.SpecularGlossiness.Map.TextureIndex);
+			AddTextureIndex(ColoredTextureIndices, Material.ClearCoat.ClearCoatMap.TextureIndex);
+			AddTextureIndex(ColoredTextureIndices, Material.ClearCoat.RoughnessMap.TextureIndex);
+			AddTextureIndex(ColoredTextureIndices, Material.Transmission.TransmissionMap.TextureIndex);
+			AddTextureIndex(ColoredTextureIndices, Material.Sheen.SheenColorMap.TextureIndex);
+			AddTextureIndex(ColoredTextureIndices, Material.Sheen.SheenRoughnessMap.TextureIndex);
+			AddTextureIndex(ColoredTextureIndices, Material.Specular.SpecularMap.TextureIndex);
+			AddTextureIndex(ColoredTextureIndices, Material.Specular.SpecularColorMap.TextureIndex);
+			AddTextureIndex(ColoredTextureIndices, Material.Packing.Map.TextureIndex);
+			AddTextureIndex(ColoredTextureIndices, Material.Iridescence.Texture.TextureIndex);
+			AddTextureIndex(ColoredTextureIndices, Material.Iridescence.Thickness.Texture.TextureIndex);
+			AddTextureIndex(ColoredTextureIndices, Material.Anisotropy.Texture.TextureIndex);
+			AddTextureIndex(ColoredTextureIndices, Material.Occlusion.TextureIndex);
+			AddTextureIndex(ColoredTextureIndices, Material.Emissive.TextureIndex);
+
+			TSet<int32> ConflictingTextureIndices = NormalMapTextureIndices.Intersect(ColoredTextureIndices);
+			if (!ConflictingTextureIndices.IsEmpty())
+			{
+				for (const int32& ConflictingTextureIndex : ConflictingTextureIndices)
+				{
+					if (!TextureMapsRequiringUpdate.Contains(ConflictingTextureIndex))
+					{
+						TextureMapsRequiringUpdate.Emplace(ConflictingTextureIndex, FAffectedTextureMaps());
+					}
+
+					DuplicateTexturesForNormalMap.Add(ConflictingTextureIndex);
+					TextureMapsRequiringUpdate[ConflictingTextureIndex].TextureMaps.Add(IdToTextureMap[ConflictingTextureIndex]);
+				}
+			}
+		}
+
+		// Add the Duplicate textures for normal maps.
+		int32 TextureIndex = Asset->Textures.Num();
+		for (const int32& NormalMapIndex : DuplicateTexturesForNormalMap)
+		{
+			FTexture DuplicateTexture = Asset->Textures[NormalMapIndex];
+			DuplicateTexture.Name += TEXT("_greenFlipped");
+			DuplicateTexture.UniqueId = GLTFNodePrefixStaticHelper::GetTexturePrefix(Asset->Name) + FString::FromInt(TextureIndex);
+			Asset->Textures.Add(DuplicateTexture);
+			FAffectedTextureMaps& AffectedTextureMaps = TextureMapsRequiringUpdate.FindChecked(NormalMapIndex);
+			for (FTextureMap* AffectedTextureMap : AffectedTextureMaps.TextureMaps)
+			{
+				AffectedTextureMap->TextureIndex = TextureIndex;
+			}
+			TextureIndex++;
+		}
+	}
+
 }  // namespace GLTF
 
 #undef LOCTEXT_NAMESPACE
\ No newline at end of file
diff --git a/Engine/Plugins/Interchange/Runtime/Source/Parsers/GLTFCore/Public/GLTFAsset.h b/Engine/Plugins/Interchange/Runtime/Source/Parsers/GLTFCore/Public/GLTFAsset.h
--- a/Engine/Plugins/Interchange/Runtime/Source/Parsers/GLTFCore/Public/GLTFAsset.h
+++ b/Engine/Plugins/Interchange/Runtime/Source/Parsers/GLTFCore/Public/GLTFAsset.h
@@ -165,4 +165,20 @@ namespace GLTF
 
 	GLTFCORE_API const TCHAR* ToString(EExtension Extension);
 
+	struct GLTFNodePrefixStaticHelper
+	{
+		GLTFCORE_API static FString GetScenePrefix(const FString& AssetName);
+		GLTFCORE_API static FString GetNodePrefix(const FString& AssetName);
+		GLTFCORE_API static FString GetJointPrefix(const FString& AssetName);
+		GLTFCORE_API static FString GetTexturePrefix(const FString& AssetName);
+		GLTFCORE_API static FString GetCameraPrefix(const FString& AssetName);
+		GLTFCORE_API static FString GetLightPrefix(const FString& AssetName);
+		GLTFCORE_API static FString GetMaterialPrefix(const FString& AssetName);
+		GLTFCORE_API static FString GetSkinPrefix(const FString& AssetName);
+		GLTFCORE_API static FString GetAnimationPrefix(const FString& AssetName);
+		GLTFCORE_API static FString GetImagePrefix(const FString& AssetName);
+		GLTFCORE_API static FString GetMeshPrefix(const FString& AssetName);
+		GLTFCORE_API static FString GetIESPrefix(const FString& AssetName);
+	};
+
 }  // namespace GLTF
diff --git a/Engine/Plugins/Interchange/Runtime/Source/Parsers/GLTFCore/Public/GLTFReader.h b/Engine/Plugins/Interchange/Runtime/Source/Parsers/GLTFCore/Public/GLTFReader.h
--- a/Engine/Plugins/Interchange/Runtime/Source/Parsers/GLTFCore/Public/GLTFReader.h
+++ b/Engine/Plugins/Interchange/Runtime/Source/Parsers/GLTFCore/Public/GLTFReader.h
@@ -73,6 +73,7 @@ namespace GLTF
 		int32 FindRootJointIndex(int32 CurrentNodeIndex) const;
 		void BuildRootJoints() const;
 
+		void ResolveConflictingTextures();
 	private:
 		bool CheckForErrors(int32 StartIndex = 0) const;
 
