[Interchange-FBX] Fixing front x axis option.

#rnx

diff --git a/Engine/Plugins/Interchange/Runtime/Source/Parsers/Fbx/Private/FbxAPI.cpp b/Engine/Plugins/Interchange/Runtime/Source/Parsers/Fbx/Private/FbxAPI.cpp
--- a/Engine/Plugins/Interchange/Runtime/Source/Parsers/Fbx/Private/FbxAPI.cpp
+++ b/Engine/Plugins/Interchange/Runtime/Source/Parsers/Fbx/Private/FbxAPI.cpp
@@ -141,7 +141,7 @@ namespace UE
 
 				//We always convert scene to UE axis and units
 				FbxAMatrix AxisConversionInverseMatrix;
-				FFbxConvert::ConvertScene(SDKScene, bConvertScene, bForceFrontXAxis, bConvertSceneUnit, FileDetails.AxisDirection, FileDetails.UnitSystem, AxisConversionInverseMatrix);
+				FFbxConvert::ConvertScene(SDKScene, bConvertScene, bForceFrontXAxis, bConvertSceneUnit, FileDetails.AxisDirection, FileDetails.UnitSystem, AxisConversionInverseMatrix, JointOrientationMatrix);
 
 				//Save the AxisConversionInverseTransform into InterchangeSourceNode (so that socket transport can use it accordingly).
 				FTransform AxisConversionInverseTransform = FFbxConvert::ConvertTransform<FTransform, FVector, FQuat>(AxisConversionInverseMatrix);
diff --git a/Engine/Plugins/Interchange/Runtime/Source/Parsers/Fbx/Private/FbxAPI.h b/Engine/Plugins/Interchange/Runtime/Source/Parsers/Fbx/Private/FbxAPI.h
--- a/Engine/Plugins/Interchange/Runtime/Source/Parsers/Fbx/Private/FbxAPI.h
+++ b/Engine/Plugins/Interchange/Runtime/Source/Parsers/Fbx/Private/FbxAPI.h
@@ -110,6 +110,8 @@ namespace UE
 
 				const FString& GetSourceFilename() const { return SourceFilename; }
 
+				FbxAMatrix JointOrientationMatrix;
+
 				/**
 				 * Critical section to avoid getting multiple payload in same time.
 				 * The FBX evaluator use a cache mechanism for evaluating global transform that is not thread safe.
diff --git a/Engine/Plugins/Interchange/Runtime/Source/Parsers/Fbx/Private/FbxAnimation.cpp b/Engine/Plugins/Interchange/Runtime/Source/Parsers/Fbx/Private/FbxAnimation.cpp
--- a/Engine/Plugins/Interchange/Runtime/Source/Parsers/Fbx/Private/FbxAnimation.cpp
+++ b/Engine/Plugins/Interchange/Runtime/Source/Parsers/Fbx/Private/FbxAnimation.cpp
@@ -513,6 +513,8 @@ namespace UE::Interchange::Private
 		if (ParentNode)
 		{
 			FbxAMatrix NodeTransform = FetchPayloadData.Node->EvaluateGlobalTransform(Currenttime);
+			NodeTransform = NodeTransform * Parser.JointOrientationMatrix;
+
 			FTransform GlobalTransform = UE::Interchange::Private::FFbxConvert::ConvertTransform<FTransform, FVector, FQuat>(NodeTransform);
 			if (GlobalTransform.ContainsNaN())
 			{
@@ -521,6 +523,11 @@ namespace UE::Interchange::Private
 			}
 
 			FbxAMatrix ParentTransform = ParentNode->EvaluateGlobalTransform(Currenttime);
+			if (!FetchPayloadData.bIsRootJoint)
+			{
+				ParentTransform = ParentTransform * Parser.JointOrientationMatrix;
+			}
+
 			FTransform ParentGlobalTransform = UE::Interchange::Private::FFbxConvert::ConvertTransform<FTransform, FVector, FQuat>(ParentTransform);
 			if (ParentGlobalTransform.ContainsNaN())
 			{
@@ -532,7 +539,9 @@ namespace UE::Interchange::Private
 		}
 		else
 		{
-			FbxAMatrix& LocalMatrix = FetchPayloadData.Node->EvaluateLocalTransform(Currenttime);
+			FbxAMatrix LocalMatrix = FetchPayloadData.Node->EvaluateLocalTransform(Currenttime);
+			LocalMatrix = LocalMatrix * Parser.JointOrientationMatrix;
+
 			FbxVector4 NewLocalT = LocalMatrix.GetT();
 			FbxVector4 NewLocalS = LocalMatrix.GetS();
 			FbxQuaternion NewLocalQ = LocalMatrix.GetQ();
@@ -807,7 +816,8 @@ namespace UE::Interchange::Private
 		, UInterchangeSceneNode* SceneNode
 		, TMap<FString, TSharedPtr<FPayloadContextBase>>& PayloadContexts
 		, UInterchangeSkeletalAnimationTrackNode* SkeletalAnimationTrackNode
-		, const int32& AnimationIndex)
+		, const int32& AnimationIndex
+		, bool bIsRootJoint)
 	{
 		FGetFbxTransformCurvesParameters Parameters(SDKScene, Node);
 		GetFbxTransformCurves(Parameters, AnimationIndex);
@@ -845,6 +855,7 @@ namespace UE::Interchange::Private
 				FNodeTransformFetchPayloadData FetchPayloadData;
 				FetchPayloadData.Node = Node;
 				FetchPayloadData.CurrentAnimStack = Parameters.CurrentAnimStack;
+				FetchPayloadData.bIsRootJoint = bIsRootJoint;
 
 				AnimPayload->NodeTransformFetchPayloadData = FetchPayloadData;
 				PayloadContexts.Add(PayLoadKey, AnimPayload);
diff --git a/Engine/Plugins/Interchange/Runtime/Source/Parsers/Fbx/Private/FbxAnimation.h b/Engine/Plugins/Interchange/Runtime/Source/Parsers/Fbx/Private/FbxAnimation.h
--- a/Engine/Plugins/Interchange/Runtime/Source/Parsers/Fbx/Private/FbxAnimation.h
+++ b/Engine/Plugins/Interchange/Runtime/Source/Parsers/Fbx/Private/FbxAnimation.h
@@ -29,6 +29,7 @@ namespace UE::Interchange::Private
 	{
 		FbxNode* Node = nullptr;
 		FbxAnimStack* CurrentAnimStack = nullptr;
+		bool bIsRootJoint = false;
 	};
 
 	//User defined animations:
@@ -106,7 +107,8 @@ namespace UE::Interchange::Private
 			, UInterchangeSceneNode* UnrealNode
 			, TMap<FString, TSharedPtr<FPayloadContextBase>>& PayloadContexts
 			, UInterchangeSkeletalAnimationTrackNode* SkeletalAnimationTrackNode
-			, const int32& AnimationIndex);
+			, const int32& AnimationIndex
+			, bool bIsRootJoint);
 
 		/** This function add the payload key for an animated node user attribute. */
 		static void AddNodeAttributeCurvesAnimation(FFbxParser& Parser
diff --git a/Engine/Plugins/Interchange/Runtime/Source/Parsers/Fbx/Private/FbxConvert.cpp b/Engine/Plugins/Interchange/Runtime/Source/Parsers/Fbx/Private/FbxConvert.cpp
--- a/Engine/Plugins/Interchange/Runtime/Source/Parsers/Fbx/Private/FbxConvert.cpp
+++ b/Engine/Plugins/Interchange/Runtime/Source/Parsers/Fbx/Private/FbxConvert.cpp
@@ -55,7 +55,10 @@ namespace UE
 				return AxisDirection;
 			}
 
-			void FFbxConvert::ConvertScene(FbxScene* SDKScene, const bool bConvertScene, const bool bForceFrontXAxis, const bool bConvertSceneUnit, FString& FileSystemDirection, FString& FileUnitSystem, FbxAMatrix& AxisConversionInverseMatrix)
+			void FFbxConvert::ConvertScene(FbxScene* SDKScene, 
+				const bool bConvertScene, const bool bForceFrontXAxis, const bool bConvertSceneUnit, 
+				FString& FileSystemDirection, FString& FileUnitSystem, 
+				FbxAMatrix& AxisConversionInverseMatrix, FbxAMatrix& JointOrientationMatrix)
 			{
 				if (!ensure(SDKScene))
 				{
@@ -75,6 +78,9 @@ namespace UE
 				FbxAxisSystem FileAxisSystem = SDKScene->GetGlobalSettings().GetAxisSystem();
 				FileSystemDirection = GetFileAxisDirection(FileAxisSystem);
 
+				AxisConversionInverseMatrix.SetIdentity();
+				JointOrientationMatrix.SetIdentity();
+
 				if (bConvertScene)
 				{
 					//UE is: z up, front x, left handed
@@ -96,6 +102,11 @@ namespace UE
 						FbxAMatrix AxisConversionMatrix;
 						AxisConversionMatrix = SourceMatrix.Inverse() * UnrealMatrix;
 						AxisConversionInverseMatrix = AxisConversionMatrix.Inverse();
+
+						if (bForceFrontXAxis)
+						{
+							JointOrientationMatrix.SetR(FbxVector4(-90.0, -90.0, 0.0));
+						}
 					}
 				}
 
diff --git a/Engine/Plugins/Interchange/Runtime/Source/Parsers/Fbx/Private/FbxConvert.h b/Engine/Plugins/Interchange/Runtime/Source/Parsers/Fbx/Private/FbxConvert.h
--- a/Engine/Plugins/Interchange/Runtime/Source/Parsers/Fbx/Private/FbxConvert.h
+++ b/Engine/Plugins/Interchange/Runtime/Source/Parsers/Fbx/Private/FbxConvert.h
@@ -162,7 +162,10 @@ namespace UE
 				/**
 				 * Convert a fbx scene 
 				 */
-				static void ConvertScene(FbxScene* SDKScene, const bool bConvertScene, const bool bForceFrontXAxis, const bool bConvertSceneUnit, FString& FileSystemDirection, FString& FileUnitSystem, FbxAMatrix& AxisConversionInverseMatrix);
+				static void ConvertScene(FbxScene* SDKScene, 
+					const bool bConvertScene, const bool bForceFrontXAxis, const bool bConvertSceneUnit, 
+					FString& FileSystemDirection, FString& FileUnitSystem, 
+					FbxAMatrix& AxisConversionInverseMatrix, FbxAMatrix& JointOrientationMatrix);
 				
 				/** Scene Conversion API End */
 				//////////////////////////////////////////////////////////////////////////
diff --git a/Engine/Plugins/Interchange/Runtime/Source/Parsers/Fbx/Private/FbxMesh.cpp b/Engine/Plugins/Interchange/Runtime/Source/Parsers/Fbx/Private/FbxMesh.cpp
--- a/Engine/Plugins/Interchange/Runtime/Source/Parsers/Fbx/Private/FbxMesh.cpp
+++ b/Engine/Plugins/Interchange/Runtime/Source/Parsers/Fbx/Private/FbxMesh.cpp
@@ -2116,7 +2116,7 @@ bool RetrievePoseFromBindPose(FbxNode* Joint, FbxAMatrix& GlobalBindPoseJointMat
 
 bool FFbxMesh::GetGlobalJointBindPoseTransform(FFbxParser* Parser, FbxScene* SDKScene, FbxNode* Joint, FFbxJointMeshBindPoseGenerator& FbxJointMeshBindPoseGenerator,
 	FbxAMatrix& GlobalBindPoseJointMatrix, TMap<FString, FMatrix>& MeshIdToGlobalBindPoseJointMap, TMap<FString, FMatrix>& MeshIdToGlobalBindPoseReferenceMap, 
-	bool& bBadBindPoseMessageDisplay)
+	bool& bBadBindPoseMessageDisplay, const FbxAMatrix& JointOrientationMatrix)
 {
 	FbxManager* SDKManager = SDKScene->GetFbxManager();
 	
@@ -2156,6 +2156,12 @@ bool FFbxMesh::GetGlobalJointBindPoseTransform(FFbxParser* Parser, FbxScene* SDK
 		SDKManager->CreateMissingBindPoses(SDKScene);
 		bRetrieveBindPoseResult = RetrievePoseFromBindPose(Joint, GlobalBindPoseJointMatrix);
 	}
+
+	if (bRetrieveBindPoseResult)
+	{
+		GlobalBindPoseJointMatrix = GlobalBindPoseJointMatrix * JointOrientationMatrix;
+	}
+
 	return bRetrieveBindPoseResult;
 }
 
@@ -2292,6 +2298,7 @@ FFbxJointMeshBindPoseGenerator::FFbxJointMeshBindPoseGenerator(FbxScene* SDKScen
 				{
 					Cluster->GetTransformLinkMatrix(GlobalBindPoseJointMatrixPerGeo);
 				}
+				GlobalBindPoseJointMatrixPerGeo = GlobalBindPoseJointMatrixPerGeo * Parser.JointOrientationMatrix;
 
 				TMap<FString, FMatrix>& MeshIdToGlobalBindPoseJointMap = JointToMeshIdToGlobalBindPoseJointMap.FindOrAdd(Joint);
 				MeshIdToGlobalBindPoseJointMap.Add(MeshUniqueID, FFbxConvert::ConvertMatrix<FMatrix>(GlobalBindPoseJointMatrixPerGeo));
@@ -2303,6 +2310,7 @@ FFbxJointMeshBindPoseGenerator::FFbxJointMeshBindPoseGenerator(FbxScene* SDKScen
 
 				FbxAMatrix GlobalBindPoseReferenceMatrix;
 				Cluster->GetTransformMatrix(GlobalBindPoseReferenceMatrix);
+				GlobalBindPoseReferenceMatrix = GlobalBindPoseReferenceMatrix * Parser.JointOrientationMatrix;
 
 				TMap<FString, FMatrix>& MeshIdToGlobalBindPoseReferenceMap = JointToMeshIdToGlobalBindPoseReferenceMap.FindOrAdd(Joint);
 				MeshIdToGlobalBindPoseReferenceMap.Add(MeshUniqueID, FFbxConvert::ConvertMatrix<FMatrix>(GlobalBindPoseReferenceMatrix));
diff --git a/Engine/Plugins/Interchange/Runtime/Source/Parsers/Fbx/Private/FbxMesh.h b/Engine/Plugins/Interchange/Runtime/Source/Parsers/Fbx/Private/FbxMesh.h
--- a/Engine/Plugins/Interchange/Runtime/Source/Parsers/Fbx/Private/FbxMesh.h
+++ b/Engine/Plugins/Interchange/Runtime/Source/Parsers/Fbx/Private/FbxMesh.h
@@ -212,7 +212,7 @@ namespace UE
 				void AddAllMeshes(FbxScene* SDKScene, FbxGeometryConverter* SDKGeometryConverter, UInterchangeBaseNodeContainer& NodeContainer, TMap<FString, TSharedPtr<FPayloadContextBase>>& PayloadContexts);
 				static bool GetGlobalJointBindPoseTransform(FFbxParser* Parser, FbxScene* SDKScene, FbxNode* Joint, FFbxJointMeshBindPoseGenerator& FbxJointMeshBindPoseGenerator,
 					FbxAMatrix& GlobalBindPoseJointMatrix, TMap<FString, FMatrix>& MeshIdToGlobalBindPoseJointMap, TMap<FString, FMatrix>& MeshIdToGlobalBindPoseReferenceMap, 
-					bool& bBadBindPoseMessageDisplay);
+					bool& bBadBindPoseMessageDisplay, const FbxAMatrix& JointOrientationMatrix);
 
 				TArray<FMorphTargetAnimationBuildingData>& GetMorphTargetAnimationsBuildingData() { return MorphTargetAnimationsBuildingData; }
 			protected:
diff --git a/Engine/Plugins/Interchange/Runtime/Source/Parsers/Fbx/Private/FbxScene.cpp b/Engine/Plugins/Interchange/Runtime/Source/Parsers/Fbx/Private/FbxScene.cpp
--- a/Engine/Plugins/Interchange/Runtime/Source/Parsers/Fbx/Private/FbxScene.cpp
+++ b/Engine/Plugins/Interchange/Runtime/Source/Parsers/Fbx/Private/FbxScene.cpp
@@ -202,13 +202,22 @@ namespace UE
 						}
 					}
 
+					bool bIsRootJoint = false;
+					if (const UInterchangeSceneNode* UEParentNode = Cast< UInterchangeSceneNode >(NodeContainer.GetNode(UnrealNode->GetParentUid())))
+					{
+						if (!UEParentNode->IsSpecializedTypeContains(FSceneNodeStaticData::GetJointSpecializeTypeString()))
+						{
+							bIsRootJoint = true;
+						}
+					}
+
 					//Add the joint specialized type
 					UnrealNode->AddSpecializedType(FSceneNodeStaticData::GetJointSpecializeTypeString());
 					//Get the bind pose transform for this joint
 					FbxAMatrix GlobalBindPoseJointMatrix = SDKScene->GetAnimationEvaluator()->GetNodeGlobalTransform(Node, 0);
 					TMap<FString, FMatrix> MeshIdToGlobalBindPoseReferenceMap;
 					TMap<FString, FMatrix> MeshIdToGlobalBindPoseJointMap;
-					FFbxMesh::GetGlobalJointBindPoseTransform(&Parser, SDKScene, Node, FbxJointMeshBindPoseGenerator, GlobalBindPoseJointMatrix, MeshIdToGlobalBindPoseJointMap, MeshIdToGlobalBindPoseReferenceMap, bBadBindPoseMessageDisplay);
+					FFbxMesh::GetGlobalJointBindPoseTransform(&Parser, SDKScene, Node, FbxJointMeshBindPoseGenerator, GlobalBindPoseJointMatrix, MeshIdToGlobalBindPoseJointMap, MeshIdToGlobalBindPoseReferenceMap, bBadBindPoseMessageDisplay, Parser.JointOrientationMatrix);
 
 					FTransform GlobalBindPoseJointTransform = GetConvertedTransform(GlobalBindPoseJointMatrix);
 					UnrealNode->SetGlobalBindPoseReferenceForMeshUIDs(MeshIdToGlobalBindPoseReferenceMap);
@@ -227,7 +236,7 @@ namespace UE
 						FbxAMatrix GlobalFbxParentMatrix = SDKScene->GetAnimationEvaluator()->GetNodeGlobalTransform(ParentNode, 0);
 						TMap<FString, FMatrix> ParentMeshIdToGlobalBindPoseReferenceMap;
 						TMap<FString, FMatrix> ParentMeshIdToGlobalBindPoseJointMap;
-						FFbxMesh::GetGlobalJointBindPoseTransform(&Parser, SDKScene, ParentNode, FbxJointMeshBindPoseGenerator, GlobalFbxParentMatrix, ParentMeshIdToGlobalBindPoseJointMap, ParentMeshIdToGlobalBindPoseReferenceMap, bBadBindPoseMessageDisplay);
+						FFbxMesh::GetGlobalJointBindPoseTransform(&Parser, SDKScene, ParentNode, FbxJointMeshBindPoseGenerator, GlobalFbxParentMatrix, ParentMeshIdToGlobalBindPoseJointMap, ParentMeshIdToGlobalBindPoseReferenceMap, bBadBindPoseMessageDisplay, bIsRootJoint ? FbxAMatrix() : Parser.JointOrientationMatrix);
 						
 						FbxAMatrix LocalFbxMatrix = GlobalFbxParentMatrix.Inverse() * GlobalBindPoseJointMatrix;
 						FTransform LocalBindPoseJointTransform = GetConvertedTransform(LocalFbxMatrix);
@@ -252,11 +261,19 @@ namespace UE
 
 						//Set the global node transform
 						FbxAMatrix GlobalFbxMatrix = SDKScene->GetAnimationEvaluator()->GetNodeGlobalTransform(Node, 0);
+						GlobalFbxMatrix = GlobalFbxMatrix * Parser.JointOrientationMatrix;
+
 						FTransform GlobalTransform = GetConvertedTransform(GlobalFbxMatrix);
 
 						if (ParentNode != nullptr)
 						{
 							FbxAMatrix GlobalFbxParentMatrix = SDKScene->GetAnimationEvaluator()->GetNodeGlobalTransform(ParentNode, 0);
+
+							if (!bIsRootJoint)
+							{
+								GlobalFbxParentMatrix = GlobalFbxParentMatrix * Parser.JointOrientationMatrix;
+							}
+
 							FbxAMatrix	LocalFbxMatrix = GlobalFbxParentMatrix.Inverse() * GlobalFbxMatrix;
 							FTransform LocalTransform = GetConvertedTransform(LocalFbxMatrix);
 							UnrealNode->SetCustomTimeZeroLocalTransform(&NodeContainer, LocalTransform, bResetCache);
@@ -954,7 +971,7 @@ namespace UE
 					else if (SkeletalAnimationTrackNode)
 					{
 						//Scene node transform can be animated, add the transform animation payload key.
-						if (FFbxAnimation::AddSkeletalTransformAnimation(NodeContainer, SDKScene, Parser, Node, UnrealNode, PayloadContexts, SkeletalAnimationTrackNode, AnimationIndex)
+						if (FFbxAnimation::AddSkeletalTransformAnimation(NodeContainer, SDKScene, Parser, Node, UnrealNode, PayloadContexts, SkeletalAnimationTrackNode, AnimationIndex, bNewSkeltalAnimationStarted)
 							&& !SkeletalAnimationAddedToContainer)
 						{
 							SkeletalAnimationAddedToContainer = true;
