Interchange FBX: Work around prefix-based collision meshes not finding their render meshes if they had a name collision.

diff --git a/Engine/Plugins/Interchange/Runtime/Source/Pipelines/Private/InterchangeGenericStaticMeshPipeline.cpp b/Engine/Plugins/Interchange/Runtime/Source/Pipelines/Private/InterchangeGenericStaticMeshPipeline.cpp
--- a/Engine/Plugins/Interchange/Runtime/Source/Pipelines/Private/InterchangeGenericStaticMeshPipeline.cpp
+++ b/Engine/Plugins/Interchange/Runtime/Source/Pipelines/Private/InterchangeGenericStaticMeshPipeline.cpp
@@ -142,26 +142,74 @@ static TTuple<EInterchangeMeshCollision, FString> GetCollisionMeshType(
 
 		// If we find a mesh named the same as the collision mesh (following the collision prefix), we have a match
 		// e.g. this will match 'UBX_House' with a mesh called 'House'
-
-		const FString* CorrespondingMeshUid = AllNodeUids.FindByPredicate(MatchPredicate(MeshName.RightChop(FirstUnderscore + 1)));
-		if (CorrespondingMeshUid)
 		{
-			return { CollisionType, *CorrespondingMeshUid };
+			const FString RenderMeshName = MeshName.RightChop(FirstUnderscore + 1);
+			const FString* CorrespondingMeshUid = AllNodeUids.FindByPredicate(MatchPredicate(RenderMeshName));
+			if (CorrespondingMeshUid)
+			{
+				return { CollisionType, *CorrespondingMeshUid };
+			}
 		}
 
 		// Otherwise strip the final underscore suffix from the collision mesh name and look again
 		// e.g. this will match 'UBX_House_01' with a mesh called 'House'
-
 		if (FirstUnderscore != LastUnderscore)
 		{
-			// Even if the mesh has a prefix, only treat it as a collision mesh if it has a corresponding render mesh
-			CorrespondingMeshUid = AllNodeUids.FindByPredicate(MatchPredicate(MeshName.Mid(FirstUnderscore + 1, LastUnderscore - FirstUnderscore - 1)));
+			const FString RenderMeshName = MeshName.Mid(FirstUnderscore + 1, LastUnderscore - FirstUnderscore - 1);
+			const FString* CorrespondingMeshUid = AllNodeUids.FindByPredicate(MatchPredicate(RenderMeshName));
 			if (CorrespondingMeshUid)
 			{
 				return { CollisionType, *CorrespondingMeshUid };
 			}
 		}
 
+		// It is possible our target mesh has been renamed away by a name sanitizing process (see UE-248981 and
+		// FFbxParser::EnsureNodeNameAreValid). Since it is unusual to have a prefixed collision mesh and no target, let's 
+		// assume we actually have a render mesh somewhere and try stripping any nameclash-looking suffixes from mesh names 
+		// to see if we have a match
+		{
+			// Like above start by just stripping the prefix (e.g. 'UBX_House_01' --> 'House_01')
+			FString RenderMeshName = MeshName.RightChop(FirstUnderscore + 1);
+
+			auto NameCollisionSearchPredicate = [&RenderMeshName, &PipelineMeshesUtilities, &NodeContainer](const FString& Other)
+			{
+				FString MeshName = UE::Interchange::Private::GetNodeName(PipelineMeshesUtilities, NodeContainer, Other);
+				if (!MeshName.StartsWith(RenderMeshName))
+				{
+					return false;
+				}
+
+				FString MeshNameSuffix = MeshName.RightChop(RenderMeshName.Len());
+
+				// Check if everything after the render mesh name is just some combination of "ncl" (nameclash), numbers and/or underscores
+				MeshNameSuffix = MeshNameSuffix.Replace(TEXT("ncl"), TEXT(""));
+
+				FString LastChar = MeshNameSuffix.Right(1);
+				while ((LastChar.IsNumeric() || LastChar == TEXT("_")) && MeshNameSuffix.Len() > 0)
+				{
+					MeshNameSuffix.LeftChopInline(1, EAllowShrinking::No);
+					LastChar = MeshNameSuffix.Right(1);
+				}
+
+				// If there's nothing left, everything after the RenderMeshName was nameclash stuff, so this is likely our mesh
+				return MeshNameSuffix.Len() == 0;
+			};
+
+			const FString* CorrespondingMeshUid = AllNodeUids.FindByPredicate(NameCollisionSearchPredicate);
+			if (CorrespondingMeshUid)
+			{
+				return {CollisionType, *CorrespondingMeshUid};
+			}
+
+			// If we still have nothing, try also stripping numbered suffixes from the render mesh name and repeating the search
+			// (e.g. 'UBX_House_01' --> 'House')
+			RenderMeshName = MeshName.Mid(FirstUnderscore + 1, LastUnderscore - FirstUnderscore - 1);
+			CorrespondingMeshUid = AllNodeUids.FindByPredicate(NameCollisionSearchPredicate);
+			if (CorrespondingMeshUid)
+			{
+				return {CollisionType, *CorrespondingMeshUid};
+			}
+		}
 	}
 
 	return { EInterchangeMeshCollision::None, FString() };
