[BUGFIX]Interchange - Crash when exiting the editor after importing into level with packed actor
- The newly created world in the InterchangeLevelFactory doesn't inform the engine of world.

#rnx

diff --git a/Engine/Plugins/Interchange/Runtime/Source/Import/Private/Scene/InterchangeLevelFactory.cpp b/Engine/Plugins/Interchange/Runtime/Source/Import/Private/Scene/InterchangeLevelFactory.cpp
--- a/Engine/Plugins/Interchange/Runtime/Source/Import/Private/Scene/InterchangeLevelFactory.cpp
+++ b/Engine/Plugins/Interchange/Runtime/Source/Import/Private/Scene/InterchangeLevelFactory.cpp
@@ -69,7 +69,7 @@ namespace UE::Interchange::Private::InterchangeLevelFactory
 			const bool bAddToRoot = false;
 
 			constexpr bool bEnableWorldPartitionStreaming = false;
-			constexpr bool bInformEngineOfWorld = true;
+			constexpr bool bInformEngineOfWorld = false;
 			bool bIsRuntimeOrPIE = false;
 			if (UInterchangeEditorUtilitiesBase* EditorUtilities = UInterchangeManager::GetInterchangeManager().GetEditorUtilities())
 			{
