Interchange:
- fix tests failing due to upgrade to MaterialX 1.39.3
- removing channels test, since 1.39.3 "channels" attribute is now deprecated

diff --git a/Engine/Plugins/Interchange/Runtime/Source/Import/Private/MaterialX/MaterialXUtils/MaterialXManager.cpp b/Engine/Plugins/Interchange/Runtime/Source/Import/Private/MaterialX/MaterialXUtils/MaterialXManager.cpp
--- a/Engine/Plugins/Interchange/Runtime/Source/Import/Private/MaterialX/MaterialXUtils/MaterialXManager.cpp
+++ b/Engine/Plugins/Interchange/Runtime/Source/Import/Private/MaterialX/MaterialXUtils/MaterialXManager.cpp
@@ -297,6 +297,7 @@ FMaterialXManager::FMaterialXManager()
 		mx::Category::SrgbTextureToLinRec709,
 		mx::Category::StandardSurface,
 		mx::Category::TiledCircles,
+		mx::Category::Time, // since 1.39.3, they added a nodegraph in stdlib/genglsl/stdlib_genglsl_impl.mtlx instead of stdlib_ng.mtlx
 		mx::Category::UnifiedNoise3D,
 		mx::Category::UsdPreviewSurface
 	}
