[Interchange] Fixing the material ordering convention.

#rnx

diff --git a/Engine/Plugins/Interchange/Runtime/Source/Pipelines/Private/InterchangeGenericMaterialPipeline.cpp b/Engine/Plugins/Interchange/Runtime/Source/Pipelines/Private/InterchangeGenericMaterialPipeline.cpp
--- a/Engine/Plugins/Interchange/Runtime/Source/Pipelines/Private/InterchangeGenericMaterialPipeline.cpp
+++ b/Engine/Plugins/Interchange/Runtime/Source/Pipelines/Private/InterchangeGenericMaterialPipeline.cpp
@@ -822,12 +822,46 @@ void UInterchangeGenericMaterialPipeline::ExecutePipeline(UInterchangeBaseNodeCo
 	TSet<UInterchangeMaterialReferenceNode*> MaterialReferenceNodes;
 	TSet<UInterchangeBaseMaterialFactoryNode*> BaseMaterialFactoryNodes;
 
+	auto FixMaterialName = [](UInterchangeBaseNode* Node)
+		{
+			FString MaterialName = Node->GetDisplayLabel();
+
+			bool bUpdateMaterialName = false;
+			if (MaterialName.Len() > 6)
+			{
+				int32 Offset = MaterialName.Find(TEXT("_SKIN"), ESearchCase::IgnoreCase, ESearchDir::FromEnd);
+				if (Offset != INDEX_NONE)
+				{
+					// Chop off the material name so we are left with the number in _SKINXX
+					FString SkinXXNumber = MaterialName.Right(MaterialName.Len() - (Offset + 1)).RightChop(4);
+
+					if (SkinXXNumber.IsNumeric())
+					{
+						// remove the '_skinXX' suffix from the material name					
+						MaterialName.LeftChopInline(MaterialName.Len() - Offset, EAllowShrinking::No);
+						bUpdateMaterialName = true;
+					}
+				}
+			}
+			else if (MaterialName.IsEmpty())
+			{
+				MaterialName = TEXT("UnnamedMaterial");
+				bUpdateMaterialName = true;
+			}
+
+			if (bUpdateMaterialName)
+			{
+				Node->SetDisplayLabel(MaterialName);
+			}
+		};
+
 	// Find all nodes we need for this pipeline
 	BaseNodeContainer->IterateNodes(
 		[&ShaderGraphNodes,
 		 &MaterialInstanceNodes,
 		 &MaterialReferenceNodes,
-		 &BaseMaterialFactoryNodes](const FString& NodeUid, UInterchangeBaseNode* Node)
+		 &BaseMaterialFactoryNodes,
+		 &FixMaterialName](const FString& NodeUid, UInterchangeBaseNode* Node)
 		{
 			switch (Node->GetNodeContainerType())
 			{
@@ -835,14 +869,17 @@ void UInterchangeGenericMaterialPipeline::ExecutePipeline(UInterchangeBaseNodeCo
 				{
 					if (UInterchangeShaderGraphNode* ShaderGraphNode = Cast<UInterchangeShaderGraphNode>(Node))
 					{
+						FixMaterialName(ShaderGraphNode);
 						ShaderGraphNodes.Add(ShaderGraphNode);
 					}
 					else if (UInterchangeMaterialInstanceNode* InstanceNode = Cast<UInterchangeMaterialInstanceNode>(Node))
 					{
+						FixMaterialName(InstanceNode);
 						MaterialInstanceNodes.Add(InstanceNode);
 					}
 					else if (UInterchangeMaterialReferenceNode* ReferenceNode = Cast<UInterchangeMaterialReferenceNode>(Node))
 					{
+						FixMaterialName(ReferenceNode);
 						MaterialReferenceNodes.Add(ReferenceNode);
 					}
 				}
@@ -850,6 +887,7 @@ void UInterchangeGenericMaterialPipeline::ExecutePipeline(UInterchangeBaseNodeCo
 				{
 					if (UInterchangeBaseMaterialFactoryNode* BaseMaterialFactoryNode = Cast<UInterchangeBaseMaterialFactoryNode>(Node))
 					{
+						FixMaterialName(BaseMaterialFactoryNode);
 						BaseMaterialFactoryNodes.Add(BaseMaterialFactoryNode);
 					}
 				}
diff --git a/Engine/Plugins/Interchange/Runtime/Source/Pipelines/Public/InterchangePipelineMeshesUtilities.h b/Engine/Plugins/Interchange/Runtime/Source/Pipelines/Public/InterchangePipelineMeshesUtilities.h
--- a/Engine/Plugins/Interchange/Runtime/Source/Pipelines/Public/InterchangePipelineMeshesUtilities.h
+++ b/Engine/Plugins/Interchange/Runtime/Source/Pipelines/Public/InterchangePipelineMeshesUtilities.h
@@ -437,14 +437,26 @@ namespace UE::Interchange::MeshesUtilities
 		//Empty all slot dependencies, they will be added backin the correct order
 		FactoryNode.ResetSlotMaterialDependencies();
 
-		TArray<FString> KeyReorder;
+		struct FOrderHelper
+		{
+			int32 Index;
+			FString Name;
+			FOrderHelper(int32 InIndex, const FString& InName)
+				: Index(InIndex)
+				, Name(InName)
+			{
+
+			}
+		};
+
+		TArray<FOrderHelper> KeyReorder;
 		KeyReorder.Reserve(SlotMaterialDependencies.Num());
+
 		TArray<FString> MissingSuffixMaterialNames;
 		MissingSuffixMaterialNames.Reserve(SlotMaterialDependencies.Num());
 		//Reorder material using the skinXX workflow
 		for (const TPair<FString, FString>& SlotMaterialDependency : SlotMaterialDependencies)
 		{
-			bool bHasSuffix = false;
 			const FString& MaterialName = SlotMaterialDependency.Key;
 			if (MaterialName.Len() > 6)
 			{
@@ -458,46 +470,42 @@ namespace UE::Interchange::MeshesUtilities
 						int32 TmpIndex = FPlatformString::Atoi(*SkinXXNumber);
 						if (TmpIndex >= 0)
 						{
-							while (KeyReorder.Num() <= TmpIndex)
-							{
-								KeyReorder.AddDefaulted();
-							}
-							check(KeyReorder.IsValidIndex(TmpIndex));
-							if (KeyReorder[TmpIndex].IsEmpty())
-							{
-								bHasSuffix = true;
-								KeyReorder[TmpIndex] = MaterialName;
-								continue;
-							}
+							KeyReorder.Add(FOrderHelper(TmpIndex, MaterialName));
+							continue;
 						}
 					}
 				}
 			}
 			
-			if (!bHasSuffix)
-			{
-				MissingSuffixMaterialNames.Add(MaterialName);
-			}
+			MissingSuffixMaterialNames.Add(MaterialName);
 		}
 
+		KeyReorder.Sort([](const FOrderHelper& A, const FOrderHelper& B)
+			{
+				if (A.Index == B.Index)
+				{
+					return A.Name < B.Name;
+				}
+				else
+				{
+					return A.Index < B.Index;
+				}
+			});
+
 		//The map is MaterialName, MaterialUid
 		TMap<FString, FString> ReorderedSlotMaterialDependencies;
 		ReorderedSlotMaterialDependencies.Reserve(SlotMaterialDependencies.Num());
 
-		//Start by adding the KeyReorder material
-		for (const FString& MaterialName : KeyReorder)
+		for (const FOrderHelper& Element : KeyReorder)
 		{
-			if (MaterialName.IsEmpty())
-			{
-				continue;
-			}
-			const FString& SlotMaterialUid = SlotMaterialDependencies.FindChecked(MaterialName);
-			ReorderedSlotMaterialDependencies.Add(MaterialName, SlotMaterialUid);
+			const FString& SlotMaterialUid = SlotMaterialDependencies.FindChecked(Element.Name);
+			ReorderedSlotMaterialDependencies.Add(Element.Name, SlotMaterialUid);
 		}
+
 		//Add the missing suffix material
 		for (const FString& MaterialName : MissingSuffixMaterialNames)
 		{
-			const FString& SlotMaterialUid = SlotMaterialDependencies.FindChecked(MaterialName);
+			const FString& SlotMaterialUid = SlotMaterialDependencies.FindChecked(MaterialName); 
 			ReorderedSlotMaterialDependencies.Add(MaterialName, SlotMaterialUid);
 		}
 
