Interchange:
- Fix new assets created during reimport

EngineTest:
- Updated tests with now the assets prefixed with the filename

diff --git a/Engine/Plugins/Interchange/Runtime/Source/Parsers/Fbx/Private/FbxMesh.cpp b/Engine/Plugins/Interchange/Runtime/Source/Parsers/Fbx/Private/FbxMesh.cpp
--- a/Engine/Plugins/Interchange/Runtime/Source/Parsers/Fbx/Private/FbxMesh.cpp
+++ b/Engine/Plugins/Interchange/Runtime/Source/Parsers/Fbx/Private/FbxMesh.cpp
@@ -2224,7 +2224,6 @@ UInterchangeMeshNode* FFbxMesh::CreateMeshNode(UInterchangeBaseNodeContainer& No
 	}
 	// Creating a UMaterialInterface
 	NodeContainer.SetupNode(MeshNode, NodeUid, DisplayLabel, EInterchangeNodeContainerType::TranslatedAsset);
-	MeshNode->AddBooleanAttribute(TEXT("RenameLikeLegacyFbx"), true);
 	return MeshNode;
 }
 
diff --git a/Engine/Plugins/Interchange/Runtime/Source/Parsers/Fbx/Private/FbxScene.cpp b/Engine/Plugins/Interchange/Runtime/Source/Parsers/Fbx/Private/FbxScene.cpp
--- a/Engine/Plugins/Interchange/Runtime/Source/Parsers/Fbx/Private/FbxScene.cpp
+++ b/Engine/Plugins/Interchange/Runtime/Source/Parsers/Fbx/Private/FbxScene.cpp
@@ -887,7 +887,6 @@ namespace UE
 								// (this could potentially be done in a neater way by passing down the stack the SkeleotnNodeUid instead of a potentially 'fake' UInterchangeSkeletalAnimationTrackNode.)
 								SkeletalAnimationTrackNode->InitializeNode(TrackNodeUid, DisplayString, EInterchangeNodeContainerType::TranslatedAsset);
 							}
-							SkeletalAnimationTrackNode->AddBooleanAttribute(TEXT("RenameLikeLegacyFbx"), true);
 
 							double FrameRate = FbxTime::GetFrameRate(SDKScene->GetGlobalSettings().GetTimeMode());
 							SkeletalAnimationTrackNode->SetCustomAnimationSampleRate(FrameRate);
@@ -1260,8 +1259,6 @@ namespace UE
 
 							SkeletalAnimationTrackNode->SetCustomAnimationStartTime(TimeSpan.GetStart().GetSecondDouble());
 							SkeletalAnimationTrackNode->SetCustomAnimationStopTime(TimeSpan.GetStop().GetSecondDouble());
-
-							SkeletalAnimationTrackNode->AddBooleanAttribute(TEXT("RenameLikeLegacyFbx"), true);
 						}
 
 						ProcessCustomAttributes(Parser, CurrentAnimationStack, SkeletalAnimationTrackNode);
diff --git a/Engine/Plugins/Interchange/Runtime/Source/Pipelines/Private/InterchangeGenericAnimationPipeline.cpp b/Engine/Plugins/Interchange/Runtime/Source/Pipelines/Private/InterchangeGenericAnimationPipeline.cpp
--- a/Engine/Plugins/Interchange/Runtime/Source/Pipelines/Private/InterchangeGenericAnimationPipeline.cpp
+++ b/Engine/Plugins/Interchange/Runtime/Source/Pipelines/Private/InterchangeGenericAnimationPipeline.cpp
@@ -75,71 +75,6 @@ namespace UE::Interchange::Private
 		return bContainOnlyJointAnimation;
 	}
 
-	//Legacy fbx has a very particular way for naming animation sequences.
-	void RenameAnimSequenceLikeLegacyFbx(FString& DisplayString, UInterchangeSkeletalAnimationTrackNode& TrackNode, const TArray<const UInterchangeSourceData*>& SourceDatas, UInterchangePipelineBase* OuterMostPipeline, UInterchangeBaseNodeContainer* NodeContainer)
-	{
-		if(!SourceDatas.IsValidIndex(0))
-		{
-			return;
-		}
-		bool bRenameAnimSequenceLegacyFbxWay = false;
-		TrackNode.GetBooleanAttribute(TEXT("RenameLikeLegacyFbx"), bRenameAnimSequenceLegacyFbxWay);
-		if (!bRenameAnimSequenceLegacyFbxWay)
-		{
-			return;
-		}
-
-		FString SkeletonNodeUid;
-		TrackNode.GetCustomSkeletonNodeUid(SkeletonNodeUid);
-		FString BaseFilename = FPaths::GetBaseFilename(SourceDatas[0]->GetFilename());
-		FString MeshName;
-		if (UInterchangeGenericAssetsPipeline* GenericAssetPipeline = Cast<UInterchangeGenericAssetsPipeline>(OuterMostPipeline))
-		{
-			if (!GenericAssetPipeline->bUseSourceNameForAsset)
-			{
-				NodeContainer->BreakableIterateNodesOfType<UInterchangeMeshNode>([&SkeletonNodeUid, &MeshName, &NodeContainer](const FString& NodeUid, UInterchangeMeshNode* MeshNode)
-					{
-						if (MeshNode)
-						{
-							TArray<FString> SkeletonDependencies;
-							MeshNode->GetSkeletonDependencies(SkeletonDependencies);
-							if(SkeletonDependencies.IsValidIndex(0))
-							{
-								if(SkeletonDependencies.Contains(SkeletonNodeUid) || NodeContainer->GetIsAncestor(SkeletonDependencies[0], SkeletonNodeUid))
-								{
-									int32 MeshReferenceCount = MeshNode->GetSceneInstanceUidsCount();
-									if (MeshReferenceCount == 1)
-									{
-										FString SceneNodeUid;
-										MeshNode->GetSceneInstanceUid(0, SceneNodeUid);
-										if (const UInterchangeSceneNode* SceneNode = Cast<UInterchangeSceneNode>(NodeContainer->GetNode(SceneNodeUid)))
-										{
-											MeshName = TEXT("_") + SceneNode->GetDisplayLabel();
-											return true;
-										}
-									}
-									MeshName = TEXT("_") + MeshNode->GetDisplayLabel();
-									return true;
-								}
-							}
-						}
-						return false;
-					});
-			}
-		}
-
-		DisplayString = BaseFilename + MeshName + TEXT("_Anim");
-
-		int32 AnimationCount = 0;
-		NodeContainer->BreakableIterateNodesOfType<UInterchangeSkeletalAnimationTrackNode>([&AnimationCount](const FString&, UInterchangeSkeletalAnimationTrackNode*) { return (AnimationCount++ > 0); });
-		if (AnimationCount > 1)
-		{
-			DisplayString += TEXT("_") + TrackNode.GetDisplayLabel();
-		}
-		constexpr bool bIsJointFalse = false;
-		UE::Interchange::SanitizeName(DisplayString, bIsJointFalse);
-	}
-
 	//Check Compatibility: (as in does the animation target the skeleton at all)
 	bool DoesSkeletalAnimationTargetSkeleton(const UInterchangeBaseNodeContainer* BaseNodeContainer, const UInterchangeSkeletalAnimationTrackNode& TrackNode, USkeleton* Skeleton)
 	{
@@ -337,6 +272,22 @@ void UInterchangeGenericAnimationPipeline::GetSupportAssetClasses(TArray<UClass*
 	PipelineSupportAssetClasses.Add(ULevelSequence::StaticClass());
 }
 
+void UInterchangeGenericAnimationPipeline::ImplementUseSourceNameForAssetOptionAnimSequence(bool bUseSourceNameForAsset, const FString& AssetName)
+{
+	const bool bShouldChangeAssetName = bUseSourceNameForAsset || !AssetName.IsEmpty();
+
+	// Prefix anim sequences with the asset name (mostly used by FBX)
+	if(bShouldChangeAssetName && !SourceDatas.IsEmpty())
+	{
+		FString DisplayLabelName = AssetName.IsEmpty() ? FPaths::GetBaseFilename(SourceDatas[0]->GetFilename()) : AssetName;
+		BaseNodeContainer->IterateNodesOfType<UInterchangeAnimSequenceFactoryNode>(
+		[&DisplayLabelName](const FString&, UInterchangeAnimSequenceFactoryNode* AnimSequenceFactoryNode)
+		{
+			AnimSequenceFactoryNode->SetDisplayLabel(DisplayLabelName + AnimSequenceFactoryNode->GetDisplayLabel());
+		});
+	}
+}
+
 #endif //WITH_EDITOR
 
 void UInterchangeGenericAnimationPipeline::ExecutePipeline(UInterchangeBaseNodeContainer* InBaseNodeContainer, const TArray<UInterchangeSourceData*>& InSourceDatas, const FString& ContentBasePath)
@@ -842,7 +793,6 @@ void UInterchangeGenericAnimationPipeline::CreateAnimSequenceFactoryNode(UInterc
 	UInterchangeAnimSequenceFactoryNode* AnimSequenceFactoryNode = NewObject<UInterchangeAnimSequenceFactoryNode>(BaseNodeContainer, NAME_None);
 
 	FString DisplayString = TrackNode.GetDisplayLabel();
-	UE::Interchange::Private::RenameAnimSequenceLikeLegacyFbx(DisplayString, TrackNode, SourceDatas, GetMostPipelineOuter(), BaseNodeContainer);
 
 	AnimSequenceFactoryNode->InitializeAnimSequenceNode(AnimSequenceUid, DisplayString, BaseNodeContainer);
 
diff --git a/Engine/Plugins/Interchange/Runtime/Source/Pipelines/Private/InterchangeGenericAssetsPipeline.cpp b/Engine/Plugins/Interchange/Runtime/Source/Pipelines/Private/InterchangeGenericAssetsPipeline.cpp
--- a/Engine/Plugins/Interchange/Runtime/Source/Pipelines/Private/InterchangeGenericAssetsPipeline.cpp
+++ b/Engine/Plugins/Interchange/Runtime/Source/Pipelines/Private/InterchangeGenericAssetsPipeline.cpp
@@ -1670,6 +1670,7 @@ void UInterchangeGenericAssetsPipeline::ImplementUseSourceNameForAssetOption(UIn
 
 	//SkeletalMesh it must always be run even if there is no rename option, skeleton and physics asset will be rename properly
 	MeshPipeline->ImplementUseSourceNameForAssetOptionSkeletalMesh(MeshesImportedNodeCount, bUseSourceNameForAsset, OverrideAssetName);
+	AnimationPipeline->ImplementUseSourceNameForAssetOptionAnimSequence(bUseSourceNameForAsset, OverrideAssetName);
 
 	if (bUseSourceNameForAsset || !OverrideAssetName.IsEmpty())
 	{
diff --git a/Engine/Plugins/Interchange/Runtime/Source/Pipelines/Private/InterchangeGenericSkeletalMeshPipeline.cpp b/Engine/Plugins/Interchange/Runtime/Source/Pipelines/Private/InterchangeGenericSkeletalMeshPipeline.cpp
--- a/Engine/Plugins/Interchange/Runtime/Source/Pipelines/Private/InterchangeGenericSkeletalMeshPipeline.cpp
+++ b/Engine/Plugins/Interchange/Runtime/Source/Pipelines/Private/InterchangeGenericSkeletalMeshPipeline.cpp
@@ -69,31 +69,6 @@ namespace UE::Interchange::SkeletalMeshGenericPipeline
 			}
 		}
 	}
-
-	void RenameSkeletalMeshLikeLegacyFbx(FString& DisplayLabel, const UInterchangeMeshNode* MeshNode, const TArray<const UInterchangeSourceData*>& SourceDatas, UInterchangePipelineBase* OuterMostPipeline)
-	{
-		if(!MeshNode || !SourceDatas.IsValidIndex(0))
-		{
-			return;
-		}
-		bool bRenameLikeLegacyFbx = false;
-		MeshNode->GetBooleanAttribute(TEXT("RenameLikeLegacyFbx"), bRenameLikeLegacyFbx);
-		if (!bRenameLikeLegacyFbx)
-		{
-			return;
-		}
-
-		FString BaseFilename = FPaths::GetBaseFilename(SourceDatas[0]->GetFilename());
-		FString MeshName;
-		if (UInterchangeGenericAssetsPipeline* GenericAssetPipeline = Cast<UInterchangeGenericAssetsPipeline>(OuterMostPipeline))
-		{
-			if (!GenericAssetPipeline->bUseSourceNameForAsset && !DisplayLabel.IsEmpty())
-			{
-				MeshName = TEXT("_") + DisplayLabel;
-			}
-		}
-		DisplayLabel = BaseFilename + MeshName;
-	}
 }
 
 void UInterchangeGenericMeshPipeline::ExecutePreImportPipelineSkeletalMesh()
@@ -248,6 +223,13 @@ void UInterchangeGenericMeshPipeline::ExecutePreImportPipelineSkeletalMesh()
 						}
 					}
 
+					// sort the MeshUids in order to ensure to have the exact same name when we create the different asset
+					// because the asset name is taken by the 1st scene node found, see CreateSkeletalMeshFactoryNode
+					for (TPair<int32, TArray<FString>>& Pair : MeshUidsPerLodIndex)
+					{
+						Pair.Value.Sort();
+					}
+
 					FinishSkeletalMeshSetup(bUseInstanceMesh, bFoundInstances, SkeletonRootUid, SkeletonFactoryNode, MeshUidsPerLodIndex, MeshUids);
 				}
 				return bFoundInstances;
@@ -425,8 +407,6 @@ UInterchangeSkeletalMeshFactoryNode* UInterchangeGenericMeshPipeline::CreateSkel
 		SkeletalMeshUid_MeshNamePart = TEXT("\\") + FirstSceneNode->GetUniqueID();
 	}
 
-	UE::Interchange::SkeletalMeshGenericPipeline::RenameSkeletalMeshLikeLegacyFbx(DisplayLabel, FirstMeshNode, SourceDatas, GetMostPipelineOuter());
-
 	const FString SkeletalMeshUid = UInterchangeFactoryBaseNode::BuildFactoryNodeUid(SkeletalMeshUid_MeshNamePart + SkeletonUid);
 	UInterchangeSkeletalMeshFactoryNode* SkeletalMeshFactoryNode = NewObject<UInterchangeSkeletalMeshFactoryNode>(BaseNodeContainer, NAME_None);
 	if (!ensure(SkeletalMeshFactoryNode))
diff --git a/Engine/Plugins/Interchange/Runtime/Source/Pipelines/Public/InterchangeGenericAnimationPipeline.h b/Engine/Plugins/Interchange/Runtime/Source/Pipelines/Public/InterchangeGenericAnimationPipeline.h
--- a/Engine/Plugins/Interchange/Runtime/Source/Pipelines/Public/InterchangeGenericAnimationPipeline.h
+++ b/Engine/Plugins/Interchange/Runtime/Source/Pipelines/Public/InterchangeGenericAnimationPipeline.h
@@ -134,6 +134,8 @@ public:
 
 #endif //WITH_EDITOR
 
+	UE_API void ImplementUseSourceNameForAssetOptionAnimSequence(bool bUseSourceNameForAsset, const FString& AssetName);
+
 protected:
 
 	UE_API virtual void ExecutePipeline(UInterchangeBaseNodeContainer* InBaseNodeContainer, const TArray<UInterchangeSourceData*>& InSourceDatas, const FString& ContentBasePath) override;
