[Interchange] Removing unused variables in Anim Sequence Factory.

#rnx

diff --git a/Engine/Plugins/Interchange/Runtime/Source/Import/Private/Animation/InterchangeAnimSequenceFactory.cpp b/Engine/Plugins/Interchange/Runtime/Source/Import/Private/Animation/InterchangeAnimSequenceFactory.cpp
--- a/Engine/Plugins/Interchange/Runtime/Source/Import/Private/Animation/InterchangeAnimSequenceFactory.cpp
+++ b/Engine/Plugins/Interchange/Runtime/Source/Import/Private/Animation/InterchangeAnimSequenceFactory.cpp
@@ -552,9 +552,7 @@ namespace UE::Interchange::Private
 				RawTrack.PosKeys.Reserve(BakeKeyCountForAnimationPayload);
 				RawTrack.RotKeys.Reserve(BakeKeyCountForAnimationPayload);
 				RawTrack.ScaleKeys.Reserve(BakeKeyCountForAnimationPayload);
-				TArray<float> TimeKeys;
-				TimeKeys.Reserve(BakeKeyCountForAnimationPayload);
-
+				
 				if (!ensure(AnimationTransformPayload.Transforms.Num() == BakeKeyCountForAnimationPayload))
 				{
 					FString PayloadKey = PayloadKeys[AnimationPayload.Key->GetUniqueID()].UniqueId;
@@ -608,15 +606,12 @@ namespace UE::Interchange::Private
 					RawTrack.ScaleKeys.Add(Scale);
 					RawTrack.PosKeys.Add(Position);
 					RawTrack.RotKeys.Add(Quaternion);
-					//Animation are always translated to zero
-					TimeKeys.Add(CurrentTime - AnimationTransformPayload.RangeStartTime);
 				}
 
 				//Make sure we create the correct amount of keys
 				if (!ensure(RawTrack.ScaleKeys.Num() == BakeKeyCountForAnimationPayload
 					&& RawTrack.PosKeys.Num() == BakeKeyCountForAnimationPayload
-					&& RawTrack.RotKeys.Num() == BakeKeyCountForAnimationPayload
-					&& TimeKeys.Num() == BakeKeyCountForAnimationPayload))
+					&& RawTrack.RotKeys.Num() == BakeKeyCountForAnimationPayload))
 				{
 					FString PayloadKey = PayloadKeys[AnimationPayload.Key->GetUniqueID()].UniqueId;
 					UInterchangeResultWarning_Generic* Message = Factory->AddMessage<UInterchangeResultWarning_Generic>();
@@ -709,11 +704,9 @@ namespace UE::Interchange::Private
 						const FName BoneName = FName(*(SkeletonNode->GetDisplayLabel()));
 						//Add only one transform key at time 0 since this node is not animated
 						FRawAnimSequenceTrack RawTrack;
-						TArray<float> TimeKeys;
 						RawTrack.ScaleKeys.Add(AnimKeyTransform.GetScale3D());
 						RawTrack.PosKeys.Add(AnimKeyTransform.GetLocation());
 						RawTrack.RotKeys.Add(AnimKeyTransform.GetRotation());
-						TimeKeys.Add(0);
 						Controller.AddBoneCurve(BoneName, bShouldTransact);
 						Controller.SetBoneTrackKeys(BoneName, RawTrack.PosKeys, RawTrack.RotKeys, RawTrack.ScaleKeys, bShouldTransact);
 					}
