MaterialX:
- add support for <ramp> and <ramp_gradient> nodes
- add support for <fract> node
- SwitchNode supports now up to 10 inputs
- renamed <atan2> inputs from in1 /in2 to iny / inx
- deprecated UMaterialExpressionMaterialXSwizzle since it has been removed from MaterialX 1.39.3

#virtualized

diff --git a/Engine/Plugins/Interchange/Runtime/Source/Common/Public/MaterialX/InterchangeMaterialXDefinitions.h b/Engine/Plugins/Interchange/Runtime/Source/Common/Public/MaterialX/InterchangeMaterialXDefinitions.h
--- a/Engine/Plugins/Interchange/Runtime/Source/Common/Public/MaterialX/InterchangeMaterialXDefinitions.h
+++ b/Engine/Plugins/Interchange/Runtime/Source/Common/Public/MaterialX/InterchangeMaterialXDefinitions.h
@@ -158,6 +158,7 @@ namespace UE
 			static constexpr const TCHAR* Dotproduct = TEXT("Dotproduct");
 			static constexpr const TCHAR* Exponential = TEXT("Exponential");
 			static constexpr const TCHAR* Floor = TEXT("Floor");
+			static constexpr const TCHAR* Frac = TEXT("Frac");
 			static constexpr const TCHAR* HsvToRgb = TEXT("HsvToRgb");
 			static constexpr const TCHAR* If = TEXT("If");
 			static constexpr const TCHAR* Logarithm = TEXT("Logarithm");
@@ -251,6 +252,11 @@ namespace UE
 			static constexpr const TCHAR* in3 = TEXT("in3");
 			static constexpr const TCHAR* in4 = TEXT("in4");
 			static constexpr const TCHAR* in5 = TEXT("in5");
+			static constexpr const TCHAR* in6 = TEXT("in6");
+			static constexpr const TCHAR* in7 = TEXT("in7");
+			static constexpr const TCHAR* in8 = TEXT("in8");
+			static constexpr const TCHAR* in9 = TEXT("in9");
+			static constexpr const TCHAR* in10 = TEXT("in10");
 			static constexpr const TCHAR* Input = TEXT("Input");
 			static constexpr const TCHAR* InputLow = TEXT("InputLow");
 			static constexpr const TCHAR* InputHigh = TEXT("InputHigh");
@@ -302,6 +308,8 @@ namespace UE
 			constexpr const TCHAR* MxLinearAdobeRGBToRec709 = TEXT("/InterchangeAssets/Functions/MX_LinearAdobeRGBToRec709.MX_LinearAdobeRGBToRec709");
 			constexpr const TCHAR* MxLinearDisplayP3ToRec709 = TEXT("/InterchangeAssets/Functions/MX_LinearDisplayP3ToRec709.MX_LinearDisplayP3ToRec709");
 			constexpr const TCHAR* MxPlace2D = TEXT("/InterchangeAssets/Functions/MX_Place2D.MX_Place2D");
+			constexpr const TCHAR* MxRamp = TEXT("/InterchangeAssets/Functions/MX_Ramp.MX_Ramp");
+			constexpr const TCHAR* MxRampGradient = TEXT("/InterchangeAssets/Functions/MX_RampGradient.MX_RampGradient");
 			constexpr const TCHAR* MxRandomFloat = TEXT("/InterchangeAssets/Functions/MX_RandomFloat.MX_RandomFloat");
 			constexpr const TCHAR* MxRandomColor = TEXT("/InterchangeAssets/Functions/MX_RandomColor.MX_RandomColor");
 			constexpr const TCHAR* MxRoughnesAnisotropy = TEXT("/InterchangeAssets/Functions/MX_Roughness_Anisotropy.MX_Roughness_Anisotropy");
@@ -664,6 +672,7 @@ MATERIALX_NAMESPACE_BEGIN
 		static constexpr const char* DotProduct = "dotproduct";
 		static constexpr const char* Exp = "exp";
 		static constexpr const char* Floor = "floor";
+		static constexpr const char* Fract = "fract";
 		static constexpr const char* Invert = "invert";
 		static constexpr const char* InvertMatrix = "invertmatrix";
 		static constexpr const char* Ln = "ln";
@@ -735,6 +744,8 @@ MATERIALX_NAMESPACE_BEGIN
 		static constexpr const char* Line = "line";
 		static constexpr const char* Noise2D = "noise2d";
 		static constexpr const char* Ramp4 = "ramp4";
+		static constexpr const char* Ramp = "ramp";
+		static constexpr const char* RampGradient = "ramp_gradient";
 		static constexpr const char* RampLR = "ramplr";
 		static constexpr const char* RampTB = "ramptb";
 		static constexpr const char* RandomFloat = "randomfloat";
diff --git a/Engine/Plugins/Interchange/Runtime/Source/Import/Private/MaterialX/MaterialExpressions/MaterialExpressionSwizzle.cpp b/Engine/Plugins/Interchange/Runtime/Source/Import/Private/MaterialX/MaterialExpressions/MaterialExpressionSwizzle.cpp
--- a/Engine/Plugins/Interchange/Runtime/Source/Import/Private/MaterialX/MaterialExpressions/MaterialExpressionSwizzle.cpp
+++ b/Engine/Plugins/Interchange/Runtime/Source/Import/Private/MaterialX/MaterialExpressions/MaterialExpressionSwizzle.cpp
@@ -139,9 +139,8 @@ void UMaterialExpressionMaterialXSwizzle::GetCaption(TArray<FString>& OutCaption
 
 void UMaterialExpressionMaterialXSwizzle::GetExpressionToolTip(TArray<FString>& OutToolTip)
 {
-	ConvertToMultilineToolTip(TEXT("Perform an arbitrary permutation of the channels of the input stream, returning a new "
-								   "stream of the specified type. Individual channels may be replicated or omitted, and the output "
-								   "stream may have a different number of channels than the input."), 40, OutToolTip);
+	ConvertToMultilineToolTip(TEXT("UMaterialExpressionMaterialXSwizzle is now deprecated since <swizzle> node has been removed from MaterialX 1.39.3")
+							  TEXT("Please use UMaterialExpressionComponent instead."), 40, OutToolTip);
 }
 
 #endif
diff --git a/Engine/Plugins/Interchange/Runtime/Source/Import/Private/MaterialX/MaterialXUtils/MaterialXManager.cpp b/Engine/Plugins/Interchange/Runtime/Source/Import/Private/MaterialX/MaterialXUtils/MaterialXManager.cpp
--- a/Engine/Plugins/Interchange/Runtime/Source/Import/Private/MaterialX/MaterialXUtils/MaterialXManager.cpp
+++ b/Engine/Plugins/Interchange/Runtime/Source/Import/Private/MaterialX/MaterialXUtils/MaterialXManager.cpp
@@ -72,8 +72,8 @@ FMaterialXManager::FMaterialXManager()
 	   {{TEXT(""),                            TEXT("valuebr")},     ExpressionInput(UE::Expressions::Inputs::D)},
 	   {{TEXT(""),                            TEXT("value1")},      ExpressionInput(UE::Expressions::Inputs::A)},
 	   {{TEXT(""),                            TEXT("value2")},      ExpressionInput(UE::Expressions::Inputs::B)},
-	   {{mx::Category::Atan2,                 TEXT("in1")},         ExpressionInput(UE::Expressions::Inputs::Y)},
-	   {{mx::Category::Atan2,                 TEXT("in2")},         ExpressionInput(UE::Expressions::Inputs::X)},
+	   {{mx::Category::Atan2,                 TEXT("iny")},         ExpressionInput(UE::Expressions::Inputs::Y)},
+	   {{mx::Category::Atan2,                 TEXT("inx")},         ExpressionInput(UE::Expressions::Inputs::X)},
 	   {{mx::Category::HeightToNormal,        TEXT("scale")},       ExpressionInput(*UE::Interchange::Materials::Standard::Nodes::NormalFromHeightMap::Inputs::Intensity.ToString())},
 	   {{mx::Category::IfGreater,             TEXT("in1")},         ExpressionInput(UE::Expressions::Inputs::AGreaterThanB)},
 	   {{mx::Category::IfGreater,             TEXT("in2")},         ExpressionInput(UE::Expressions::Inputs::ALessThanB)},     //another input is added for the case 'equal', see ConnectIfGreater
@@ -114,6 +114,11 @@ FMaterialXManager::FMaterialXManager()
 	   {{mx::Category::Switch,                TEXT("in3")},         ExpressionInput(UE::Expressions::Inputs::in3)},
 	   {{mx::Category::Switch,                TEXT("in4")},         ExpressionInput(UE::Expressions::Inputs::in4)},
 	   {{mx::Category::Switch,                TEXT("in5")},         ExpressionInput(UE::Expressions::Inputs::in5)},
+	   {{mx::Category::Switch,                TEXT("in6")},         ExpressionInput(UE::Expressions::Inputs::in6)},
+	   {{mx::Category::Switch,                TEXT("in7")},         ExpressionInput(UE::Expressions::Inputs::in7)},
+	   {{mx::Category::Switch,                TEXT("in8")},         ExpressionInput(UE::Expressions::Inputs::in8)},
+	   {{mx::Category::Switch,                TEXT("in9")},         ExpressionInput(UE::Expressions::Inputs::in9)},
+	   {{mx::Category::Switch,                TEXT("in10")},        ExpressionInput(UE::Expressions::Inputs::in10)},
 	   {{mx::Category::Switch,                TEXT("which")},       ExpressionInput(UE::Expressions::Inputs::SwitchValue)},
 	}
 	, MatchingMaterialExpressions {
@@ -132,6 +137,7 @@ FMaterialXManager::FMaterialXManager()
 		{mx::Category::DotProduct,   UE::Expressions::Names::Dotproduct},
 		{mx::Category::Exp,          UE::Expressions::Names::Exponential},
 		{mx::Category::Floor,        UE::Expressions::Names::Floor},
+		{mx::Category::Fract,        UE::Expressions::Names::Frac},
 		{mx::Category::Invert,       UE::Expressions::Names::Subtract},
 		{mx::Category::Ln,           UE::Expressions::Names::Logarithm},
 		{mx::Category::Magnitude,    UE::Expressions::Names::Length},
@@ -254,7 +260,8 @@ FMaterialXManager::FMaterialXManager()
 		{mx::Category::Checkerboard,			FMaterialXMaterialFunction{TInPlaceType<FString>{}, UE::MaterialFunctions::Path::MxCheckerboard}},
 		{mx::Category::Circle,					FMaterialXMaterialFunction{TInPlaceType<FString>{}, UE::MaterialFunctions::Path::MxCircle}},
 		{mx::Category::Line,					FMaterialXMaterialFunction{TInPlaceType<FString>{}, UE::MaterialFunctions::Path::MxLine}},
-		{mx::Category::RandomFloat,	            FMaterialXMaterialFunction{TInPlaceType<FString>{}, UE::MaterialFunctions::Path::MxRandomFloat}},
+		{mx::Category::Ramp,					FMaterialXMaterialFunction{TInPlaceType<FString>{}, UE::MaterialFunctions::Path::MxRamp}},
+		{mx::Category::RampGradient,	        FMaterialXMaterialFunction{TInPlaceType<FString>{}, UE::MaterialFunctions::Path::MxRampGradient}},
 		{mx::Category::RandomColor,	            FMaterialXMaterialFunction{TInPlaceType<FString>{}, UE::MaterialFunctions::Path::MxRandomColor}},
 		{mx::Category::TiledCircles,			FMaterialXMaterialFunction{TInPlaceType<FString>{}, UE::MaterialFunctions::Path::MxTiledCircles}},
 		{mx::Category::UnifiedNoise3D,			FMaterialXMaterialFunction{TInPlaceType<FString>{}, UE::MaterialFunctions::Path::MxUnifiedNoise3D}},
@@ -271,6 +278,7 @@ FMaterialXManager::FMaterialXManager()
 		mx::Category::Circle,
 		mx::Category::ColorCorrect,
 		mx::Category::Contrast,
+		mx::Category::Convert,
 		mx::Category::Distance,
 		mx::Category::Extract,
 		mx::Category::Fractal3D,
@@ -284,6 +292,8 @@ FMaterialXManager::FMaterialXManager()
 		mx::Category::OpenPBRSurface,
 		mx::Category::Overlay,
 		mx::Category::Place2D,
+		mx::Category::Ramp,
+		mx::Category::RampGradient,
 		mx::Category::Ramp4,
 		mx::Category::RandomFloat,
 		mx::Category::RandomColor,
@@ -296,6 +306,7 @@ FMaterialXManager::FMaterialXManager()
 		mx::Category::SrgbDisplayP3ToLinRec709,
 		mx::Category::SrgbTextureToLinRec709,
 		mx::Category::StandardSurface,
+		mx::Category::Switch,
 		mx::Category::TiledCircles,
 		mx::Category::Time, // since 1.39.3, they added a nodegraph in stdlib/genglsl/stdlib_genglsl_impl.mtlx instead of stdlib_ng.mtlx
 		mx::Category::UnifiedNoise3D,
@@ -734,6 +745,8 @@ namespace UE::Interchange::MaterialX
 																UE::MaterialFunctions::Path::MxTiledCircles,
 																UE::MaterialFunctions::Path::MxUnifiedNoise3D,
 																UE::MaterialFunctions::Path::MxLine,
+																UE::MaterialFunctions::Path::MxRamp,
+																UE::MaterialFunctions::Path::MxRampGradient,
 															  });
 
 		return bPackagesLoaded;
diff --git a/Engine/Plugins/Interchange/Runtime/Source/Import/Public/MaterialX/MaterialExpressions/MaterialExpressionSwizzle.h b/Engine/Plugins/Interchange/Runtime/Source/Import/Public/MaterialX/MaterialExpressions/MaterialExpressionSwizzle.h
--- a/Engine/Plugins/Interchange/Runtime/Source/Import/Public/MaterialX/MaterialExpressions/MaterialExpressionSwizzle.h
+++ b/Engine/Plugins/Interchange/Runtime/Source/Import/Public/MaterialX/MaterialExpressions/MaterialExpressionSwizzle.h
@@ -6,8 +6,9 @@
 
 #include "MaterialExpressionSwizzle.generated.h"
 
-UCLASS(collapsecategories, hidecategories = Object, MinimalAPI, meta = (Private))
-class UMaterialExpressionMaterialXSwizzle : public UMaterialExpression
+UCLASS()
+class UE_DEPRECATED(5.7, "UMaterialExpressionMaterialXSwizzle is now deprecated since <swizzle> node has been removed from MaterialX 1.39.3. Please use UMaterialExpressionComponent instead.")
+UMaterialExpressionMaterialXSwizzle : public UMaterialExpression
 {
 	GENERATED_UCLASS_BODY()
 
