[Interchange] Level Import fix for StaticMeshes with MorphTargets.

#rnx

diff --git a/Engine/Plugins/Interchange/Runtime/Source/Pipelines/Private/InterchangeGenericScenesPipeline.cpp b/Engine/Plugins/Interchange/Runtime/Source/Pipelines/Private/InterchangeGenericScenesPipeline.cpp
--- a/Engine/Plugins/Interchange/Runtime/Source/Pipelines/Private/InterchangeGenericScenesPipeline.cpp
+++ b/Engine/Plugins/Interchange/Runtime/Source/Pipelines/Private/InterchangeGenericScenesPipeline.cpp
@@ -544,23 +544,19 @@ void UInterchangeGenericLevelPipeline::ExecuteSceneNodePreImport(const FTransfor
 			{
 				return;
 			}
-		}
-	}
-
-	// Don't generate actor factory nodes for static mesh scene nodes that haven't produced any static mesh factory nodes,
-	// as we're not going to have any static mesh to assign to the component anyway!
-	//
-	// We currently use this to prevent the generation of static mesh components for collision mesh nodes, when importing into level.
-	// In that case we won't produce a factory node for the collider mesh node, which we can detect here and early out
-	if (const UInterchangeMeshNode* MeshNode = Cast<UInterchangeMeshNode>(TranslatedAssetNode))
-	{
-		if (!MeshNode->IsSkinnedMesh())
-		{
-			const FString MeshFactoryNodeUid = UInterchangeFactoryBaseNode::BuildFactoryNodeUid(MeshNode->GetUniqueID());
-			const UInterchangeStaticMeshFactoryNode* MeshFactoryNode = Cast<UInterchangeStaticMeshFactoryNode>(BaseNodeContainer->GetFactoryNode(MeshFactoryNodeUid));
-			if (!MeshFactoryNode)
+			else
 			{
-				return;
+				// Don't generate actor factory nodes for static mesh scene nodes that haven't produced any static mesh factory nodes, unless it is a rigid mesh,
+				// as we're not going to have any static mesh to assign to the component anyway!
+				//
+				// We currently use this to prevent the generation of static mesh components for collision mesh nodes, when importing into level.
+				// In that case we won't produce a factory node for the collider mesh node, which we can detect here and early out
+				const FString MeshFactoryNodeUid = UInterchangeFactoryBaseNode::BuildFactoryNodeUid(MeshNode->GetUniqueID());
+				const UInterchangeStaticMeshFactoryNode* MeshFactoryNode = Cast<UInterchangeStaticMeshFactoryNode>(BaseNodeContainer->GetFactoryNode(MeshFactoryNodeUid));
+				if (!MeshFactoryNode)
+				{
+					return;
+				}
 			}
 		}
 	}
