[Interchange-glTF] Fixing bind pose calculation and node types for Joint Hierachies with Gaps.

#rnx

diff --git a/Engine/Plugins/Interchange/Runtime/Source/Import/Private/Animation/InterchangeAnimSequenceFactory.cpp b/Engine/Plugins/Interchange/Runtime/Source/Import/Private/Animation/InterchangeAnimSequenceFactory.cpp
--- a/Engine/Plugins/Interchange/Runtime/Source/Import/Private/Animation/InterchangeAnimSequenceFactory.cpp
+++ b/Engine/Plugins/Interchange/Runtime/Source/Import/Private/Animation/InterchangeAnimSequenceFactory.cpp
@@ -672,34 +672,6 @@ namespace UE::Interchange::Private
 							}
 						}
 
-						// On the unusal edge use case where Non-Joint Node(s) are inserted into a Skeleton hierarchy, bake said transforms onto the Joint.
-						if (!SkeletonNode->GetUniqueID().Equals(SkeletonRootUid))
-						{
-							TArray<FTransform> NonJointTransforms; //Child to Parent
-							const UInterchangeSceneNode* ParentNode = Cast<UInterchangeSceneNode>(NodeContainer->GetNode(SkeletonNode->GetParentUid()));
-							while (ParentNode && !ParentNode->IsSpecializedTypeContains(UE::Interchange::FSceneNodeStaticData::GetJointSpecializeTypeString()))
-							{
-								FTransform NonJointLocalTransform;
-								if (ParentNode->GetCustomLocalTransform(NonJointLocalTransform))
-								{
-									NonJointTransforms.Add(NonJointLocalTransform);
-								}
-
-								ParentNode = Cast<UInterchangeSceneNode>(NodeContainer->GetNode(ParentNode->GetParentUid()));
-							}
-							
-							if (NonJointTransforms.Num() > 0)
-							{
-								FTransform NonJointTransformToApply;
-								for (int32 NonJointCounter = NonJointTransforms.Num() - 1; NonJointCounter >= 0; NonJointCounter--)
-								{
-									NonJointTransformToApply *= NonJointTransforms[NonJointCounter];
-								}
-
-								LocalTransform = NonJointTransformToApply * LocalTransform;
-							}
-						}
-
 						FTransform3f AnimKeyTransform = static_cast<FTransform3f>(LocalTransform);
 						const FName BoneName = FName(*(SkeletonNode->GetDisplayLabel()));
 						//Add only one transform key at time 0 since this node is not animated
diff --git a/Engine/Plugins/Interchange/Runtime/Source/Import/Private/Gltf/InterchangeGltfAnimation.cpp b/Engine/Plugins/Interchange/Runtime/Source/Import/Private/Gltf/InterchangeGltfAnimation.cpp
--- a/Engine/Plugins/Interchange/Runtime/Source/Import/Private/Gltf/InterchangeGltfAnimation.cpp
+++ b/Engine/Plugins/Interchange/Runtime/Source/Import/Private/Gltf/InterchangeGltfAnimation.cpp
@@ -48,7 +48,7 @@ namespace UE::Interchange::Gltf::Private
 		}
 		else
 		{
-			for (size_t Index = 0; Index < Seconds.Num() - 1; Index++)
+			for (int32 Index = 0; Index < Seconds.Num() - 1; Index++)
 			{
 				if ((Seconds[Index] <= SampleTime) && (SampleTime < Seconds[Index + 1]))
 				{
@@ -538,19 +538,6 @@ namespace UE::Interchange::Gltf::Private
 			FVector3d(ScaleData));
 	}
 
-	void GetNonJointTransforms(const GLTF::FAsset& GltfAsset, int32 CurrentNodeIndex, FTransform& Transform)
-	{
-		if (GltfAsset.Nodes.IsValidIndex(CurrentNodeIndex))
-		{
-			if (GltfAsset.Nodes[CurrentNodeIndex].Type != GLTF::FNode::EType::Joint)
-			{
-				Transform = GltfAsset.Nodes[CurrentNodeIndex].Transform * Transform;
-
-				GetNonJointTransforms(GltfAsset, GltfAsset.Nodes[CurrentNodeIndex].ParentIndex, Transform);
-			}
-		}
-	}
-
 	bool GetBakedAnimationTransformPayloadData(const FString& PayLoadKey, const GLTF::FAsset& GltfAsset, FAnimationPayloadData& PayloadData)
 	{
 		const double BakeInterval = 1.0 / PayloadData.BakeFrequency;
@@ -576,9 +563,6 @@ namespace UE::Interchange::Gltf::Private
 		//channels are targetting the same node:
 		bool BaseTransformsAcquired = false;
 
-		bool bNonJointTransformNeedsApplication = false;
-		FTransform NonJointTransformToApply;
-
 		for (const int32& ChannelIndex : ChannelIndices)
 		{
 			const GLTF::FAnimation::FChannel& Channel = GltfAnimation.Channels[ChannelIndex];
@@ -596,13 +580,6 @@ namespace UE::Interchange::Gltf::Private
 				TranslationData.Init(JointBaseTranslation, BakeKeyCount);
 				RotationData.Init(JointBaseRotation, BakeKeyCount);
 				ScaleData.Init(JointBaseScale, BakeKeyCount);
-
-				if (Channel.Target.Node.Type == GLTF::FNode::EType::Joint &&
-					Channel.Target.Node.ParentIndex != Channel.Target.Node.ParentJointIndex)
-				{
-					GetNonJointTransforms(GltfAsset, Channel.Target.Node.ParentIndex, NonJointTransformToApply);
-					bNonJointTransformNeedsApplication = true;
-				}
 			}
 
 			TArray<float> Seconds;
@@ -679,26 +656,12 @@ namespace UE::Interchange::Gltf::Private
 			}
 		}
 
-		if (bNonJointTransformNeedsApplication)
-		{
-			FVector3d NonJointTranslation = NonJointTransformToApply.GetTranslation();
-			FQuat4d NonJointRotation = NonJointTransformToApply.GetRotation();
-			FVector3d NonJointScale = NonJointTransformToApply.GetScale3D();
-
-			for (size_t Index = 0; Index < BakeKeyCount; Index++)
-			{
-				PayloadData.Transforms.Add(FTransform(NonJointRotation * FQuat4d(RotationData[Index]), NonJointTranslation + FVector3d(TranslationData[Index]), NonJointScale * FVector3d(ScaleData[Index])));
-			}
-		}
-		else
+		for (size_t Index = 0; Index < BakeKeyCount; Index++)
 		{
-			for (size_t Index = 0; Index < BakeKeyCount; Index++)
-			{
-				PayloadData.Transforms.Add(FTransform(
-					FQuat4d(RotationData[Index]),
-					FVector3d(TranslationData[Index]),
-					FVector3d(ScaleData[Index])));
-			}
+			PayloadData.Transforms.Add(FTransform(
+				FQuat4d(RotationData[Index]),
+				FVector3d(TranslationData[Index]),
+				FVector3d(ScaleData[Index])));
 		}
 
 		return true;
diff --git a/Engine/Plugins/Interchange/Runtime/Source/Import/Private/Mesh/InterchangeSkeletonHelper.cpp b/Engine/Plugins/Interchange/Runtime/Source/Import/Private/Mesh/InterchangeSkeletonHelper.cpp
--- a/Engine/Plugins/Interchange/Runtime/Source/Import/Private/Mesh/InterchangeSkeletonHelper.cpp
+++ b/Engine/Plugins/Interchange/Runtime/Source/Import/Private/Mesh/InterchangeSkeletonHelper.cpp
@@ -247,8 +247,7 @@ namespace UE::Interchange::Private
 		, const bool bUseTimeZeroAsBindPose
 		, bool& bOutDiffPose
 		, TArray<FString>& OutBoneNotBindNames
-		, bool bImportSockets
-		, bool bAllowNonJoints)
+		, bool bImportSockets)
 	{
 		const UInterchangeSceneNode* JointNode = Cast<UInterchangeSceneNode>(NodeContainer->GetNode(JointNodeId));
 		if (!JointNode)
@@ -256,32 +255,6 @@ namespace UE::Interchange::Private
 			UE_LOG(LogInterchangeImport, Warning, TEXT("Invalid Skeleton Joint"));
 			return;
 		}
-		
-		//Do not add non-joint Nodes as Bones.
-		//	Exception is when a Skeleton was created in order to create AnimSeqence from LevelSequence.
-		//	Which can be tracked against by checking the first Node that this function is executed on: if it does not have Joint Specialization then we add all non-joint nodes.
-
-		if (ParentIndex == INDEX_NONE)
-		{
-			if (JointNode->IsSpecializedTypeContains(UE::Interchange::FSceneNodeStaticData::GetJointSpecializeTypeString()))
-			{
-				bAllowNonJoints = false;
-			}
-			else
-			{
-				bAllowNonJoints = true;
-			}
-		}
-
-		if (!bAllowNonJoints && !JointNode->IsSpecializedTypeContains(UE::Interchange::FSceneNodeStaticData::GetJointSpecializeTypeString()))
-		{
-			const TArray<FString> ChildrenIds = NodeContainer->GetNodeChildrenUids(JointNodeId);
-			for (int32 ChildIndex = 0; ChildIndex < ChildrenIds.Num(); ++ChildIndex)
-			{
-				RecursiveAddBones(NodeContainer, ChildrenIds[ChildIndex], JointInfos, ParentIndex, RefBonesBinary, bUseTimeZeroAsBindPose, bOutDiffPose, OutBoneNotBindNames, bImportSockets, bAllowNonJoints);
-			}
-			return;
-		}
 
 		int32 JointInfoIndex = JointInfos.Num();
 		FJointInfo& Info = JointInfos.AddZeroed_GetRef();
@@ -360,7 +333,7 @@ namespace UE::Interchange::Private
 					//We dont want a socket to be a skeleton joint
 					continue;
 				}
-				RecursiveAddBones(NodeContainer, ChildrenIds[ChildIndex], JointInfos, JointInfoIndex, RefBonesBinary, bUseTimeZeroAsBindPose, bOutDiffPose, OutBoneNotBindNames, bImportSockets, bAllowNonJoints);
+				RecursiveAddBones(NodeContainer, ChildrenIds[ChildIndex], JointInfos, JointInfoIndex, RefBonesBinary, bUseTimeZeroAsBindPose, bOutDiffPose, OutBoneNotBindNames, bImportSockets);
 			}
 		}
 	}
diff --git a/Engine/Plugins/Interchange/Runtime/Source/Import/Public/InterchangeSkeletonHelper.h b/Engine/Plugins/Interchange/Runtime/Source/Import/Public/InterchangeSkeletonHelper.h
--- a/Engine/Plugins/Interchange/Runtime/Source/Import/Public/InterchangeSkeletonHelper.h
+++ b/Engine/Plugins/Interchange/Runtime/Source/Import/Public/InterchangeSkeletonHelper.h
@@ -59,8 +59,7 @@ namespace UE::Interchange::Private
 			, const bool bUseTimeZeroAsBindPose
 			, bool& bOutDiffPose
 			, TArray<FString>& OutBoneNotBindNames
-			, bool bImportSockets
-			, bool bAllowNonJoints = true /*It will be set on based on Root node's details */);
+			, bool bImportSockets);
 
 	private:
 		static FName SkeletalLodGetBoneName(const TArray<FMeshBoneInfo>& SkeletalLodRawInfos, int32 BoneIndex);
diff --git a/Engine/Plugins/Interchange/Runtime/Source/Parsers/GLTFCore/Private/GLTF/GLTFReader.cpp b/Engine/Plugins/Interchange/Runtime/Source/Parsers/GLTFCore/Private/GLTF/GLTFReader.cpp
--- a/Engine/Plugins/Interchange/Runtime/Source/Parsers/GLTFCore/Private/GLTF/GLTFReader.cpp
+++ b/Engine/Plugins/Interchange/Runtime/Source/Parsers/GLTFCore/Private/GLTF/GLTFReader.cpp
@@ -1180,6 +1180,8 @@ namespace GLTF
 				}
 			}
 
+			FixNodeTypesForIndirectJoints();
+
 			BuildRootJoints();
 
 			GenerateInverseBindPosesPerSkinIndices();
@@ -1360,6 +1362,17 @@ namespace GLTF
 
 	void FFileReader::GenerateLocalBindPosesPerSkinIndices() const
 	{
+		auto GetTransform = [](const FMatrix& Matrix) -> FTransform
+			{
+				FTransform Transform;
+				Transform.SetFromMatrix(Matrix);
+				Transform.SetRotation(GLTF::ConvertQuat(Transform.GetRotation()));
+				Transform.SetTranslation(GLTF::ConvertVec3(Transform.GetTranslation()));
+				Transform.SetScale3D(GLTF::ConvertVec3(Transform.GetScale3D()));
+
+				return Transform;
+			};
+
 		for (size_t SkinIndex = 0; SkinIndex < Asset->Skins.Num(); SkinIndex++)
 		{
 			const FSkinInfo& Skin = Asset->Skins[SkinIndex];
@@ -1370,12 +1383,13 @@ namespace GLTF
 				{
 					FNode& CurrentNode = Asset->Nodes[Skin.Joints[JointCounter]];
 
-					if (CurrentNode.ParentJointIndex != INDEX_NONE &&
-						Asset->Nodes.IsValidIndex(CurrentNode.ParentJointIndex) &&
-						(Asset->Nodes[CurrentNode.ParentJointIndex].SkinIndexToGlobalInverseBindMatrix.Contains(SkinIndex) || Asset->Nodes[CurrentNode.ParentJointIndex].SkinIndexToGlobalInverseBindMatrix.Num() > 0)
+					
+					if (CurrentNode.DirectParentJointIndex != INDEX_NONE &&
+						Asset->Nodes.IsValidIndex(CurrentNode.DirectParentJointIndex) &&
+						(Asset->Nodes[CurrentNode.DirectParentJointIndex].SkinIndexToGlobalInverseBindMatrix.Contains(SkinIndex) || Asset->Nodes[CurrentNode.DirectParentJointIndex].SkinIndexToGlobalInverseBindMatrix.Num() > 0)
 						)
 					{
-						FNode& ParentNode = Asset->Nodes[CurrentNode.ParentJointIndex];
+						FNode& ParentNode = Asset->Nodes[CurrentNode.DirectParentJointIndex];
 
 						//LocalBindPose; //bind pose would be CurrentNode.GlobalInverseBindTransform.Inverse() * ParentNode.GlobalInverseBindTransform
 						FMatrix ParentGlobalInverseBindMatrix;
@@ -1390,14 +1404,40 @@ namespace GLTF
 							ParentGlobalInverseBindMatrix = ParentNode.SkinIndexToGlobalInverseBindMatrix.begin().Value();
 						}
 
-						FMatrix LocalBindMatrix = CurrentNode.SkinIndexToGlobalInverseBindMatrix[SkinIndex].Inverse() * ParentGlobalInverseBindMatrix;
+						FTransform LocalBindTransform = FTransform::Identity;
+
+						FTransform ParentGlobalInverseBindTransform = GetTransform(ParentGlobalInverseBindMatrix);
+						FTransform GlobalBindTransform = GetTransform(CurrentNode.SkinIndexToGlobalInverseBindMatrix[SkinIndex].Inverse());
+
+						if (CurrentNode.DirectParentJointIndex != CurrentNode.ParentJointIndex
+							&& Asset->Nodes.IsValidIndex(CurrentNode.ParentJointIndex))
+						{
+							int32 ParentJointIndex = CurrentNode.ParentJointIndex;
+
+							TArray<FTransform> ParentTransformsToApply;
+
+							while (Asset->Nodes.IsValidIndex(ParentJointIndex)
+								&& Asset->Nodes[ParentJointIndex].bHasLocalBindPose) //this cycle meant to go through the "JointGaps". "JointGaps" have their LocalBindPose already set here (set @FixNodeTypesForIndirectJoints)
+							{
+								ParentTransformsToApply.Add(Asset->Nodes[ParentJointIndex].Transform);
+
+								ParentJointIndex = Asset->Nodes[ParentJointIndex].ParentJointIndex;
+							}
+
+							for (int32 MatrixToApplyIndex = ParentTransformsToApply.Num()-1; MatrixToApplyIndex >= 0; MatrixToApplyIndex--)
+							{
+								ParentGlobalInverseBindTransform = ParentGlobalInverseBindTransform * ParentTransformsToApply[MatrixToApplyIndex].Inverse();
+							}
+						}
+
+						LocalBindTransform = GlobalBindTransform * ParentGlobalInverseBindTransform;
 
-						CurrentNode.SkinIndexToLocalBindMatrix.Add(SkinIndex, LocalBindMatrix);
+						CurrentNode.SkinIndexToLocalBindTransform.Add(SkinIndex, LocalBindTransform);
 					}
 					else
 					{
-							FMatrix LocalBindPose = CurrentNode.SkinIndexToGlobalInverseBindMatrix[SkinIndex].Inverse();
-							CurrentNode.SkinIndexToLocalBindMatrix.Add(SkinIndex, LocalBindPose);
+						FTransform LocalBindPose = GetTransform(CurrentNode.SkinIndexToGlobalInverseBindMatrix[SkinIndex].Inverse());
+						CurrentNode.SkinIndexToLocalBindTransform.Add(SkinIndex, LocalBindPose);
 					}
 				}
 			}
@@ -1411,8 +1451,8 @@ namespace GLTF
 
 		for (const FNode& CurrentNode : Asset->Nodes)
 		{
-			TMap<int, FMatrix>::TRangedForConstIterator Iter(CurrentNode.SkinIndexToLocalBindMatrix.begin());
-			FMatrix ToCompareAgainst = Iter ? Iter.Value() : FMatrix();
+			TMap<int, FTransform>::TRangedForConstIterator Iter(CurrentNode.SkinIndexToLocalBindTransform.begin());
+			FTransform ToCompareAgainst = Iter ? Iter.Value() : FTransform();
 			++Iter;
 			for (; Iter; ++Iter)
 			{
@@ -1454,24 +1494,39 @@ namespace GLTF
 					FNode& CurrentNode = Asset->Nodes[Skin.Joints[JointCounter]];
 
 					if (!CurrentNode.bHasLocalBindPose
-						&& CurrentNode.SkinIndexToLocalBindMatrix.Contains(SkinIndex))
+						&& CurrentNode.SkinIndexToLocalBindTransform.Contains(SkinIndex))
 					{
 						CurrentNode.bHasLocalBindPose = true;
-						
-						FTransform LocalBindPoseTransform;
-						LocalBindPoseTransform.SetFromMatrix(CurrentNode.SkinIndexToLocalBindMatrix[SkinIndex]);
-						LocalBindPoseTransform.SetRotation(GLTF::ConvertQuat(LocalBindPoseTransform.GetRotation()));
-						LocalBindPoseTransform.SetTranslation(GLTF::ConvertVec3(LocalBindPoseTransform.GetTranslation()));
-						LocalBindPoseTransform.SetScale3D(GLTF::ConvertVec3(LocalBindPoseTransform.GetScale3D()));
 
-						CurrentNode.LocalBindPose = LocalBindPoseTransform;
+						CurrentNode.LocalBindPose = CurrentNode.SkinIndexToLocalBindTransform[SkinIndex];
 					}
 				}
 			}
 		}
 	}
 
-	void FFileReader::BuildParentIndices(int32 ParentNodeIndex, int32 ParentJointIndex, int32 CurrentNodeIndex) const
+	void FFileReader::BuildParentJointIndices() const
+	{
+		for (size_t SkinIndex = 0; SkinIndex < Asset->Skins.Num(); SkinIndex++)
+		{
+			GLTF::FSkinInfo& Skin = Asset->Skins[SkinIndex];
+
+			for (int32 BoneIndex : Skin.Joints)
+			{
+				int32 CurrentIndex = BoneIndex;
+				while (Asset->Nodes.IsValidIndex(CurrentIndex)
+					&& Asset->Nodes.IsValidIndex(Asset->Nodes[CurrentIndex].ParentIndex)
+					&& Asset->Nodes[Asset->Nodes[CurrentIndex].ParentIndex].Type == GLTF::FNode::EType::Joint)
+				{
+					Asset->Nodes[CurrentIndex].ParentJointIndex = Asset->Nodes[CurrentIndex].ParentIndex;
+
+					CurrentIndex = Asset->Nodes[CurrentIndex].ParentJointIndex;
+				}
+			}
+		}
+	}
+
+	void FFileReader::BuildParentIndices(int32 ParentNodeIndex, int32 DirectParentJointIndex, int32 CurrentNodeIndex) const
 	{
 		if (!Asset->Nodes.IsValidIndex(CurrentNodeIndex))
 		{
@@ -1493,11 +1548,11 @@ namespace GLTF
 		//		However this way we are taking into account that different skins can have different Skeleton compositions.
 		//		We have to update the FNodes' to contain information as to which Skin are they Joint in.
 		//		Then root calculations should take into account the SkinIdx information and create Roots based on said information.
-		Node.ParentJointIndex = ParentJointIndex;
+		Node.DirectParentJointIndex = DirectParentJointIndex;
 
 		for (const int32 ChildNodeIndex : Node.Children)
 		{
-			BuildParentIndices(CurrentNodeIndex, Node.Type == GLTF::FNode::EType::Joint ? CurrentNodeIndex : ParentJointIndex, ChildNodeIndex);
+			BuildParentIndices(CurrentNodeIndex, Node.Type == GLTF::FNode::EType::Joint ? CurrentNodeIndex : DirectParentJointIndex, ChildNodeIndex);
 		}
 	}
 
@@ -1514,6 +1569,63 @@ namespace GLTF
 		}
 		return CurrentIndex;
 	}
+
+	void FFileReader::FixNodeTypesForIndirectJoints() const
+	{
+		for (size_t SkinIndex = 0; SkinIndex < Asset->Skins.Num(); SkinIndex++)
+		{
+			GLTF::FSkinInfo& Skin = Asset->Skins[SkinIndex];
+
+			TSet<int32> BoneIndices;
+			BoneIndices.Append(Skin.Joints);
+
+			for (int32 BoneIndex : Skin.Joints)
+			{
+				if (!Asset->Nodes.IsValidIndex(BoneIndex))
+				{
+					continue;
+				}
+
+				int32 ParentIndex = Asset->Nodes[BoneIndex].ParentIndex;
+
+				TArray<int32> NodeToUpdateToJointIndices;
+				bool bSkeletonGapFound = false;
+
+				while (Asset->Nodes.IsValidIndex(ParentIndex))
+				{
+					if (BoneIndices.Contains(ParentIndex))
+					{
+						if (NodeToUpdateToJointIndices.Num() > 0)
+						{
+							bSkeletonGapFound = true;
+						}
+						break;
+					}
+					NodeToUpdateToJointIndices.Add(ParentIndex);
+					ParentIndex = Asset->Nodes[ParentIndex].ParentIndex;
+				}
+
+				if (bSkeletonGapFound)
+				{
+					for (int32 NodeToUpdate : NodeToUpdateToJointIndices)
+					{
+						GLTF::FNode& Node = Asset->Nodes[NodeToUpdate];
+
+						ensure(Node.Type == FNode::EType::None
+							|| Node.Type == FNode::EType::Transform
+							|| Node.Type == FNode::EType::Joint);
+						Node.Type = FNode::EType::Joint;
+
+						Node.bHasLocalBindPose = true;
+						Node.LocalBindPose = Node.Transform;
+					}
+				}
+			}
+		}
+
+		BuildParentJointIndices();
+	}
+
 	void FFileReader::BuildRootJoints() const
 	{
 		//Fix Multi root problems (at least for the semi-trivial "same parent non-joint" scenario)
@@ -1639,7 +1751,7 @@ namespace GLTF
 						OriginalNode.Children.Reset();
 						OriginalNode.Children.Add(Node.Index);
 						Node.ParentIndex = OriginalNode.Index;
-						Node.ParentJointIndex = OriginalNode.ParentJointIndex;
+						Node.DirectParentJointIndex = OriginalNode.DirectParentJointIndex;
 
 						Node.Transform = OriginalNode.Transform;
 						Node.bHasLocalBindPose = true;
@@ -1659,6 +1771,7 @@ namespace GLTF
 						{
 							GLTF::FNode& JointNode = Asset->Nodes[Group.Value.Indices[RootJointIndex]];
 							JointNode.ParentIndex = Node.Index;
+							JointNode.ParentJointIndex = Node.Index;
 						}
 					}
 					else
@@ -1672,6 +1785,7 @@ namespace GLTF
 						{
 							GLTF::FNode& JointNode = Asset->Nodes[Group.Value.Indices[RootJointIndex]];
 							JointNode.ParentIndex = Node.Index;
+							JointNode.ParentJointIndex = Node.Index;
 							Node.Children.Add(JointNode.Index);
 						}
 
diff --git a/Engine/Plugins/Interchange/Runtime/Source/Parsers/GLTFCore/Public/GLTFNode.h b/Engine/Plugins/Interchange/Runtime/Source/Parsers/GLTFCore/Public/GLTFNode.h
--- a/Engine/Plugins/Interchange/Runtime/Source/Parsers/GLTFCore/Public/GLTFNode.h
+++ b/Engine/Plugins/Interchange/Runtime/Source/Parsers/GLTFCore/Public/GLTFNode.h
@@ -73,13 +73,14 @@ namespace GLTF
 
 		int32 Index;		// index of FNode in GLTFAsset.Nodes array.
 		int32 ParentIndex;
-		int32 ParentJointIndex;
+		int32 ParentJointIndex; //glTF marks the Nodes in Skins.Joints as a joint, however there are certain indirections that still make a node into a joint. (such as skin.joints having gaps/skips in the node hierarchy)
+		int32 DirectParentJointIndex; //With this variable we are tracking the original Parent Joint Index. Joint index that was actually marked by the glTF file without any indirections, aka index is part of the Skin.Joints list.
 		int32 RootJointIndex; //only valid if node is of Joint Type
 
 		TArray<float> MorphTargetWeights; //for the instantiated mesh with morph targets.
 
-		TMap<int, FMatrix> SkinIndexToGlobalInverseBindMatrix; //Matrix in glTF space
-		TMap<int, FMatrix> SkinIndexToLocalBindMatrix;	//Matrix in glTF space
+		TMap<int, FMatrix> SkinIndexToGlobalInverseBindMatrix;
+		TMap<int, FTransform> SkinIndexToLocalBindTransform;
 														//bind pose would be CurrentNode.GlobalInverseBindTransform.Inverse() * ParentNode.GlobalInverseBindTransform
 		bool bHasLocalBindPose;
 		FTransform LocalBindPose;	// First Skin that's using the joint will fill the LocalBindPose.
@@ -99,6 +100,7 @@ namespace GLTF
 			, Index(INDEX_NONE)
 			, ParentIndex(INDEX_NONE)
 			, ParentJointIndex(INDEX_NONE)
+			, DirectParentJointIndex(INDEX_NONE)
 			, RootJointIndex(INDEX_NONE)
 			, bHasLocalBindPose(false)
 		{
diff --git a/Engine/Plugins/Interchange/Runtime/Source/Parsers/GLTFCore/Public/GLTFReader.h b/Engine/Plugins/Interchange/Runtime/Source/Parsers/GLTFCore/Public/GLTFReader.h
--- a/Engine/Plugins/Interchange/Runtime/Source/Parsers/GLTFCore/Public/GLTFReader.h
+++ b/Engine/Plugins/Interchange/Runtime/Source/Parsers/GLTFCore/Public/GLTFReader.h
@@ -69,8 +69,12 @@ namespace GLTF
 		void GenerateLocalBindPosesPerSkinIndices() const; //Per node
 		void SetLocalBindPosesForJoints() const; //Per node
 
-		void BuildParentIndices(int32 ParentNodeIndex, int32 ParentJointIndex, int32 CurrentNodeIndex) const;
+		void BuildParentJointIndices() const;
+		void BuildParentIndices(int32 ParentNodeIndex, int32 DirectParentJointIndex, int32 CurrentNodeIndex) const;
 		int32 FindRootJointIndex(int32 CurrentNodeIndex) const;
+
+		//Sets the Node Types to Joint for Nodes that were not marked Joint in Skins, however they have descendant and ancestor both in Skins.
+		void FixNodeTypesForIndirectJoints() const;
 		void BuildRootJoints() const;
 
 		void ResolveConflictingTextures();
