Interchange:
- Use Material in search location with glTF

diff --git a/Engine/Plugins/Interchange/Runtime/Source/Pipelines/Private/InterchangeGenericMaterialPipeline.cpp b/Engine/Plugins/Interchange/Runtime/Source/Pipelines/Private/InterchangeGenericMaterialPipeline.cpp
--- a/Engine/Plugins/Interchange/Runtime/Source/Pipelines/Private/InterchangeGenericMaterialPipeline.cpp
+++ b/Engine/Plugins/Interchange/Runtime/Source/Pipelines/Private/InterchangeGenericMaterialPipeline.cpp
@@ -1100,6 +1100,8 @@ void UInterchangeGenericMaterialPipeline::ExecutePipeline(UInterchangeBaseNodeCo
 		FString MaterialName = BaseMaterialFactoryNode->GetDisplayLabel();
 		if (UMaterialInterface* ExistingMaterial = FindExistingMaterialFromSearchLocation(MaterialName, ContentBasePath, SearchLocation))
 		{
+			// Add the search location to notify other pipelines what to do about assets.
+			BaseMaterialFactoryNode->AddInt32Attribute(TEXT("SearchLocation"), int32(SearchLocation));
 			BaseMaterialFactoryNode->SetCustomReferenceObject(ExistingMaterial);
 
 			// We're reimporting the existing material we found: Enable the factory node if we can.
diff --git a/Engine/Plugins/Interchange/Runtime/Source/Pipelines/Private/InterchangeglTFPipeline.cpp b/Engine/Plugins/Interchange/Runtime/Source/Pipelines/Private/InterchangeglTFPipeline.cpp
--- a/Engine/Plugins/Interchange/Runtime/Source/Pipelines/Private/InterchangeglTFPipeline.cpp
+++ b/Engine/Plugins/Interchange/Runtime/Source/Pipelines/Private/InterchangeglTFPipeline.cpp
@@ -7,8 +7,7 @@
 
 #include "InterchangeManager.h"
 #include "InterchangeMaterialFactoryNode.h"
-#include "InterchangeMaterialInstanceNode.h"
-#include "InterchangeMeshFactoryNode.h"
+#include "InterchangeGenericMaterialPipeline.h"
 #include "InterchangeShaderGraphNode.h"
 #include "Nodes/InterchangeSourceNode.h"
 
@@ -17,8 +16,6 @@
 #include "Materials/MaterialInstanceConstant.h"
 #include "Materials/MaterialInstanceDynamic.h"
 
-#include "Misc/App.h"
-
 const TArray<FString> UGLTFPipelineSettings::ExpectedMaterialInstanceIdentifiers = {TEXT("MI_Default_Opaque"), TEXT("MI_Default_Mask"), TEXT("MI_Default_Blend"), 
 																					TEXT("MI_Unlit_Opaque"), TEXT("MI_Unlit_Mask"), TEXT("MI_Unlit_Blend"), 
 																					TEXT("MI_ClearCoat_Opaque"), TEXT("MI_ClearCoat_Mask"), TEXT("MI_ClearCoat_Blend"),
@@ -187,6 +184,20 @@ UInterchangeMaterialInstanceFactoryNode* UGLTFPipelineSettings::BuildMaterialIns
 {
 	using namespace UE::Interchange::GLTFMaterials;
 
+	FString MaterialFactoryNodeUid = UInterchangeFactoryBaseNode::BuildFactoryNodeUid(ShaderGraphNode->GetUniqueID());
+
+	// Check the search location of the material, if under root, then we shouldn't replace the material factory node and just take the existing one
+	if (UInterchangeBaseMaterialFactoryNode* MaterialFactoryNode = Cast<UInterchangeBaseMaterialFactoryNode>(NodeContainer->GetFactoryNode(MaterialFactoryNodeUid)))
+	{
+		if (EInterchangeMaterialSearchLocation SearchLocation; MaterialFactoryNode->GetInt32Attribute(TEXT("SearchLocation"), reinterpret_cast<int32&>(SearchLocation)))
+		{
+			if (SearchLocation != EInterchangeMaterialSearchLocation::DoNotSearch)
+			{
+				return nullptr;
+			}
+		}
+	}
+
 	TArray<UE::Interchange::FAttributeKey> AttributeKeys;
 	ShaderGraphNode->GetAttributeKeys(AttributeKeys);
 
@@ -221,7 +232,6 @@ UInterchangeMaterialInstanceFactoryNode* UGLTFPipelineSettings::BuildMaterialIns
 		return nullptr;
 	}
 
-	FString MaterialFactoryNodeUid = UInterchangeFactoryBaseNode::BuildFactoryNodeUid(ShaderGraphNode->GetUniqueID());
 	FString MaterialFactoryNodeName = ShaderGraphNode->GetDisplayLabel();
 
 	UInterchangeMaterialInstanceFactoryNode* MaterialInstanceFactoryNode = NewObject<UInterchangeMaterialInstanceFactoryNode>(NodeContainer);
