Interchange MaterialX:
- add support for <chiang_hair_roughness>
- add support for <chiang_hair_absorption_from_color>
- add support for <deon_hair_absorption_from_melanin>
- add data driven support for <chian_hair_bsdf>
- add new input "energy_compensation" to <oren_nayar_diffuse_bsdf>
- add new inputs "thinfilm_thickness" and "thinfilm_ior" to <dielectric_bsdf>
- deprecated <thin_film_bsdf> since it has been removed from MaterialX 1.39.3
- use BreakOutFloat material functions for <separate> nodes
- add a custom a library node definitions in Engine/Binaries/ThirdParty/MaterialX/libraries/Interchange
- fixed wrong output index when resolving the correct output to take of a multiouput node
- Using now <convert> nodegraphs

#virtualized

diff --git a/Engine/Plugins/Interchange/Runtime/Config/BaseInterchange.ini b/Engine/Plugins/Interchange/Runtime/Config/BaseInterchange.ini
--- a/Engine/Plugins/Interchange/Runtime/Config/BaseInterchange.ini
+++ b/Engine/Plugins/Interchange/Runtime/Config/BaseInterchange.ini
@@ -60,7 +60,7 @@ MaterialParents=(("MI_Default_Opaque", "/InterchangeAssets/gltf/MaterialInstance
 
 [/Script/InterchangePipelines.MaterialXPipelineSettings]
 PredefinedSurfaceShaders=((OpenPBRSurface, "/InterchangeAssets/Functions/MX_OpenPBR_Opaque.MX_OpenPBR_Opaque"),(OpenPBRSurfaceTransmission, "/InterchangeAssets/Functions/MX_OpenPBR_Translucent.MX_OpenPBR_Translucent"),(StandardSurface, "/InterchangeAssets/Functions/MX_StandardSurface.MX_StandardSurface"),(StandardSurfaceTransmission, "/InterchangeAssets/Functions/MX_TransmissionSurface.MX_TransmissionSurface"),(SurfaceUnlit, "/InterchangeAssets/Functions/MX_SurfaceUnlit.MX_SurfaceUnlit"),(UsdPreviewSurface, "/InterchangeAssets/Functions/MX_UsdPreviewSurface.MX_UsdPreviewSurface"),(Surface, "/InterchangeAssets/Functions/MX_Surface.MX_Surface"),(Displacement, "/InterchangeAssets/Functions/MX_Displacement.MX_Displacement"))
-PredefinedBSDF=((OrenNayarDiffuse, "/InterchangeAssets/Functions/MX_OrenNayarBSDF.MX_OrenNayarBSDF"),(BurleyDiffuse,"/InterchangeAssets/Functions/MX_BurleyDiffuseBSDF.MX_BurleyDiffuseBSDF"),(Translucent,"/InterchangeAssets/Functions/MX_TranslucentBSDF.MX_TranslucentBSDF"),(Dielectric,"/InterchangeAssets/Functions/MX_DielectricBSDF.MX_DielectricBSDF"),(Conductor,"/InterchangeAssets/Functions/MX_ConductorBSDF.MX_ConductorBSDF"),(GeneralizedSchlick,"/InterchangeAssets/Functions/MX_GeneralizedSchlickBSDF.MX_GeneralizedSchlickBSDF"),(Subsurface,"/InterchangeAssets/Functions/MX_SubsurfaceBSDF.MX_SubsurfaceBSDF"),(Sheen,"/InterchangeAssets/Functions/MX_SheenBSDF.MX_SheenBSDF"),(ThinFilm,"/InterchangeAssets/Functions/MX_ThinFilmBSDF.MX_ThinFilmBSDF"))
+PredefinedBSDF=((OrenNayarDiffuse, "/InterchangeAssets/Functions/MX_OrenNayarBSDF.MX_OrenNayarBSDF"),(BurleyDiffuse,"/InterchangeAssets/Functions/MX_BurleyDiffuseBSDF.MX_BurleyDiffuseBSDF"),(Translucent,"/InterchangeAssets/Functions/MX_TranslucentBSDF.MX_TranslucentBSDF"),(Dielectric,"/InterchangeAssets/Functions/MX_DielectricBSDF.MX_DielectricBSDF"),(Conductor,"/InterchangeAssets/Functions/MX_ConductorBSDF.MX_ConductorBSDF"),(GeneralizedSchlick,"/InterchangeAssets/Functions/MX_GeneralizedSchlickBSDF.MX_GeneralizedSchlickBSDF"),(Subsurface,"/InterchangeAssets/Functions/MX_SubsurfaceBSDF.MX_SubsurfaceBSDF"),(Sheen,"/InterchangeAssets/Functions/MX_SheenBSDF.MX_SheenBSDF"),(ChiangHair,"/InterchangeAssets/Functions/MX_ChiangHairBSDF.MX_ChiangHairBSDF"))
 PredefinedEDF=((Uniform, "/InterchangeAssets/Functions/MX_UniformEDF.MX_UniformEDF"),(Conical,"/InterchangeAssets/Functions/MX_ConicalEDF.MX_ConicalEDF"),(Measured,"/InterchangeAssets/Functions/MX_MeasuredEDF.MX_MeasuredEDF"))
 PredefinedVDF=((Absorption, "/InterchangeAssets/Functions/MX_AbsorptionVDF.MX_AbsorptionVDF"),(Anisotropic,"/InterchangeAssets/Functions/MX_AnisotropicVDF.MX_AnisotropicVDF"))
 
diff --git a/Engine/Plugins/Interchange/Runtime/Source/Common/Public/MaterialX/InterchangeMaterialXDefinitions.h b/Engine/Plugins/Interchange/Runtime/Source/Common/Public/MaterialX/InterchangeMaterialXDefinitions.h
--- a/Engine/Plugins/Interchange/Runtime/Source/Common/Public/MaterialX/InterchangeMaterialXDefinitions.h
+++ b/Engine/Plugins/Interchange/Runtime/Source/Common/Public/MaterialX/InterchangeMaterialXDefinitions.h
@@ -81,8 +81,10 @@ enum class EInterchangeMaterialXBSDF : uint8
 	/** A microfacet BSDF for the back-scattering properties of cloth-like materials. */
 	Sheen,
 
-	/** Adds an iridescent thin film layer over a microfacet base BSDF. */
-	ThinFilm,
+	/* Constructs a hair BSDF based on the Chiang hair shading model [^Chiang2016]. This node does not support vertical layering. */
+	ChiangHair,
+
+	ThinFilm UE_DEPRECATED(5.7, "<thin_film_bsdf> has been removed from MaterialX 1.39.3, ThinFilm is therefore deprecated.") UMETA(hidden),
 
 	MaxBSDFCount UMETA(hidden)
 };
@@ -244,6 +246,9 @@ namespace UE
 			static constexpr const TCHAR* Exponent = TEXT("Exponent");
 			static constexpr const TCHAR* Factor = TEXT("Factor");
 			static constexpr const TCHAR* Flatness = TEXT("Flatness");
+			static constexpr const TCHAR* Float2 = TEXT("Float2");
+			static constexpr const TCHAR* Float3 = TEXT("Float3");
+			static constexpr const TCHAR* Float4 = TEXT("Float4");
 			static constexpr const TCHAR* Foreground = TEXT("Foreground");
 			static constexpr const TCHAR* Fraction = TEXT("Fraction");
 			static constexpr const TCHAR* Gamma = TEXT("Gamma");
@@ -295,13 +300,19 @@ namespace UE
 	{
 		namespace Path
 		{
+			constexpr const TCHAR* BreakOutFloat2Components = TEXT("/Engine/Functions/Engine_MaterialFunctions02/Utility/BreakOutFloat2Components.BreakOutFloat2Components");
+			constexpr const TCHAR* BreakOutFloat3Components = TEXT("/Engine/Functions/Engine_MaterialFunctions02/Utility/BreakOutFloat3Components.BreakOutFloat3Components");
+			constexpr const TCHAR* BreakOutFloat4Components = TEXT("/Engine/Functions/Engine_MaterialFunctions02/Utility/BreakOutFloat4Components.BreakOutFloat4Components");
 			constexpr const TCHAR* HeightToNormalSmooth = TEXT("/Engine/Functions/Engine_MaterialFunctions03/Procedurals/HeightToNormalSmooth.HeightToNormalSmooth");
 			constexpr const TCHAR* MxACEScgToRec709 = TEXT("/InterchangeAssets/Functions/MX_ACEScgToRec709.MX_ACEScgToRec709");
 			constexpr const TCHAR* MxAdobeRGBToRec709 = TEXT("/InterchangeAssets/Functions/MX_AdobeRGBToRec709.MX_AdobeRGBToRec709");
 			constexpr const TCHAR* MxArtisticIOR = TEXT("/InterchangeAssets/Functions/MX_Artistic_IOR.MX_Artistic_IOR");
 			constexpr const TCHAR* MxCircle = TEXT("/InterchangeAssets/Functions/MX_Circle.MX_Circle");
 			constexpr const TCHAR* MxCheckerboard = TEXT("/InterchangeAssets/Functions/MX_Checkerboard.MX_Checkerboard");
+			constexpr const TCHAR* MxChiangHairAbsorptionFromColor = TEXT("/InterchangeAssets/Functions/MX_ChiangHairAbsorptionFromColor.MX_ChiangHairAbsorptionFromColor");
+			constexpr const TCHAR* MxChiangHairRoughness = TEXT("/InterchangeAssets/Functions/MX_ChiangHairRoughness.MX_ChiangHairRoughness");
 			constexpr const TCHAR* MxColorCorrect = TEXT("/InterchangeAssets/Functions/MX_ColorCorrect.MX_ColorCorrect");
+			constexpr const TCHAR* MxDeonHairAbsorptionFromMelanin = TEXT("/InterchangeAssets/Functions/MX_DeonHairAbsorptionFromMelanin.MX_DeonHairAbsorptionFromMelanin");
 			constexpr const TCHAR* MxGoochShade = TEXT("/InterchangeAssets/Functions/MX_GoochShade.MX_GoochShade");
 			constexpr const TCHAR* MxHsvAdjust = TEXT("/InterchangeAssets/Functions/MX_HSVAdjust.MX_HSVAdjust");
 			constexpr const TCHAR* MxLine = TEXT("/InterchangeAssets/Functions/MX_Line.MX_Line");
@@ -798,7 +809,6 @@ MATERIALX_NAMESPACE_BEGIN
 		static constexpr const char* OrenNayarDiffuseBSDF = "oren_nayar_diffuse_bsdf";
 		static constexpr const char* SheenBSDF = "sheen_bsdf";
 		static constexpr const char* SubsurfaceBSDF= "subsurface_bsdf";
-		static constexpr const char* ThinFilmBSDF= "thin_film_bsdf";
 		static constexpr const char* TranslucentBSDF = "translucent_bsdf";
 		// EDF
 		static constexpr const char* ConicalEDF = "conical_edf";
@@ -810,6 +820,9 @@ MATERIALX_NAMESPACE_BEGIN
 		// PBR Utility Nodes
 		static constexpr const char* ArtisticIOR = "artistic_ior";
 		static constexpr const char* Blackbody = "blackbody";
+		static constexpr const char* ChiangHairAbsorptionFromColor = "chiang_hair_absorption_from_color";
+		static constexpr const char* ChiangHairRoughness = "chiang_hair_roughness";
+		static constexpr const char* DeonHairAbsorptionFromMelanin = "deon_hair_absorption_from_melanin";
 		static constexpr const char* Layer = "layer";
 		static constexpr const char* RoughnessAnisotropy = "roughness_anisotropy";
 		static constexpr const char* RoughnessDual = "roughness_dual";
diff --git a/Engine/Plugins/Interchange/Runtime/Source/Import/Private/MaterialX/MaterialXUtils/MaterialXManager.cpp b/Engine/Plugins/Interchange/Runtime/Source/Import/Private/MaterialX/MaterialXUtils/MaterialXManager.cpp
--- a/Engine/Plugins/Interchange/Runtime/Source/Import/Private/MaterialX/MaterialXUtils/MaterialXManager.cpp
+++ b/Engine/Plugins/Interchange/Runtime/Source/Import/Private/MaterialX/MaterialXUtils/MaterialXManager.cpp
@@ -109,6 +109,9 @@ FMaterialXManager::FMaterialXManager()
 	   {{mx::Category::Rotate3D,              TEXT("in")},          ExpressionInput(UE::Expressions::Inputs::Position)},
 	   {{mx::Category::Saturate,              TEXT("amount")},      ExpressionInput(UE::Expressions::Inputs::Fraction)},
 	   {{mx::Category::Smoothstep,            TEXT("in")},          ExpressionInput(UE::Expressions::Inputs::Value)},
+	   {{mx::Category::Separate2,             TEXT("in")},          ExpressionInput(UE::Expressions::Inputs::Float2)},
+	   {{mx::Category::Separate3,             TEXT("in")},          ExpressionInput(UE::Expressions::Inputs::Float3)},
+	   {{mx::Category::Separate4,             TEXT("in")},          ExpressionInput(UE::Expressions::Inputs::Float4)},
 	   {{mx::Category::Switch,                TEXT("in1")},         ExpressionInput(UE::Expressions::Inputs::in1)},
 	   {{mx::Category::Switch,                TEXT("in2")},         ExpressionInput(UE::Expressions::Inputs::in2)},
 	   {{mx::Category::Switch,                TEXT("in3")},         ExpressionInput(UE::Expressions::Inputs::in3)},
@@ -218,68 +221,76 @@ FMaterialXManager::FMaterialXManager()
 
 		// In case of StandardSurface and OpenPBR we set by default the value to the opaque shader, we decide later if we need to change for the transmission shader
 		// Surface Shader nodes
-		{mx::Category::StandardSurface,		    FMaterialXMaterialFunction{TInPlaceType<EInterchangeMaterialXShaders>{}, EInterchangeMaterialXShaders::StandardSurface}},
-		{mx::Category::OpenPBRSurface,		    FMaterialXMaterialFunction{TInPlaceType<EInterchangeMaterialXShaders>{}, EInterchangeMaterialXShaders::OpenPBRSurface}},
-		{mx::Category::UsdPreviewSurface,		FMaterialXMaterialFunction{TInPlaceType<EInterchangeMaterialXShaders>{}, EInterchangeMaterialXShaders::UsdPreviewSurface}},
-		{mx::Category::Surface,					FMaterialXMaterialFunction{TInPlaceType<EInterchangeMaterialXShaders>{}, EInterchangeMaterialXShaders::Surface}},
-		{mx::Category::SurfaceUnlit,			FMaterialXMaterialFunction{TInPlaceType<EInterchangeMaterialXShaders>{}, EInterchangeMaterialXShaders::SurfaceUnlit}},
-		{mx::Category::Displacement,			FMaterialXMaterialFunction{TInPlaceType<EInterchangeMaterialXShaders>{}, EInterchangeMaterialXShaders::Displacement}},
+		{mx::Category::StandardSurface,					FMaterialXMaterialFunction{TInPlaceType<EInterchangeMaterialXShaders>{}, EInterchangeMaterialXShaders::StandardSurface}},
+		{mx::Category::OpenPBRSurface,					FMaterialXMaterialFunction{TInPlaceType<EInterchangeMaterialXShaders>{}, EInterchangeMaterialXShaders::OpenPBRSurface}},
+		{mx::Category::UsdPreviewSurface,				FMaterialXMaterialFunction{TInPlaceType<EInterchangeMaterialXShaders>{}, EInterchangeMaterialXShaders::UsdPreviewSurface}},
+		{mx::Category::Surface,							FMaterialXMaterialFunction{TInPlaceType<EInterchangeMaterialXShaders>{}, EInterchangeMaterialXShaders::Surface}},
+		{mx::Category::SurfaceUnlit,					FMaterialXMaterialFunction{TInPlaceType<EInterchangeMaterialXShaders>{}, EInterchangeMaterialXShaders::SurfaceUnlit}},
+		{mx::Category::Displacement,					FMaterialXMaterialFunction{TInPlaceType<EInterchangeMaterialXShaders>{}, EInterchangeMaterialXShaders::Displacement}},
 
 		// BSDF Nodes
-		{mx::Category::BurleyDiffuseBSDF,		FMaterialXMaterialFunction{TInPlaceType<EInterchangeMaterialXBSDF>{}, EInterchangeMaterialXBSDF::BurleyDiffuse}},
-		{mx::Category::ConductorBSDF,			FMaterialXMaterialFunction{TInPlaceType<EInterchangeMaterialXBSDF>{}, EInterchangeMaterialXBSDF::Conductor}},
-		{mx::Category::DielectricBSDF,			FMaterialXMaterialFunction{TInPlaceType<EInterchangeMaterialXBSDF>{}, EInterchangeMaterialXBSDF::Dielectric}},
-		{mx::Category::GeneralizedSchlickBSDF,	FMaterialXMaterialFunction{TInPlaceType<EInterchangeMaterialXBSDF>{}, EInterchangeMaterialXBSDF::GeneralizedSchlick}},
-		{mx::Category::OrenNayarDiffuseBSDF ,	FMaterialXMaterialFunction{TInPlaceType<EInterchangeMaterialXBSDF>{}, EInterchangeMaterialXBSDF::OrenNayarDiffuse}},
-		{mx::Category::SheenBSDF,				FMaterialXMaterialFunction{TInPlaceType<EInterchangeMaterialXBSDF>{}, EInterchangeMaterialXBSDF::Sheen}},
-		{mx::Category::SubsurfaceBSDF ,			FMaterialXMaterialFunction{TInPlaceType<EInterchangeMaterialXBSDF>{}, EInterchangeMaterialXBSDF::Subsurface}},
-		{mx::Category::ThinFilmBSDF,			FMaterialXMaterialFunction{TInPlaceType<EInterchangeMaterialXBSDF>{}, EInterchangeMaterialXBSDF::ThinFilm}},
-		{mx::Category::TranslucentBSDF,			FMaterialXMaterialFunction{TInPlaceType<EInterchangeMaterialXBSDF>{}, EInterchangeMaterialXBSDF::Translucent}},
+		{mx::Category::BurleyDiffuseBSDF,				FMaterialXMaterialFunction{TInPlaceType<EInterchangeMaterialXBSDF>{}, EInterchangeMaterialXBSDF::BurleyDiffuse}},
+		{mx::Category::ConductorBSDF,					FMaterialXMaterialFunction{TInPlaceType<EInterchangeMaterialXBSDF>{}, EInterchangeMaterialXBSDF::Conductor}},
+		{mx::Category::DielectricBSDF,					FMaterialXMaterialFunction{TInPlaceType<EInterchangeMaterialXBSDF>{}, EInterchangeMaterialXBSDF::Dielectric}},
+		{mx::Category::GeneralizedSchlickBSDF,			FMaterialXMaterialFunction{TInPlaceType<EInterchangeMaterialXBSDF>{}, EInterchangeMaterialXBSDF::GeneralizedSchlick}},
+		{mx::Category::OrenNayarDiffuseBSDF ,			FMaterialXMaterialFunction{TInPlaceType<EInterchangeMaterialXBSDF>{}, EInterchangeMaterialXBSDF::OrenNayarDiffuse}},
+		{mx::Category::SheenBSDF,						FMaterialXMaterialFunction{TInPlaceType<EInterchangeMaterialXBSDF>{}, EInterchangeMaterialXBSDF::Sheen}},
+		{mx::Category::SubsurfaceBSDF ,					FMaterialXMaterialFunction{TInPlaceType<EInterchangeMaterialXBSDF>{}, EInterchangeMaterialXBSDF::Subsurface}},
+		{mx::Category::TranslucentBSDF,					FMaterialXMaterialFunction{TInPlaceType<EInterchangeMaterialXBSDF>{}, EInterchangeMaterialXBSDF::Translucent}},
 		// EDF Nodes
-		{mx::Category::ConicalEDF,				FMaterialXMaterialFunction{TInPlaceType<EInterchangeMaterialXEDF>{}, EInterchangeMaterialXEDF::Conical}},
-		{mx::Category::MeasuredEDF,				FMaterialXMaterialFunction{TInPlaceType<EInterchangeMaterialXEDF>{}, EInterchangeMaterialXEDF::Measured}},
-		{mx::Category::UniformEDF,				FMaterialXMaterialFunction{TInPlaceType<EInterchangeMaterialXEDF>{}, EInterchangeMaterialXEDF::Uniform}},
+		{mx::Category::ConicalEDF,						FMaterialXMaterialFunction{TInPlaceType<EInterchangeMaterialXEDF>{}, EInterchangeMaterialXEDF::Conical}},
+		{mx::Category::MeasuredEDF,						FMaterialXMaterialFunction{TInPlaceType<EInterchangeMaterialXEDF>{}, EInterchangeMaterialXEDF::Measured}},
+		{mx::Category::UniformEDF,						FMaterialXMaterialFunction{TInPlaceType<EInterchangeMaterialXEDF>{}, EInterchangeMaterialXEDF::Uniform}},
 		// VDF Nodes
-		{mx::Category::AbsorptionVDF,			FMaterialXMaterialFunction{TInPlaceType<EInterchangeMaterialXVDF>{}, EInterchangeMaterialXVDF::Absorption}},
-		{mx::Category::AnisotropicVDF,			FMaterialXMaterialFunction{TInPlaceType<EInterchangeMaterialXVDF>{}, EInterchangeMaterialXVDF::Anisotropic}},
+		{mx::Category::AbsorptionVDF,					FMaterialXMaterialFunction{TInPlaceType<EInterchangeMaterialXVDF>{}, EInterchangeMaterialXVDF::Absorption}},
+		{mx::Category::AnisotropicVDF,					FMaterialXMaterialFunction{TInPlaceType<EInterchangeMaterialXVDF>{}, EInterchangeMaterialXVDF::Anisotropic}},
 		// Utility nodes
-		{mx::Category::ArtisticIOR,				FMaterialXMaterialFunction{TInPlaceType<FString>{}, UE::MaterialFunctions::Path::MxArtisticIOR}},
-		{mx::Category::RoughnessAnisotropy,		FMaterialXMaterialFunction{TInPlaceType<FString>{}, UE::MaterialFunctions::Path::MxRoughnesAnisotropy}},
-		{mx::Category::RoughnessDual,			FMaterialXMaterialFunction{TInPlaceType<FString>{}, UE::MaterialFunctions::Path::MxRoughnessDual}},
-		// Math
-		{mx::Category::Place2D,					FMaterialXMaterialFunction{TInPlaceType<FString>{}, UE::MaterialFunctions::Path::MxPlace2D}},
-		{mx::Category::Refract,					FMaterialXMaterialFunction{TInPlaceType<FString>{}, UE::MaterialFunctions::Path::Refract}},
-		// ColorTransform
-		{mx::Category::ACEScgToLinRec709,		FMaterialXMaterialFunction{TInPlaceType<FString>{}, UE::MaterialFunctions::Path::MxACEScgToRec709}},
-		{mx::Category::AdobeRgbToLinRec709,		FMaterialXMaterialFunction{TInPlaceType<FString>{}, UE::MaterialFunctions::Path::MxAdobeRGBToRec709}},
-		{mx::Category::LinAdobeRgbToLinRec709,	FMaterialXMaterialFunction{TInPlaceType<FString>{}, UE::MaterialFunctions::Path::MxLinearAdobeRGBToRec709}},
-		{mx::Category::LinDisplayP3ToLinRec709,	FMaterialXMaterialFunction{TInPlaceType<FString>{}, UE::MaterialFunctions::Path::MxLinearDisplayP3ToRec709}},
-		{mx::Category::SrgbDisplayP3ToLinRec709,FMaterialXMaterialFunction{TInPlaceType<FString>{}, UE::MaterialFunctions::Path::MxSrgbDisplayP3ToRec709}},
-		{mx::Category::SrgbTextureToLinRec709,	FMaterialXMaterialFunction{TInPlaceType<FString>{}, UE::MaterialFunctions::Path::MxSrgbToRec709}},
-		// Procedural
-		{mx::Category::Checkerboard,			FMaterialXMaterialFunction{TInPlaceType<FString>{}, UE::MaterialFunctions::Path::MxCheckerboard}},
-		{mx::Category::Circle,					FMaterialXMaterialFunction{TInPlaceType<FString>{}, UE::MaterialFunctions::Path::MxCircle}},
-		{mx::Category::Line,					FMaterialXMaterialFunction{TInPlaceType<FString>{}, UE::MaterialFunctions::Path::MxLine}},
-		{mx::Category::Ramp,					FMaterialXMaterialFunction{TInPlaceType<FString>{}, UE::MaterialFunctions::Path::MxRamp}},
-		{mx::Category::RampGradient,	        FMaterialXMaterialFunction{TInPlaceType<FString>{}, UE::MaterialFunctions::Path::MxRampGradient}},
-		{mx::Category::RandomColor,	            FMaterialXMaterialFunction{TInPlaceType<FString>{}, UE::MaterialFunctions::Path::MxRandomColor}},
-		{mx::Category::TiledCircles,			FMaterialXMaterialFunction{TInPlaceType<FString>{}, UE::MaterialFunctions::Path::MxTiledCircles}},
-		{mx::Category::UnifiedNoise3D,			FMaterialXMaterialFunction{TInPlaceType<FString>{}, UE::MaterialFunctions::Path::MxUnifiedNoise3D}},
-		// Adjustment
-		{mx::Category::ColorCorrect,	        FMaterialXMaterialFunction{TInPlaceType<FString>{}, UE::MaterialFunctions::Path::MxColorCorrect}},
-		{mx::Category::HsvAdjust,	            FMaterialXMaterialFunction{TInPlaceType<FString>{}, UE::MaterialFunctions::Path::MxHsvAdjust}},
-		//NPR
-		{mx::Category::GoochShade,	            FMaterialXMaterialFunction{TInPlaceType<FString>{}, UE::MaterialFunctions::Path::MxGoochShade}},
+		{mx::Category::ArtisticIOR,						FMaterialXMaterialFunction{TInPlaceType<FString>{}, MaterialFunctionPackage(UE::MaterialFunctions::Path::MxArtisticIOR)}},
+		{mx::Category::ChiangHairAbsorptionFromColor,	FMaterialXMaterialFunction{TInPlaceType<FString>{}, MaterialFunctionPackage(UE::MaterialFunctions::Path::MxChiangHairAbsorptionFromColor)}},
+		{mx::Category::ChiangHairRoughness,				FMaterialXMaterialFunction{TInPlaceType<FString>{}, MaterialFunctionPackage(UE::MaterialFunctions::Path::MxChiangHairRoughness)}},
+		{mx::Category::DeonHairAbsorptionFromMelanin,	FMaterialXMaterialFunction{TInPlaceType<FString>{}, MaterialFunctionPackage(UE::MaterialFunctions::Path::MxDeonHairAbsorptionFromMelanin)}},
+		{mx::Category::RoughnessAnisotropy,				FMaterialXMaterialFunction{TInPlaceType<FString>{}, MaterialFunctionPackage(UE::MaterialFunctions::Path::MxRoughnesAnisotropy)}},
+		{mx::Category::RoughnessDual,					FMaterialXMaterialFunction{TInPlaceType<FString>{}, MaterialFunctionPackage(UE::MaterialFunctions::Path::MxRoughnessDual)}},
+		// Math																								
+		{mx::Category::Place2D,							FMaterialXMaterialFunction{TInPlaceType<FString>{}, MaterialFunctionPackage(UE::MaterialFunctions::Path::MxPlace2D)}},
+		{mx::Category::Refract,							FMaterialXMaterialFunction{TInPlaceType<FString>{}, MaterialFunctionPackage(UE::MaterialFunctions::Path::Refract)}},
+		// ColorTransform																					
+		{mx::Category::ACEScgToLinRec709,				FMaterialXMaterialFunction{TInPlaceType<FString>{}, MaterialFunctionPackage(UE::MaterialFunctions::Path::MxACEScgToRec709)}},
+		{mx::Category::AdobeRgbToLinRec709,				FMaterialXMaterialFunction{TInPlaceType<FString>{}, MaterialFunctionPackage(UE::MaterialFunctions::Path::MxAdobeRGBToRec709)}},
+		{mx::Category::LinAdobeRgbToLinRec709,			FMaterialXMaterialFunction{TInPlaceType<FString>{}, MaterialFunctionPackage(UE::MaterialFunctions::Path::MxLinearAdobeRGBToRec709)}},
+		{mx::Category::LinDisplayP3ToLinRec709,			FMaterialXMaterialFunction{TInPlaceType<FString>{}, MaterialFunctionPackage(UE::MaterialFunctions::Path::MxLinearDisplayP3ToRec709)}},
+		{mx::Category::SrgbDisplayP3ToLinRec709,		FMaterialXMaterialFunction{TInPlaceType<FString>{}, MaterialFunctionPackage(UE::MaterialFunctions::Path::MxSrgbDisplayP3ToRec709)}},
+		{mx::Category::SrgbTextureToLinRec709,			FMaterialXMaterialFunction{TInPlaceType<FString>{}, MaterialFunctionPackage(UE::MaterialFunctions::Path::MxSrgbToRec709)}},
+		// Procedural																						
+		{mx::Category::Checkerboard,					FMaterialXMaterialFunction{TInPlaceType<FString>{}, MaterialFunctionPackage(UE::MaterialFunctions::Path::MxCheckerboard)}},
+		{mx::Category::Circle,							FMaterialXMaterialFunction{TInPlaceType<FString>{}, MaterialFunctionPackage(UE::MaterialFunctions::Path::MxCircle)}},
+		{mx::Category::Line,							FMaterialXMaterialFunction{TInPlaceType<FString>{}, MaterialFunctionPackage(UE::MaterialFunctions::Path::MxLine)}},
+		{mx::Category::Ramp,							FMaterialXMaterialFunction{TInPlaceType<FString>{}, MaterialFunctionPackage(UE::MaterialFunctions::Path::MxRamp)}},
+		{mx::Category::RampGradient,					FMaterialXMaterialFunction{TInPlaceType<FString>{}, MaterialFunctionPackage(UE::MaterialFunctions::Path::MxRampGradient)}},
+		{mx::Category::RandomColor,						FMaterialXMaterialFunction{TInPlaceType<FString>{}, MaterialFunctionPackage(UE::MaterialFunctions::Path::MxRandomColor)}},
+		{mx::Category::TiledCircles,					FMaterialXMaterialFunction{TInPlaceType<FString>{}, MaterialFunctionPackage(UE::MaterialFunctions::Path::MxTiledCircles)}},
+		{mx::Category::UnifiedNoise3D,					FMaterialXMaterialFunction{TInPlaceType<FString>{}, MaterialFunctionPackage(UE::MaterialFunctions::Path::MxUnifiedNoise3D)}},
+		// Adjustment																						
+		{mx::Category::ColorCorrect,					FMaterialXMaterialFunction{TInPlaceType<FString>{}, MaterialFunctionPackage(UE::MaterialFunctions::Path::MxColorCorrect)}},
+		{mx::Category::HsvAdjust,						FMaterialXMaterialFunction{TInPlaceType<FString>{}, MaterialFunctionPackage(UE::MaterialFunctions::Path::MxHsvAdjust)}},
+		// NPR
+		{mx::Category::GoochShade,						FMaterialXMaterialFunction{TInPlaceType<FString>{}, MaterialFunctionPackage(UE::MaterialFunctions::Path::MxGoochShade)}},
+		// Channel
+		{mx::Category::Separate2,						FMaterialXMaterialFunction{TInPlaceType<FString>{}, MaterialFunctionPackage(UE::MaterialFunctions::Path::BreakOutFloat2Components)}},
+		{mx::Category::Separate3,						FMaterialXMaterialFunction{TInPlaceType<FString>{}, MaterialFunctionPackage(UE::MaterialFunctions::Path::BreakOutFloat3Components)}},
+		{mx::Category::Separate4,						FMaterialXMaterialFunction{TInPlaceType<FString>{}, MaterialFunctionPackage(UE::MaterialFunctions::Path::BreakOutFloat4Components)}},
 	}
 	, CategoriesToSkip {
 		mx::Category::ACEScgToLinRec709,
 		mx::Category::AdobeRgbToLinRec709,
 		mx::Category::Checkerboard,
+		mx::Category::ChiangHairAbsorptionFromColor,
+		mx::Category::ChiangHairRoughness,
 		mx::Category::Circle,
 		mx::Category::ColorCorrect,
 		mx::Category::Contrast,
-		mx::Category::Convert,
 		mx::Category::Distance,
+		mx::Category::DeonHairAbsorptionFromMelanin,
 		mx::Category::Extract,
 		mx::Category::Fractal3D,
 		mx::Category::GoochShade,
@@ -691,10 +702,16 @@ const TCHAR* FMaterialXManager::ExpressionInput(const TCHAR* Input)
 	return Input;
 }
 
+INTERCHANGEIMPORT_API const TCHAR* FMaterialXManager::MaterialFunctionPackage(const TCHAR* Input)
+{
+	MaterialFunctionsToLoad.FindOrAdd(Input);
+	return Input;
+}
+
 #endif
 
 namespace UE::Interchange::MaterialX
-{
+{	
 	bool AreMaterialFunctionPackagesLoaded()
 	{
 #if WITH_EDITOR
@@ -722,32 +739,7 @@ namespace UE::Interchange::MaterialX
 			return bAllLoaded;
 		};
 
-		static const bool bPackagesLoaded =	ArePackagesLoaded({ UE::MaterialFunctions::Path::NormalFromHeightMap,
-															    UE::MaterialFunctions::Path::MxArtisticIOR,
-															    UE::MaterialFunctions::Path::MxRoughnesAnisotropy,
-															    UE::MaterialFunctions::Path::MxRoughnessDual,
-																UE::MaterialFunctions::Path::MxPlace2D,
-																UE::MaterialFunctions::Path::Refract,
-																UE::MaterialFunctions::Path::MxACEScgToRec709,
-																UE::MaterialFunctions::Path::MxAdobeRGBToRec709,
-																UE::MaterialFunctions::Path::MxLinearDisplayP3ToRec709,
-																UE::MaterialFunctions::Path::MxLinearAdobeRGBToRec709,
-																UE::MaterialFunctions::Path::MxSrgbToRec709,
-																UE::MaterialFunctions::Path::MxSrgbDisplayP3ToRec709,
-																UE::MaterialFunctions::Path::MxRandomFloat,
-																UE::MaterialFunctions::Path::MxRandomColor,
-																UE::MaterialFunctions::Path::HeightToNormalSmooth,
-																UE::MaterialFunctions::Path::MxColorCorrect,
-																UE::MaterialFunctions::Path::MxHsvAdjust,
-																UE::MaterialFunctions::Path::MxGoochShade,
-																UE::MaterialFunctions::Path::MxCircle,
-																UE::MaterialFunctions::Path::MxCheckerboard,
-																UE::MaterialFunctions::Path::MxTiledCircles,
-																UE::MaterialFunctions::Path::MxUnifiedNoise3D,
-																UE::MaterialFunctions::Path::MxLine,
-																UE::MaterialFunctions::Path::MxRamp,
-																UE::MaterialFunctions::Path::MxRampGradient,
-															  });
+		static const bool bPackagesLoaded =	ArePackagesLoaded(FMaterialXManager::GetInstance().MaterialFunctionsToLoad.Array());
 
 		return bPackagesLoaded;
 #else
diff --git a/Engine/Plugins/Interchange/Runtime/Source/Import/Private/MaterialX/MaterialXUtils/MaterialXSurfaceShaderAbstract.cpp b/Engine/Plugins/Interchange/Runtime/Source/Import/Private/MaterialX/MaterialXUtils/MaterialXSurfaceShaderAbstract.cpp
--- a/Engine/Plugins/Interchange/Runtime/Source/Import/Private/MaterialX/MaterialXUtils/MaterialXSurfaceShaderAbstract.cpp
+++ b/Engine/Plugins/Interchange/Runtime/Source/Import/Private/MaterialX/MaterialXUtils/MaterialXSurfaceShaderAbstract.cpp
@@ -265,7 +265,7 @@ bool FMaterialXSurfaceShaderAbstract::ConnectMatchingNodeOutputToInput(const FCo
 	{
 		mx::NodeDefPtr NodeDef = Connect.UpstreamNode->getNodeDef(mx::EMPTY_STRING, true);
 		int Index = NodeDef->getChildIndex(TCHAR_TO_UTF8(*Connect.OutputName));
-		return Index < 0 ? Index : Index % NodeDef->getInputCount();
+		return Index -= NodeDef->getInputCount();
 	};
 
 	bool bIsConnected = false;
@@ -295,8 +295,7 @@ bool FMaterialXSurfaceShaderAbstract::ConnectMatchingNodeOutputToInput(const FCo
 	{
 		UInterchangeShaderNode* OperatorNode = nullptr;
 
-		//We don't take the node output here because it would cause the creation of a new node (output is meaningful with ComponentMaskNode/separate where we have to create a new expression
-		OperatorNode = (this->*CreateFunctionCallOrShaderNode)(Connect.UpstreamNode, Connect.UpstreamNode->getName().c_str(), *ShaderType, DefaultOutput);
+		OperatorNode = (this->*CreateFunctionCallOrShaderNode)(Connect.UpstreamNode, Connect.UpstreamNode->getName().c_str(), *ShaderType, Connect.OutputName);
 
 		ConnectOutputToInputInternal(OperatorNode);
 	};
@@ -305,8 +304,7 @@ bool FMaterialXSurfaceShaderAbstract::ConnectMatchingNodeOutputToInput(const FCo
 	{
 		UInterchangeShaderNode* OperatorNode = nullptr;
 
-		//We don't take the node output here because it would cause the creation of a new node (output is meaningful with ComponentMaskNode/separate where we have to create a new expression
-		OperatorNode = (this->*CreateFunctionCallOrShaderNode)(Connect.UpstreamNode, Connect.UpstreamNode->getName().c_str(), EnumType, EnumValue, DefaultOutput);
+		OperatorNode = (this->*CreateFunctionCallOrShaderNode)(Connect.UpstreamNode, Connect.UpstreamNode->getName().c_str(), EnumType, EnumValue, Connect.OutputName);
 
 		ConnectOutputToInputInternal(OperatorNode);
 	};
@@ -385,22 +383,28 @@ void FMaterialXSurfaceShaderAbstract::ConnectNodeCategoryOutputToInput(const Mat
 				if(ConnectedInput)
 				{
 					InputChannelName = GetInputName(ConnectedInput);
-					if(ConnectedInput->hasOutputString())
+					if (ConnectedInput->hasOutputString())
 					{
 						OutputChannelName = ConnectedInput->getOutputString().c_str();
 					}
+					// we take the output name from the nodedef, because some nodes, like <chiang_hair_absorption_from_color>, default output name is not called "out"
+					// and we make sure we store the correct output name even for default outputs
+					else if (std::vector<mx::OutputPtr> Outputs = UpstreamNode->getNodeDef(mx::EMPTY_STRING, true)->getActiveOutputs(); Outputs.size() == 1)
+					{
+						OutputChannelName = Outputs[0]->getName().c_str();
+					}
 				}
 
-				std::vector<mx::OutputPtr> Outputs{ DownstreamNode->getActiveOutputs()};
-				if(Outputs.empty())
+				std::vector<mx::OutputPtr> Outputs{ DownstreamNode->getActiveOutputs() };
+				if (Outputs.empty())
 				{
 					Outputs = DownstreamNode->getNodeDef(mx::EMPTY_STRING, true)->getActiveOutputs();
 				}
 
 				ParentShaderNodeOutputs.Empty();
-				for(std::size_t i = 0; i < Outputs.size(); ++i)
+				for (std::size_t Index = 0; Index < Outputs.size(); ++Index)
 				{
-					if(UInterchangeShaderNode** FoundNode = ShaderNodes.Find({ GetAttributeParentName(DownstreamNode, ConnectedInput), Outputs[i]->getName().c_str()}))
+					if(UInterchangeShaderNode** FoundNode = ShaderNodes.Find({ GetAttributeParentName(DownstreamNode, ConnectedInput), Outputs[Index]->getName().c_str()}))
 					{
 						ParentShaderNodeOutputs.Emplace(*FoundNode);
 					}
@@ -408,9 +412,9 @@ void FMaterialXSurfaceShaderAbstract::ConnectNodeCategoryOutputToInput(const Mat
 			}
 		}
 
-		for(UInterchangeShaderNode* ParentShaderNode: ParentShaderNodeOutputs)
+		for (UInterchangeShaderNode* ParentShaderNode : ParentShaderNodeOutputs)
 		{
-			if(!ConnectMatchingNodeOutputToInput({ UpstreamNode, ParentShaderNode, InputChannelName, OutputChannelName }))
+			if (!ConnectMatchingNodeOutputToInput({ UpstreamNode, ParentShaderNode, InputChannelName, OutputChannelName }))
 			{
 				UE_LOG(LogInterchangeImport, Warning, TEXT("<%s>: \"%s\" is not supported. [%s]"), UTF8_TO_TCHAR(UpstreamNode->getCategory().c_str()), UTF8_TO_TCHAR(UpstreamNode->getName().c_str()), *SurfaceMaterialName);
 			}
@@ -601,148 +605,6 @@ void FMaterialXSurfaceShaderAbstract::ConnectImageInputToOutput(const FConnectNo
 	}
 }
 
-void FMaterialXSurfaceShaderAbstract::ConnectConvertInputToOutput(const FConnectNode& Connect)
-{
-	using namespace UE::Interchange::Materials::Standard::Nodes;
-
-	//In case of an upwards conversion, let's do an append, downwards a mask, otherwise leave as is
-	mx::InputPtr Input = Connect.UpstreamNode->getInput("in");
-	const std::string& NodeType = Connect.UpstreamNode->getType();
-	const std::string& InputType = Input->getType();
-
-	//Ensure that the types are supported
-	const bool bIsNodeTypeSupported =
-		NodeType == mx::Type::Color4 ||
-		NodeType == mx::Type::Color3 ||
-		NodeType == mx::Type::Vector4 ||
-		NodeType == mx::Type::Vector3 ||
-		NodeType == mx::Type::Vector2 ||
-		NodeType == mx::Type::Float ||
-		NodeType == mx::Type::Integer ||
-		NodeType == mx::Type::Boolean;
-
-	const bool bIsInputTypeSupported =
-		InputType == mx::Type::Color4 ||
-		InputType == mx::Type::Color3 ||
-		InputType == mx::Type::Vector4 ||
-		InputType == mx::Type::Vector3 ||
-		InputType == mx::Type::Vector2 ||
-		InputType == mx::Type::Float ||
-		InputType == mx::Type::Integer ||
-		InputType == mx::Type::Boolean;
-
-	if(!bIsNodeTypeSupported || !bIsInputTypeSupported)
-	{
-		UE_LOG(LogInterchangeImport, Warning, TEXT("<convert> node has unsupported types."));
-		return;
-	}
-
-	char NodeN = NodeType.back();
-	char InputN = InputType.back();
-
-	constexpr auto Remap = [](auto Value, auto Low1, auto High1, auto Low2, auto High2)
-	{
-		return Low2 + (Value - Low1) * (High2 - Low2) / (High1 - Low1);
-	};
-
-	//Remap NodeN and InputN in a range lower than '2' in case of integer/boolean/float
-	if(NodeN > '4')
-	{
-		NodeN = Remap(NodeN, 'a', 'z', 0, '1');
-	}
-	if(InputN > '4')
-	{
-		InputN = Remap(InputN, 'a', 'z', 0, '1');
-	}
-
-	if(InputN > NodeN) // Mask
-	{
-		UInterchangeShaderNode* MaskShaderNode = nullptr;
-		if(NodeN == '3')
-		{
-			MaskShaderNode = CreateMaskShaderNode(0b1110, Connect.UpstreamNode, Connect.UpstreamNode->getName().c_str());
-		}
-		else if(NodeN == '2')
-		{
-			MaskShaderNode = CreateMaskShaderNode(0b1100, Connect.UpstreamNode, Connect.UpstreamNode->getName().c_str());
-		}
-		else
-		{
-			MaskShaderNode = CreateMaskShaderNode(0b1000, Connect.UpstreamNode, Connect.UpstreamNode->getName().c_str());
-		}
-
-		AddAttributeFromValueOrInterface(Input, UE::Expressions::Inputs::Input, MaskShaderNode);
-
-		UInterchangeShaderPortsAPI::ConnectDefaultOuputToInput(Connect.ParentShaderNode, Connect.InputChannelName, MaskShaderNode->GetUniqueID());
-	}
-	else // Append
-	{
-		//same as dot, just connect the next output to this parent input
-		SetAttributeNewName(Input, TCHAR_TO_UTF8(*Connect.InputChannelName)); //let's take the parent node's input name
-		ShaderNodes.Add({ Connect.UpstreamNode->getName().c_str(), Connect.OutputName }, Connect.ParentShaderNode);
-
-		//No need to create a node, since the input and the node have the same channels, we just check if there's a value
-		if(NodeN == InputN || (NodeN < '2' && InputN < '2'))
-		{
-			AddAttributeFromValueOrInterface(Input, Connect.InputChannelName, Connect.ParentShaderNode);
-			return;
-		}
-		std::string Category;
-		TArray<FInputToCopy> InputsToCopy;
-		TArray<FInputToCreate> InputsToCreate;
-
-		// float to N
-		if(InputN < '2')
-		{
-			if(NodeN == '2')
-			{
-				Category = mx::Category::Combine2;
-				InputsToCopy.Add({ "in1", Input });
-				InputsToCopy.Add({ "in2", Input });
-			}
-			else if(NodeN == '3')
-			{
-				Category = mx::Category::Combine3;
-				InputsToCopy.Add({ "in1", Input });
-				InputsToCopy.Add({ "in2", Input });
-				InputsToCopy.Add({ "in3", Input });
-			}
-			else if(NodeN == '4')
-			{
-				Category = mx::Category::Combine4;
-				InputsToCopy.Add({ "in1", Input });
-				InputsToCopy.Add({ "in2", Input });
-				InputsToCopy.Add({ "in3", Input });
-				InputsToCopy.Add({ "in4", Input });
-			}
-		}
-		else if((InputN == '2' && NodeN == '3') || (InputN == '3' && NodeN == '4'))
-		{
-			Category = mx::Category::Combine2;
-			InputsToCopy.Add({ "in1", Input });
-			InputsToCreate.Add({ "in2", FAttributeValueArray{{"type","float"}, {"value","1"}} });
-		}
-
-		//copy "in" input into "in1" and create "amount" input as "in2" with the value 360 because UE's angle is a value between [0,1]
-		mx::NodePtr CombineNode = CreateNode(Connect.UpstreamNode->getParent(),
-											 Connect.UpstreamNode->getName().c_str(),
-											 Category.c_str(),
-											 InputsToCopy, // rename input to match <divide> input "in1"
-											 InputsToCreate);
-
-		//We set the type and an output to find it during connect node resolution phase
-		//Because the case Combine2 vector2->vector3 has no nodedef in the standard library
-		CombineNode->setType(Connect.UpstreamNode->getType());
-		if(!CombineNode->getOutput("out"))
-		{
-			CombineNode->addOutput("out", Connect.UpstreamNode->getType());
-		}
-
-		// Input now points to the new node
-		Input->setNodeName(CombineNode->getName());
-	}
-}
-
 void FMaterialXSurfaceShaderAbstract::ConnectIfGreaterInputToOutput(const FConnectNode& Connect)
 {
 	UInterchangeShaderNode* NodeIf = CreateShaderNode(Connect.UpstreamNode, Connect.UpstreamNode->getName().c_str(), UE::Expressions::Names::If);
@@ -1329,60 +1191,30 @@ void FMaterialXSurfaceShaderAbstract::ConnectBlurInputToOutput(const FConnectNod
 void FMaterialXSurfaceShaderAbstract::ConnectTexCoordInputToOutput(const FConnectNode& Connect)
 {
 	using namespace UE::Interchange::Materials::Standard::Nodes;
-	UInterchangeShaderNode* TexCoord = CreateShaderNode(Connect.UpstreamNode, Connect.UpstreamNode->getName().c_str(), TextureCoordinate::Name.ToString());
-	if(mx::InputPtr Input = Connect.UpstreamNode->getInput("index"))
-	{
-		TexCoord->AddInt32Attribute(UInterchangeShaderPortsAPI::MakeInputValueKey(TextureCoordinate::Inputs::Index.ToString()), mx::fromValueString<int>(Input->getValueString()));
-	}
-
-	UInterchangeShaderPortsAPI::ConnectDefaultOuputToInput(Connect.ParentShaderNode, Connect.InputChannelName, TexCoord->GetUniqueID());
-}
-
-void FMaterialXSurfaceShaderAbstract::ConnectSeparateInputToOutput(const FConnectNode& Connect)
-{
-	bool bHasValueOrInterface = false;
-	if(mx::InputPtr Input = Connect.UpstreamNode->getInput("in"))
-	{
-		bHasValueOrInterface = Input->hasValue() || Input->hasInterfaceName();
-	}
-
-	uint8 Mask = 0b1110;
-	int32 OutpuxIndex = 0;
+	
+	UInterchangeShaderNode* TexCoord;
+	UInterchangeShaderNode* TexCoord2D;
 
-	if(Connect.OutputName == TEXT("outx") || Connect.OutputName == TEXT("outr"))
-	{
-		Mask = 0b1000;
-		OutpuxIndex = 1;
-	}
-	else if(Connect.OutputName == TEXT("outy") || Connect.OutputName == TEXT("outg"))
-	{
-		Mask = 0b0100;
-		OutpuxIndex = 2;
-	}
-	else if(Connect.OutputName == TEXT("outz") || Connect.OutputName == TEXT("outb"))
-	{
-		Mask = 0b0010;
-		OutpuxIndex = 3;
-	}
-	else if(Connect.OutputName == TEXT("outw") || Connect.OutputName == TEXT("outa"))
+	// UE only supports 2D UVs, let's just append a default 3rd channel
+	if (Connect.UpstreamNode->getType() == mx::Type::Vector3)	
 	{
-		Mask = 0b0001;
-		OutpuxIndex = 4;
+		TexCoord2D = CreateShaderNode(Connect.UpstreamNode, (Connect.UpstreamNode->getName() + "_vector3").c_str(), TextureCoordinate::Name.ToString());
+		TexCoord = CreateShaderNode(Connect.UpstreamNode, Connect.UpstreamNode->getName().c_str(), UE::Expressions::Names::AppendVector);
+		TexCoord->AddFloatAttribute(UInterchangeShaderPortsAPI::MakeInputValueKey(UE::Expressions::Inputs::B), 0.f);
+		UInterchangeShaderPortsAPI::ConnectDefaultOuputToInput(TexCoord, UE::Expressions::Inputs::A, TexCoord2D->GetUniqueID());
 	}
 	else
 	{
-		UE_LOG(LogInterchangeImport, Warning, TEXT("output <%s> not defined in <%s>."), *Connect.OutputName, UTF8_TO_TCHAR(Connect.UpstreamNode->getCategory().c_str()));
+		TexCoord = CreateShaderNode(Connect.UpstreamNode, Connect.UpstreamNode->getName().c_str(), TextureCoordinate::Name.ToString());
+		TexCoord2D = TexCoord;
 	}
 
-	if(!bHasValueOrInterface)
-	{
-		UInterchangeShaderNode* OutNode = CreateMaskShaderNode(Mask, Connect.UpstreamNode, Connect.UpstreamNode->getName().c_str(), Connect.OutputName);
-		UInterchangeShaderPortsAPI::ConnectDefaultOuputToInput(Connect.ParentShaderNode, Connect.InputChannelName, OutNode->GetUniqueID());
-	}
-	else
+	if(mx::InputPtr Input = Connect.UpstreamNode->getInput("index"))
 	{
-		AddAttributeFromValueOrInterface(Connect.UpstreamNode->getInput("in"), Connect.InputChannelName, Connect.ParentShaderNode, OutpuxIndex);
+		TexCoord2D->AddInt32Attribute(UInterchangeShaderPortsAPI::MakeInputValueKey(TextureCoordinate::Inputs::Index.ToString()), mx::fromValueString<int>(Input->getValueString()));
 	}
+
+	UInterchangeShaderPortsAPI::ConnectDefaultOuputToInput(Connect.ParentShaderNode, Connect.InputChannelName, TexCoord->GetUniqueID());
 }
 
 void FMaterialXSurfaceShaderAbstract::ConnectSwitchInputToOutput(const FConnectNode& Connect)
@@ -1524,7 +1356,7 @@ UInterchangeShaderNode* FMaterialXSurfaceShaderAbstract::CreateShaderNode(mx::El
 {
 	UInterchangeShaderNode* Node;
 
-	const FString NodeUID = UInterchangeShaderNode::MakeNodeUid(GetUniqueName(Element) + TEXT("_") + NodeName + TEXT('_') + OutputName, FStringView{});
+	const FString NodeUID = UInterchangeShaderNode::MakeNodeUid(GetUniqueName(Element) + TEXT("_") + NodeName + TEXT('_'), FStringView{});
 
 	//Test directly in the NodeContainer, because the ShaderNodes can be altered during the node graph either by the parent (dot/normalmap),
 	//or by putting an intermediary node between the child and the parent (tiledimage)
@@ -1533,8 +1365,7 @@ UInterchangeShaderNode* FMaterialXSurfaceShaderAbstract::CreateShaderNode(mx::El
 		Node = NewObject<UInterchangeShaderNode>(&NodeContainer);
 		NodeContainer.SetupNode(Node, NodeUID, NodeName, EInterchangeNodeContainerType::TranslatedAsset);
 		Node->SetCustomShaderType(ShaderType);
-
-		ShaderNodes.Add({ NodeName, OutputName }, Node);
+		ShaderNodes.FindOrAdd({ NodeName, OutputName }, Node);
 	}
 
 	return Node;
@@ -1544,15 +1375,14 @@ UInterchangeFunctionCallShaderNode* FMaterialXSurfaceShaderAbstract::CreateFunct
 {
 	UInterchangeFunctionCallShaderNode* Node;
 
-	const FString NodeUID = UInterchangeShaderNode::MakeNodeUid(GetUniqueName(Element) + TEXT("_") + NodeName  + TEXT('_') + OutputName, FStringView{});
+	const FString NodeUID = UInterchangeShaderNode::MakeNodeUid(GetUniqueName(Element) + TEXT("_") + NodeName  + TEXT('_'), FStringView{});
 
 	if(Node = const_cast<UInterchangeFunctionCallShaderNode*>(Cast<UInterchangeFunctionCallShaderNode>(NodeContainer.GetNode(NodeUID))); !Node)
 	{
 		Node = NewObject<UInterchangeFunctionCallShaderNode>(&NodeContainer);
 		NodeContainer.SetupNode(Node, NodeUID, NodeName, EInterchangeNodeContainerType::TranslatedAsset);
 		Node->SetCustomMaterialFunction(FunctionPath);
-
-		ShaderNodes.Add({ NodeName, OutputName}, Node);
+		ShaderNodes.FindOrAdd({ NodeName, OutputName }, Node);
 	}
 
 	return Node;
@@ -1562,7 +1392,7 @@ UInterchangeFunctionCallShaderNode* FMaterialXSurfaceShaderAbstract::CreateFunct
 {
 	UInterchangeFunctionCallShaderNode* Node;
 
-	const FString NodeUID = UInterchangeShaderNode::MakeNodeUid(GetUniqueName(Element) + TEXT("_") + NodeName + TEXT('_') + OutputName, FStringView{});
+	const FString NodeUID = UInterchangeShaderNode::MakeNodeUid(GetUniqueName(Element) + TEXT("_") + NodeName + TEXT('_'), FStringView{});
 
 	if(Node = const_cast<UInterchangeFunctionCallShaderNode*>(Cast<UInterchangeFunctionCallShaderNode>(NodeContainer.GetNode(NodeUID))); !Node)
 	{
@@ -1572,8 +1402,7 @@ UInterchangeFunctionCallShaderNode* FMaterialXSurfaceShaderAbstract::CreateFunct
 		Node->SetCustomMaterialFunction(TEXT("/Game/Default.Default"));
 		Node->AddInt32Attribute(UE::Interchange::MaterialX::Attributes::EnumType, EnumType);
 		Node->AddInt32Attribute(UE::Interchange::MaterialX::Attributes::EnumValue, EnumValue);
-
-		ShaderNodes.Add({ NodeName, OutputName }, Node);
+		ShaderNodes.FindOrAdd({ NodeName, OutputName }, Node);
 	}
 
 	return Node;
@@ -1722,7 +1551,6 @@ void FMaterialXSurfaceShaderAbstract::RegisterConnectNodeOutputToInputDelegates(
 	MatchingConnectNodeDelegates.Add(mx::Category::Rotate2D,		FOnConnectNodeOutputToInput::CreateSP(this, &FMaterialXSurfaceShaderAbstract::ConnectRotate2DInputToOutput));
 	MatchingConnectNodeDelegates.Add(mx::Category::Rotate3D,		FOnConnectNodeOutputToInput::CreateSP(this, &FMaterialXSurfaceShaderAbstract::ConnectRotate3DInputToOutput));
 	MatchingConnectNodeDelegates.Add(mx::Category::Image,			FOnConnectNodeOutputToInput::CreateSP(this, &FMaterialXSurfaceShaderAbstract::ConnectImageInputToOutput));
-	MatchingConnectNodeDelegates.Add(mx::Category::Convert,			FOnConnectNodeOutputToInput::CreateSP(this, &FMaterialXSurfaceShaderAbstract::ConnectConvertInputToOutput));
 	MatchingConnectNodeDelegates.Add(mx::Category::IfGreater,		FOnConnectNodeOutputToInput::CreateSP(this, &FMaterialXSurfaceShaderAbstract::ConnectIfGreaterInputToOutput));
 	MatchingConnectNodeDelegates.Add(mx::Category::IfGreaterEq,		FOnConnectNodeOutputToInput::CreateSP(this, &FMaterialXSurfaceShaderAbstract::ConnectIfGreaterEqInputToOutput));
 	MatchingConnectNodeDelegates.Add(mx::Category::IfEqual,			FOnConnectNodeOutputToInput::CreateSP(this, &FMaterialXSurfaceShaderAbstract::ConnectIfEqualInputToOutput));
@@ -1740,9 +1568,6 @@ void FMaterialXSurfaceShaderAbstract::RegisterConnectNodeOutputToInputDelegates(
 	MatchingConnectNodeDelegates.Add(mx::Category::WorleyNoise3D,	FOnConnectNodeOutputToInput::CreateSP(this, &FMaterialXSurfaceShaderAbstract::ConnectWorleyNoise3DInputToOutput));
 	MatchingConnectNodeDelegates.Add(mx::Category::Blur,			FOnConnectNodeOutputToInput::CreateSP(this, &FMaterialXSurfaceShaderAbstract::ConnectBlurInputToOutput));
 	MatchingConnectNodeDelegates.Add(mx::Category::HeightToNormal,	FOnConnectNodeOutputToInput::CreateSP(this, &FMaterialXSurfaceShaderAbstract::ConnectHeightToNormalInputToOutput));
-	MatchingConnectNodeDelegates.Add(mx::Category::Separate2,		FOnConnectNodeOutputToInput::CreateSP(this, &FMaterialXSurfaceShaderAbstract::ConnectSeparateInputToOutput));
-	MatchingConnectNodeDelegates.Add(mx::Category::Separate3,		FOnConnectNodeOutputToInput::CreateSP(this, &FMaterialXSurfaceShaderAbstract::ConnectSeparateInputToOutput));
-	MatchingConnectNodeDelegates.Add(mx::Category::Separate4,		FOnConnectNodeOutputToInput::CreateSP(this, &FMaterialXSurfaceShaderAbstract::ConnectSeparateInputToOutput));
 	MatchingConnectNodeDelegates.Add(mx::Category::TexCoord,		FOnConnectNodeOutputToInput::CreateSP(this, &FMaterialXSurfaceShaderAbstract::ConnectTexCoordInputToOutput));
 	MatchingConnectNodeDelegates.Add(mx::Category::Switch,			FOnConnectNodeOutputToInput::CreateSP(this, &FMaterialXSurfaceShaderAbstract::ConnectSwitchInputToOutput));
 	MatchingConnectNodeDelegates.Add(mx::Category::Refract,			FOnConnectNodeOutputToInput::CreateSP(this, &FMaterialXSurfaceShaderAbstract::ConnectRefractInputToOutput));
diff --git a/Engine/Plugins/Interchange/Runtime/Source/Import/Public/MaterialX/MaterialXUtils/MaterialXManager.h b/Engine/Plugins/Interchange/Runtime/Source/Import/Public/MaterialX/MaterialXUtils/MaterialXManager.h
--- a/Engine/Plugins/Interchange/Runtime/Source/Import/Public/MaterialX/MaterialXUtils/MaterialXManager.h
+++ b/Engine/Plugins/Interchange/Runtime/Source/Import/Public/MaterialX/MaterialXUtils/MaterialXManager.h
@@ -17,6 +17,16 @@ class UInterchangeShaderNode;
 class FMaterialXSurfaceShaderAbstract;
 class UInterchangeTranslatorBase;
 
+#endif // WITH_EDITOR
+
+namespace UE::Interchange::MaterialX
+{
+	// Load necessary material functions, this function can only be called in the Game Thread
+	INTERCHANGEIMPORT_API bool AreMaterialFunctionPackagesLoaded();
+}
+
+#if WITH_EDITOR
+
 class FMaterialXManager
 {
 public:
@@ -119,9 +129,16 @@ private:
 	/** Helper function that returns the input as a string by also adding it in the MaterialExpressionInputs container*/
 	INTERCHANGEIMPORT_API const TCHAR* ExpressionInput(const TCHAR* Input);
 
+	/** Helper function that returns the material function package path as a string by also adding it in the MaterialFunctionsToLoad container,
+	  * in order to centralize material functions that need to be loaded in the Game Thread by AreMaterialFunctionPackagesLoaded function*/
+	INTERCHANGEIMPORT_API const TCHAR* MaterialFunctionPackage(const TCHAR* Input);
+
 	/** The different inputs of material expression that we may encounter, the MaterialX Document is modified consequently regarding those*/
 	TSet<FString> MaterialExpressionInputs;
 
+	/** List of material functions packages to load in the Game Thread. */
+	TSet<FString> MaterialFunctionsToLoad;
+
 	/** Given a MaterialX node (category (optionally a nodegroup) - input), return the UE/Interchange input name.*/
 	TMap<TPair<FKeyExpression, FString>, FString> MatchingInputNames;
 
@@ -152,11 +169,7 @@ private:
 	std::unordered_map<std::string, std::vector<std::string>> NodeInputsToRemove;
 
 	bool bIsSubstrateEnabled;
-};
-#endif
 
-namespace UE::Interchange::MaterialX
-{
-	// Load necessary material functions, this function can only be called in the Game Thread
-	INTERCHANGEIMPORT_API bool AreMaterialFunctionPackagesLoaded();
-}
+	friend bool UE::Interchange::MaterialX::AreMaterialFunctionPackagesLoaded();
+};
+#endif // WITH_EDITOR
\ No newline at end of file
diff --git a/Engine/Plugins/Interchange/Runtime/Source/Import/Public/MaterialX/MaterialXUtils/MaterialXSurfaceShaderAbstract.h b/Engine/Plugins/Interchange/Runtime/Source/Import/Public/MaterialX/MaterialXUtils/MaterialXSurfaceShaderAbstract.h
--- a/Engine/Plugins/Interchange/Runtime/Source/Import/Public/MaterialX/MaterialXUtils/MaterialXSurfaceShaderAbstract.h
+++ b/Engine/Plugins/Interchange/Runtime/Source/Import/Public/MaterialX/MaterialXUtils/MaterialXSurfaceShaderAbstract.h
@@ -43,7 +43,9 @@ protected:
 		// The input of the ParentShaderNode to connect to.
 		FString InputChannelName;
 
-		// The output name of the MaterialX node. The default name is 'out' as stated by the standard library.
+		// The output name of the MaterialX node that we need to connect to.
+		// This is the active output (e.g: the output given by an input connecting 2 nodes)
+		// The default name is 'out' as stated by the standard library. (not necessarily always the case though)
 		FString OutputName{ DefaultOutput };
 	};
 		
@@ -252,9 +254,6 @@ protected:
 	/** <image> */
 	void ConnectImageInputToOutput(const FConnectNode& Connect);
 
-	/** <convert> */
-	void ConnectConvertInputToOutput(const FConnectNode& Connect);
-
 	/** <ifgreater> */
 	void ConnectIfGreaterInputToOutput(const FConnectNode& Connect);
 
@@ -309,9 +308,6 @@ protected:
 	/** <texcoord> */
 	void ConnectTexCoordInputToOutput(const FConnectNode& Connect);
 
-	/** <separate2/3/4> */
-	void ConnectSeparateInputToOutput(const FConnectNode& Connect);
-
 	/** <switch> */
 	void ConnectSwitchInputToOutput(const FConnectNode& Connect);
 
diff --git a/Engine/Plugins/Interchange/Runtime/Source/Pipelines/Private/InterchangeMaterialXPipeline.cpp b/Engine/Plugins/Interchange/Runtime/Source/Pipelines/Private/InterchangeMaterialXPipeline.cpp
--- a/Engine/Plugins/Interchange/Runtime/Source/Pipelines/Private/InterchangeMaterialXPipeline.cpp
+++ b/Engine/Plugins/Interchange/Runtime/Source/Pipelines/Private/InterchangeMaterialXPipeline.cpp
@@ -42,9 +42,9 @@ namespace
 	MATERIALX_FUNCTIONS_SUBSTRATE_PATH(ConductorBSDF);
 	MATERIALX_FUNCTIONS_SUBSTRATE_PATH(SheenBSDF);
 	MATERIALX_FUNCTIONS_SUBSTRATE_PATH(SubsurfaceBSDF);
-	MATERIALX_FUNCTIONS_SUBSTRATE_PATH(ThinFilmBSDF);
 	MATERIALX_FUNCTIONS_SUBSTRATE_PATH(GeneralizedSchlickBSDF);
 	MATERIALX_FUNCTIONS_SUBSTRATE_PATH(TranslucentBSDF);
+	MATERIALX_FUNCTIONS_SUBSTRATE_PATH(ChiangHairBSDF);
 	
 	MATERIALX_FUNCTIONS_SUBSTRATE_PATH(UniformEDF);
 	MATERIALX_FUNCTIONS_SUBSTRATE_PATH(ConicalEDF);
@@ -415,7 +415,10 @@ UMaterialXPipelineSettings::UMaterialXPipelineSettings()
 						TEXT("scale"),
 					},
 					// Displacement Outputs
-					TSet<FName>{TEXT("Displacement")}
+					TSet<FName>{
+						TEXT("Displacement"),
+						TEXT("Normal"), // Not part of the spec, but we have to recompute them
+					}
 				}
 			},
 			// BSDF
@@ -427,7 +430,8 @@ UMaterialXPipelineSettings::UMaterialXPipelineSettings()
 						TEXT("weight"),
 						TEXT("color"),
 						TEXT("roughness"),
-						TEXT("normal")
+						TEXT("normal"),
+						TEXT("energy_compensation"),
 					},
 					// OrenNayarDiffuse Outputs
 					TSet<FName>{
@@ -475,8 +479,12 @@ UMaterialXPipelineSettings::UMaterialXPipelineSettings()
 						TEXT("tint"),
 						TEXT("ior"),
 						TEXT("roughness"),
+						TEXT("thinfilm_thickness"),
+						TEXT("thinfilm_ior"),
 						TEXT("normal"),
-						TEXT("tangent")
+						TEXT("tangent"),
+						TEXT("distribution"),
+						TEXT("scatter_mode"),
 					},
 					// Dielectric Outputs
 					TSet<FName>{
@@ -539,30 +547,39 @@ UMaterialXPipelineSettings::UMaterialXPipelineSettings()
 				}
 			},
 			{
-				UMaterialXPipelineSettings::ToEnumKey(EInterchangeMaterialXBSDF::Sheen),
+				UMaterialXPipelineSettings::ToEnumKey(EInterchangeMaterialXBSDF::ChiangHair),
 				{
-					// Sheen Inputs
+					// ChiangHair Inputs
 					TSet<FName>{
-						TEXT("weight"),
-						TEXT("color"),
-						TEXT("roughness"),
-						TEXT("normal")
+						TEXT("tint_R"),
+						TEXT("tint_TT"),
+						TEXT("tint_TRT"),
+						TEXT("ior"),
+						TEXT("roughness_R"),
+						TEXT("roughness_TT"),
+						TEXT("roughness_TRT"),
+						TEXT("cuticle_angle"),
+						TEXT("absorption_coefficient"),
+						TEXT("normal"),
+						TEXT("curve_direction"),
 					},
-					// Sheen Outputs
+					// ChiangHair Outputs
 					TSet<FName>{
 						TEXT("Output")
 					}
 				}
 			},
 			{
-				UMaterialXPipelineSettings::ToEnumKey(EInterchangeMaterialXBSDF::ThinFilm),
+				UMaterialXPipelineSettings::ToEnumKey(EInterchangeMaterialXBSDF::Sheen),
 				{
-					// ThinFilm Inputs
+					// Sheen Inputs
 					TSet<FName>{
-						TEXT("thickness"),
-						TEXT("ior")
+						TEXT("weight"),
+						TEXT("color"),
+						TEXT("roughness"),
+						TEXT("normal")
 					},
-					// ThinFilm Outputs
+					// Sheen Outputs
 					TSet<FName>{
 						TEXT("Output")
 					}
@@ -669,9 +686,9 @@ UInterchangeMaterialXPipeline::UInterchangeMaterialXPipeline()
 			{MATERIALX_MATERIALFUNCTION_PATH(ConductorBSDF),			  UMaterialXPipelineSettings::ToEnumKey(EInterchangeMaterialXBSDF::Conductor)},
 			{MATERIALX_MATERIALFUNCTION_PATH(SheenBSDF),				  UMaterialXPipelineSettings::ToEnumKey(EInterchangeMaterialXBSDF::Sheen)},
 			{MATERIALX_MATERIALFUNCTION_PATH(SubsurfaceBSDF),			  UMaterialXPipelineSettings::ToEnumKey(EInterchangeMaterialXBSDF::Subsurface)},
-			{MATERIALX_MATERIALFUNCTION_PATH(ThinFilmBSDF),				  UMaterialXPipelineSettings::ToEnumKey(EInterchangeMaterialXBSDF::ThinFilm)},
 			{MATERIALX_MATERIALFUNCTION_PATH(GeneralizedSchlickBSDF),	  UMaterialXPipelineSettings::ToEnumKey(EInterchangeMaterialXBSDF::GeneralizedSchlick)},
 			{MATERIALX_MATERIALFUNCTION_PATH(TranslucentBSDF),			  UMaterialXPipelineSettings::ToEnumKey(EInterchangeMaterialXBSDF::Translucent)},
+			{MATERIALX_MATERIALFUNCTION_PATH(ChiangHairBSDF),			  UMaterialXPipelineSettings::ToEnumKey(EInterchangeMaterialXBSDF::ChiangHair)},
 																		  
 			{MATERIALX_MATERIALFUNCTION_PATH(UniformEDF),				  UMaterialXPipelineSettings::ToEnumKey(EInterchangeMaterialXEDF::Uniform)},
 			{MATERIALX_MATERIALFUNCTION_PATH(ConicalEDF),				  UMaterialXPipelineSettings::ToEnumKey(EInterchangeMaterialXEDF::Conical)},
@@ -843,9 +860,9 @@ void UMaterialXPipelineSettings::InitPredefinedAssets()
 			{UMaterialXPipelineSettings::ToEnumKey(EInterchangeMaterialXBSDF::Conductor), ConductorBSDFFunctionsPath, ConductorBSDFSubstratePath},
 			{UMaterialXPipelineSettings::ToEnumKey(EInterchangeMaterialXBSDF::Sheen), SheenBSDFFunctionsPath, SheenBSDFSubstratePath},
 			{UMaterialXPipelineSettings::ToEnumKey(EInterchangeMaterialXBSDF::Subsurface), SubsurfaceBSDFFunctionsPath, SubsurfaceBSDFSubstratePath},
-			{UMaterialXPipelineSettings::ToEnumKey(EInterchangeMaterialXBSDF::ThinFilm), ThinFilmBSDFFunctionsPath, ThinFilmBSDFSubstratePath},
 			{UMaterialXPipelineSettings::ToEnumKey(EInterchangeMaterialXBSDF::GeneralizedSchlick), GeneralizedSchlickBSDFFunctionsPath, GeneralizedSchlickBSDFSubstratePath},
 			{UMaterialXPipelineSettings::ToEnumKey(EInterchangeMaterialXBSDF::Translucent), TranslucentBSDFFunctionsPath, TranslucentBSDFSubstratePath},
+			{UMaterialXPipelineSettings::ToEnumKey(EInterchangeMaterialXBSDF::ChiangHair), ChiangHairBSDFFunctionsPath, ChiangHairBSDFSubstratePath},
 
 			{UMaterialXPipelineSettings::ToEnumKey(EInterchangeMaterialXEDF::Uniform), UniformEDFFunctionsPath, UniformEDFSubstratePath},
 			{UMaterialXPipelineSettings::ToEnumKey(EInterchangeMaterialXEDF::Conical), ConicalEDFFunctionsPath, ConicalEDFSubstratePath},
