Interchange MaterialX:
- Add conditional nodes <and> <or> <xor> <not>

#virtualized

diff --git a/Engine/Plugins/Interchange/Runtime/Source/Common/Public/MaterialX/InterchangeMaterialXDefinitions.h b/Engine/Plugins/Interchange/Runtime/Source/Common/Public/MaterialX/InterchangeMaterialXDefinitions.h
--- a/Engine/Plugins/Interchange/Runtime/Source/Common/Public/MaterialX/InterchangeMaterialXDefinitions.h
+++ b/Engine/Plugins/Interchange/Runtime/Source/Common/Public/MaterialX/InterchangeMaterialXDefinitions.h
@@ -307,6 +307,7 @@ namespace UE
 			constexpr const TCHAR* MxACEScgToRec709 = TEXT("/InterchangeAssets/Functions/MX_ACEScgToRec709.MX_ACEScgToRec709");
 			constexpr const TCHAR* MxAdobeRGBToRec709 = TEXT("/InterchangeAssets/Functions/MX_AdobeRGBToRec709.MX_AdobeRGBToRec709");
 			constexpr const TCHAR* MxArtisticIOR = TEXT("/InterchangeAssets/Functions/MX_Artistic_IOR.MX_Artistic_IOR");
+			constexpr const TCHAR* MxAnd = TEXT("/InterchangeAssets/Functions/MX_And.MX_And");
 			constexpr const TCHAR* MxCircle = TEXT("/InterchangeAssets/Functions/MX_Circle.MX_Circle");
 			constexpr const TCHAR* MxCheckerboard = TEXT("/InterchangeAssets/Functions/MX_Checkerboard.MX_Checkerboard");
 			constexpr const TCHAR* MxChiangHairAbsorptionFromColor = TEXT("/InterchangeAssets/Functions/MX_ChiangHairAbsorptionFromColor.MX_ChiangHairAbsorptionFromColor");
@@ -318,6 +319,8 @@ namespace UE
 			constexpr const TCHAR* MxLine = TEXT("/InterchangeAssets/Functions/MX_Line.MX_Line");
 			constexpr const TCHAR* MxLinearAdobeRGBToRec709 = TEXT("/InterchangeAssets/Functions/MX_LinearAdobeRGBToRec709.MX_LinearAdobeRGBToRec709");
 			constexpr const TCHAR* MxLinearDisplayP3ToRec709 = TEXT("/InterchangeAssets/Functions/MX_LinearDisplayP3ToRec709.MX_LinearDisplayP3ToRec709");
+			constexpr const TCHAR* MxOr = TEXT("/InterchangeAssets/Functions/MX_Or.MX_Or");
+			constexpr const TCHAR* MxNot = TEXT("/InterchangeAssets/Functions/MX_Not.MX_Not");
 			constexpr const TCHAR* MxPlace2D = TEXT("/InterchangeAssets/Functions/MX_Place2D.MX_Place2D");
 			constexpr const TCHAR* MxRamp = TEXT("/InterchangeAssets/Functions/MX_Ramp.MX_Ramp");
 			constexpr const TCHAR* MxRampGradient = TEXT("/InterchangeAssets/Functions/MX_RampGradient.MX_RampGradient");
@@ -329,6 +332,7 @@ namespace UE
 			constexpr const TCHAR* MxSrgbToRec709 = TEXT("/InterchangeAssets/Functions/MX_SrgbToRec709.MX_SrgbToRec709");
 			constexpr const TCHAR* MxUnifiedNoise3D = TEXT("/InterchangeAssets/Functions/MX_UnifiedNoise3d.MX_UnifiedNoise3d");
 			constexpr const TCHAR* MxTiledCircles = TEXT("/InterchangeAssets/Functions/MX_TiledCircles.MX_TiledCircles");
+			constexpr const TCHAR* MxXor = TEXT("/InterchangeAssets/Functions/MX_Xor.MX_Xor");
 			constexpr const TCHAR* NormalFromHeightMap = TEXT("/Engine/Functions/Engine_MaterialFunctions03/Procedurals/NormalFromHeightmap.NormalFromHeightmap");
 			constexpr const TCHAR* Refract = TEXT("/Engine/Functions/Engine_MaterialFunctions01/Vectors/Refract.Refract");
 		}
@@ -733,10 +737,14 @@ MATERIALX_NAMESPACE_BEGIN
 		static constexpr const char* Screen = "screen";
 		static constexpr const char* Unpremult = "unpremult";
 		// Conditional nodes
+		static constexpr const char* And = "and";
 		static constexpr const char* IfGreater = "ifgreater";
 		static constexpr const char* IfGreaterEq = "ifgreatereq";
 		static constexpr const char* IfEqual = "ifequal";
+		static constexpr const char* Not = "not";
+		static constexpr const char* Or = "or";
 		static constexpr const char* Switch = "switch";
+		static constexpr const char* Xor = "xor";
 		// Channel nodes
 		static constexpr const char* Convert = "convert";
 		static constexpr const char* Combine2 = "combine2";
diff --git a/Engine/Plugins/Interchange/Runtime/Source/Import/Private/MaterialX/MaterialXUtils/MaterialXManager.cpp b/Engine/Plugins/Interchange/Runtime/Source/Import/Private/MaterialX/MaterialXUtils/MaterialXManager.cpp
--- a/Engine/Plugins/Interchange/Runtime/Source/Import/Private/MaterialX/MaterialXUtils/MaterialXManager.cpp
+++ b/Engine/Plugins/Interchange/Runtime/Source/Import/Private/MaterialX/MaterialXUtils/MaterialXManager.cpp
@@ -321,7 +321,8 @@ FMaterialXManager::FMaterialXManager()
 		mx::Category::TiledCircles,
 		mx::Category::Time, // since 1.39.3, they added a nodegraph in stdlib/genglsl/stdlib_genglsl_impl.mtlx instead of stdlib_ng.mtlx
 		mx::Category::UnifiedNoise3D,
-		mx::Category::UsdPreviewSurface
+		mx::Category::UsdPreviewSurface,
+		mx::Category::Xor,
 	}
 	, NodeDefsCategories{
 		mx::Category::Add,
@@ -344,6 +345,12 @@ FMaterialXManager::FMaterialXManager()
 	}
 	, bIsSubstrateEnabled{ IInterchangeImportModule::IsAvailable() ? IInterchangeImportModule::Get().IsSubstrateEnabled() : false }
 {
+	// Load in the Game thread
+	MaterialFunctionPackage(UE::MaterialFunctions::Path::MxAnd);
+	MaterialFunctionPackage(UE::MaterialFunctions::Path::MxOr);
+	MaterialFunctionPackage(UE::MaterialFunctions::Path::MxNot);
+	MaterialFunctionPackage(UE::MaterialFunctions::Path::MxXor);
+
 	if (bIsSubstrateEnabled)
 	{
 		// Vertical Layering
diff --git a/Engine/Plugins/Interchange/Runtime/Source/Import/Private/MaterialX/MaterialXUtils/MaterialXSurfaceShaderAbstract.cpp b/Engine/Plugins/Interchange/Runtime/Source/Import/Private/MaterialX/MaterialXUtils/MaterialXSurfaceShaderAbstract.cpp
--- a/Engine/Plugins/Interchange/Runtime/Source/Import/Private/MaterialX/MaterialXUtils/MaterialXSurfaceShaderAbstract.cpp
+++ b/Engine/Plugins/Interchange/Runtime/Source/Import/Private/MaterialX/MaterialXUtils/MaterialXSurfaceShaderAbstract.cpp
@@ -31,14 +31,14 @@ FMaterialXSurfaceShaderAbstract::FMaterialXSurfaceShaderAbstract(UInterchangeBas
 	, bTangentSpaceInput{ false }
 {}
 
-bool FMaterialXSurfaceShaderAbstract::AddAttribute(MaterialX::InputPtr Input, const FString& InputChannelName, UInterchangeShaderNode* ShaderNode, int32 OutputIndex)
+bool FMaterialXSurfaceShaderAbstract::AddAttribute(MaterialX::InputPtr Input, const FString& InputChannelName, UInterchangeShaderNode* ShaderNode, int32 OutputIndex, bool bBoolAsScalar)
 {
 	using namespace UE::Interchange::Materials::Standard::Nodes;
 	if(Input)
 	{
 		if(Input->getType() == mx::Type::Boolean)
 		{
-			return AddBooleanAttribute(Input, InputChannelName, ShaderNode);
+			return AddBooleanAttribute(Input, InputChannelName, ShaderNode, bBoolAsScalar);
 		}
 		else if(Input->getType() == mx::Type::Float)
 		{
@@ -65,7 +65,7 @@ bool FMaterialXSurfaceShaderAbstract::AddAttribute(MaterialX::InputPtr Input, co
 	return false;
 }
 
-bool FMaterialXSurfaceShaderAbstract::AddAttributeFromValueOrInterface(MaterialX::InputPtr Input, const FString& InputChannelName, UInterchangeShaderNode* ShaderNode, int32 OutputIndex)
+bool FMaterialXSurfaceShaderAbstract::AddAttributeFromValueOrInterface(MaterialX::InputPtr Input, const FString& InputChannelName, UInterchangeShaderNode* ShaderNode, int32 OutputIndex, bool bBoolAsScalar)
 {
 	bool bAttribute = false;
 
@@ -76,13 +76,13 @@ bool FMaterialXSurfaceShaderAbstract::AddAttributeFromValueOrInterface(MaterialX
 
 		if(Input->hasValue())
 		{
-			bAttribute = AddAttribute(Input, InputToConnectTo, ShaderNodeToConnectTo, OutputIndex);
+			bAttribute = AddAttribute(Input, InputToConnectTo, ShaderNodeToConnectTo, OutputIndex, bBoolAsScalar);
 		}
 		else if(Input->hasInterfaceName())
 		{
 			if(mx::InputPtr InputInterface = Input->getInterfaceInput(); InputInterface->hasValue())
 			{
-				bAttribute = AddAttribute(InputInterface, InputToConnectTo, ShaderNodeToConnectTo, OutputIndex);
+				bAttribute = AddAttribute(InputInterface, InputToConnectTo, ShaderNodeToConnectTo, OutputIndex, bBoolAsScalar);
 			}
 		}
 	}
@@ -90,9 +90,10 @@ bool FMaterialXSurfaceShaderAbstract::AddAttributeFromValueOrInterface(MaterialX
 	return bAttribute;
 }
 
-bool FMaterialXSurfaceShaderAbstract::AddBooleanAttribute(MaterialX::InputPtr Input, const FString& InputChannelName, UInterchangeShaderNode* ShaderNode)
+bool FMaterialXSurfaceShaderAbstract::AddBooleanAttribute(MaterialX::InputPtr Input, const FString& InputChannelName, UInterchangeShaderNode* ShaderNode, bool bAsScalar)
 {
 	using namespace UE::Interchange::Materials::Standard::Nodes;
+
 	if(Input && Input->hasValue())
 	{
 		bool Value = mx::fromValueString<bool>(Input->getValueString());
@@ -103,7 +104,16 @@ bool FMaterialXSurfaceShaderAbstract::AddBooleanAttribute(MaterialX::InputPtr In
 		NodeName += Input->getName().c_str();
 
 		UInterchangeShaderNode* StaticBoolParameterNode = CreateShaderNode(Input, NodeName, StaticBoolParameter::Name.ToString());
-		StaticBoolParameterNode->AddBooleanAttribute(UInterchangeShaderPortsAPI::MakeInputParameterKey(StaticBoolParameter::Attributes::DefaultValue.ToString()), Value);
+		if (bAsScalar)
+		{
+			StaticBoolParameterNode->SetCustomShaderType(ScalarParameter::Name.ToString());
+			StaticBoolParameterNode->AddFloatAttribute(UInterchangeShaderPortsAPI::MakeInputParameterKey(ScalarParameter::Attributes::DefaultValue.ToString()), Value);
+		}
+		else
+		{
+			StaticBoolParameterNode->AddBooleanAttribute(UInterchangeShaderPortsAPI::MakeInputParameterKey(StaticBoolParameter::Attributes::DefaultValue.ToString()), Value);
+		}
+
 		return UInterchangeShaderPortsAPI::ConnectDefaultOuputToInput(ShaderNode, InputChannelName, StaticBoolParameterNode->GetUniqueID());
 	}
 
@@ -1336,6 +1346,36 @@ void FMaterialXSurfaceShaderAbstract::ConnectViewDirectionInputToOutput(const FC
 	UInterchangeShaderPortsAPI::ConnectDefaultOuputToInput(Connect.ParentShaderNode, Connect.InputChannelName, NodeToConnectTo->GetUniqueID());
 }
 
+void FMaterialXSurfaceShaderAbstract::ConnectLogicalInputToOutput(const FConnectNode& Connect, const FString& LogicalMaterialFunction)
+{
+	UInterchangeShaderNode* LogicalNode = CreateFunctionCallShaderNode(Connect.UpstreamNode, Connect.UpstreamNode->getName().c_str(), LogicalMaterialFunction);
+	for(mx::InputPtr Input : Connect.UpstreamNode->getInputs())
+	{
+		AddAttributeFromValueOrInterface(Input, GetInputName(Input), LogicalNode, 0, true);
+	}
+	UInterchangeShaderPortsAPI::ConnectDefaultOuputToInput(Connect.ParentShaderNode, Connect.InputChannelName, LogicalNode->GetUniqueID());
+}
+
+void FMaterialXSurfaceShaderAbstract::ConnectAndInputToOutput(const FConnectNode& Connect)
+{
+	ConnectLogicalInputToOutput(Connect, UE::MaterialFunctions::Path::MxAnd);
+}
+
+void FMaterialXSurfaceShaderAbstract::ConnectOrInputToOutput(const FConnectNode& Connect)
+{
+	ConnectLogicalInputToOutput(Connect, UE::MaterialFunctions::Path::MxOr);
+}
+
+void FMaterialXSurfaceShaderAbstract::ConnectXorInputToOutput(const FConnectNode & Connect)
+{
+	ConnectLogicalInputToOutput(Connect, UE::MaterialFunctions::Path::MxXor);
+}
+
+void FMaterialXSurfaceShaderAbstract::ConnectNotInputToOutput(const FConnectNode & Connect)
+{
+	ConnectLogicalInputToOutput(Connect, UE::MaterialFunctions::Path::MxNot);
+}
+
 UInterchangeShaderNode* FMaterialXSurfaceShaderAbstract::CreateMaskShaderNode(uint8 RGBA, mx::ElementPtr Element, const FString& NodeName, const FString& OutputName)
 {
 	bool bR = (0b1000 & RGBA) >> 3;
@@ -1436,7 +1476,7 @@ const FString& FMaterialXSurfaceShaderAbstract::GetMatchedInputName(MaterialX::N
 
 FString FMaterialXSurfaceShaderAbstract::GetInputName(MaterialX::InputPtr Input) const
 {
-	if(Input->hasAttribute(mx::Attributes::NewName))
+	if (Input->hasAttribute(mx::Attributes::NewName))
 	{
 		return Input->getAttribute(mx::Attributes::NewName).c_str();
 	}
@@ -1572,6 +1612,10 @@ void FMaterialXSurfaceShaderAbstract::RegisterConnectNodeOutputToInputDelegates(
 	MatchingConnectNodeDelegates.Add(mx::Category::Switch,			FOnConnectNodeOutputToInput::CreateSP(this, &FMaterialXSurfaceShaderAbstract::ConnectSwitchInputToOutput));
 	MatchingConnectNodeDelegates.Add(mx::Category::Refract,			FOnConnectNodeOutputToInput::CreateSP(this, &FMaterialXSurfaceShaderAbstract::ConnectRefractInputToOutput));
 	MatchingConnectNodeDelegates.Add(mx::Category::ViewDirection,	FOnConnectNodeOutputToInput::CreateSP(this, &FMaterialXSurfaceShaderAbstract::ConnectViewDirectionInputToOutput));
+	MatchingConnectNodeDelegates.Add(mx::Category::And,				FOnConnectNodeOutputToInput::CreateSP(this, &FMaterialXSurfaceShaderAbstract::ConnectAndInputToOutput));
+	MatchingConnectNodeDelegates.Add(mx::Category::Or,				FOnConnectNodeOutputToInput::CreateSP(this, &FMaterialXSurfaceShaderAbstract::ConnectOrInputToOutput));
+	MatchingConnectNodeDelegates.Add(mx::Category::Xor,				FOnConnectNodeOutputToInput::CreateSP(this, &FMaterialXSurfaceShaderAbstract::ConnectXorInputToOutput));
+	MatchingConnectNodeDelegates.Add(mx::Category::Not,				FOnConnectNodeOutputToInput::CreateSP(this, &FMaterialXSurfaceShaderAbstract::ConnectNotInputToOutput));
 }
 
 void FMaterialXSurfaceShaderAbstract::SetMatchingInputsNames(MaterialX::NodePtr Node) const
diff --git a/Engine/Plugins/Interchange/Runtime/Source/Import/Public/MaterialX/MaterialXUtils/MaterialXSurfaceShaderAbstract.h b/Engine/Plugins/Interchange/Runtime/Source/Import/Public/MaterialX/MaterialXUtils/MaterialXSurfaceShaderAbstract.h
--- a/Engine/Plugins/Interchange/Runtime/Source/Import/Public/MaterialX/MaterialXUtils/MaterialXSurfaceShaderAbstract.h
+++ b/Engine/Plugins/Interchange/Runtime/Source/Import/Public/MaterialX/MaterialXUtils/MaterialXSurfaceShaderAbstract.h
@@ -59,10 +59,11 @@ protected:
 	 * @param Input - The MaterialX input to retrieve and add the value from. Must be of type float/color/vector.
 	 * @param InputChannelName - The name of the shader node's input to add the attribute.
 	 * @param ShaderNode - The shader node to which we want to add the attribute.
+	 * @param bBoolAsScalar - Create a scalar instead of a static bool parameter when applicable.
 	 *
 	 * @return true if the attribute was successfully added.
 	 */
-	bool AddAttribute(MaterialX::InputPtr Input, const FString& InputChannelName, UInterchangeShaderNode* ShaderNode, int32 OutputIndex = 0);
+	bool AddAttribute(MaterialX::InputPtr Input, const FString& InputChannelName, UInterchangeShaderNode* ShaderNode, int32 OutputIndex = 0, bool bBoolAsScalar = false);
 
 	/**
 	 * Add an attribute to a shader node from the given MaterialX input if that input has either a value or an interface name.
@@ -70,10 +71,11 @@ protected:
 	 * @param Input - The MaterialX input to retrieve and add the value from. Must be of type float/color/vector.
 	 * @param InputChannelName - The name of the shader node's input to add the attribute.
 	 * @param ShaderNode - The shader node to which we want to add the attribute.
+	 * @param bBoolAsScalar - Create a scalar instead of a static bool parameter when applicable.
 	 *
 	 * @return true if the attribute was successfully added.
 	 */
-	bool AddAttributeFromValueOrInterface(MaterialX::InputPtr Input, const FString& InputChannelName, UInterchangeShaderNode* ShaderNode, int32 OutputIndex = 0);
+	bool AddAttributeFromValueOrInterface(MaterialX::InputPtr Input, const FString& InputChannelName, UInterchangeShaderNode* ShaderNode, int32 OutputIndex = 0, bool bBoolAsScalar = false);
 
 	/**
 	 * Add a bool attribute to a shader node only if its value taken from the input is not equal to its default value. Return false if the attribute does not exist or if we cannot add it
@@ -82,10 +84,11 @@ protected:
 	 * @param InputChannelName - The name of the shader node's input to add the attribute
 	 * @param ShaderNode - The shader node to which we want to add the attribute
 	 * @param DefaultValue - the default value to test the input against
+	 * @param bAsScalar - Create a scalar value instead of a static bool parameter
 	 *
 	 * @return true if the attribute was successfully added
 	 */
-	bool AddBooleanAttribute(MaterialX::InputPtr Input, const FString& InputChannelName, UInterchangeShaderNode* ShaderNode);
+	bool AddBooleanAttribute(MaterialX::InputPtr Input, const FString& InputChannelName, UInterchangeShaderNode* ShaderNode, bool bAsScalar = false);
 
 	/**
 	 * Add a float attribute to a shader node only if its value taken from the input is not equal to its default value. Return false if the attribute does not exist or if we cannot add it
@@ -320,6 +323,18 @@ protected:
 	/** <viewdirection> */
 	void ConnectViewDirectionInputToOutput(const FConnectNode& Connect);
 
+	/** <and> */
+	void ConnectAndInputToOutput(const FConnectNode& Connect);
+
+	/** <or> */
+	void ConnectOrInputToOutput(const FConnectNode& Connect);
+
+	/** <xor> */
+	void ConnectXorInputToOutput(const FConnectNode& Connect);
+
+	/** <not> */
+	void ConnectNotInputToOutput(const FConnectNode& Connect);
+
 	/** End Connect MaterialX nodes*/
 
 	/**
@@ -558,6 +573,7 @@ private:
 	
 	UInterchangeShaderNode* ConnectGeometryInputToOutput(const FConnectNode& Connect, const FString& ShaderType, const FString& TransformShaderType, const FString& TransformInput, const FString& TransformSourceType, int32 TransformSource, const FString& TransformType, int32 TransformSDestination, bool bIsVector = true);
 	UInterchangeShaderNode* ConnectNoise2DInputToOutput(const FConnectNode& Connect, const FString& ShaderType, EVectorNoiseFunction NoiseFunction, uint8 Mask);
+	void ConnectLogicalInputToOutput(const FConnectNode& Connect, const FString& LogicalMaterialFunction);
 protected:
 
 	/** 
