[Interchange-FBX] Crash fix for AddAllMeshes.

#rnx

diff --git a/Engine/Plugins/Interchange/Runtime/Source/Parsers/Fbx/Private/FbxMesh.cpp b/Engine/Plugins/Interchange/Runtime/Source/Parsers/Fbx/Private/FbxMesh.cpp
--- a/Engine/Plugins/Interchange/Runtime/Source/Parsers/Fbx/Private/FbxMesh.cpp
+++ b/Engine/Plugins/Interchange/Runtime/Source/Parsers/Fbx/Private/FbxMesh.cpp
@@ -1863,6 +1863,11 @@ void FFbxMesh::AddAllMeshes(FbxScene* SDKScene, FbxGeometryConverter* SDKGeometr
 				}
 				//Find which morph target should we use according to the weight.
 				const int32 CurrentChannelMorphTargetCount = Channel->GetTargetShapeCount();
+				if (CurrentChannelMorphTargetCount == 0)
+				{
+					continue;
+				}
+
 				FString ChannelName = Parser.GetFbxHelper()->GetFbxObjectName(Channel);
 				// Maya adds the name of the MorphTarget and an underscore to the front of the channel name, so remove it
 				if (ChannelName.StartsWith(MorphTargetName) && (ChannelName.Compare(MorphTargetName, ESearchCase::IgnoreCase) != 0))
