Interchange:
- Fix material instances not updated when static switch parameters are modified after a reimport

diff --git a/Engine/Plugins/Interchange/Runtime/Source/Import/Private/Material/InterchangeMaterialFactory.cpp b/Engine/Plugins/Interchange/Runtime/Source/Import/Private/Material/InterchangeMaterialFactory.cpp
--- a/Engine/Plugins/Interchange/Runtime/Source/Import/Private/Material/InterchangeMaterialFactory.cpp
+++ b/Engine/Plugins/Interchange/Runtime/Source/Import/Private/Material/InterchangeMaterialFactory.cpp
@@ -1008,6 +1008,7 @@ UInterchangeFactoryBase::FImportAssetResult UInterchangeMaterialFactory::BeginIm
 							if(ParentPath == PreviousParentPath)
 							{
 								SetupReimportedMaterialInstance(*MaterialInstance, *Arguments.NodeContainer, *MaterialInstanceFactoryNode, *PreviousNode);
+								MaterialInstance->UpdateStaticPermutation();
 								bApplyPipelineProperties = false;
 							}
 						}
