Fixed corrupted skeletal mesh

#rnx

diff --git a/Engine/Plugins/Interchange/Runtime/Source/Pipelines/Private/InterchangeGenericSkeletalMeshPipeline.cpp b/Engine/Plugins/Interchange/Runtime/Source/Pipelines/Private/InterchangeGenericSkeletalMeshPipeline.cpp
--- a/Engine/Plugins/Interchange/Runtime/Source/Pipelines/Private/InterchangeGenericSkeletalMeshPipeline.cpp
+++ b/Engine/Plugins/Interchange/Runtime/Source/Pipelines/Private/InterchangeGenericSkeletalMeshPipeline.cpp
@@ -223,13 +223,6 @@ void UInterchangeGenericMeshPipeline::ExecutePreImportPipelineSkeletalMesh()
 						}
 					}
 
-					// sort the MeshUids in order to ensure to have the exact same name when we create the different asset
-					// because the asset name is taken by the 1st scene node found, see CreateSkeletalMeshFactoryNode
-					for (TPair<int32, TArray<FString>>& Pair : MeshUidsPerLodIndex)
-					{
-						Pair.Value.Sort();
-					}
-
 					FinishSkeletalMeshSetup(bUseInstanceMesh, bFoundInstances, SkeletonRootUid, SkeletonFactoryNode, MeshUidsPerLodIndex, MeshUids);
 				}
 				return bFoundInstances;
@@ -343,44 +336,43 @@ UInterchangeSkeletalMeshFactoryNode* UInterchangeGenericMeshPipeline::CreateSkel
 		return nullptr;
 	}
 	
-	auto GetFirstNodeInfo = [this, &MeshUidsPerLodIndex](const int32 Index, FString& OutFirstMeshNodeUid, int32& OutSceneNodeCount)->const UInterchangeBaseNode*
+	auto GetFirstNodeInfo = [this, &MeshUidsPerLodIndex](const int32 LodIndex, FString& OutFirstMeshNodeUid, int32& OutSceneNodeCount)->const UInterchangeBaseNode*
 	{
 		OutSceneNodeCount = 0;
-		if (!ensure(Index >= 0 && MeshUidsPerLodIndex.Num() > Index))
-		{
-			//Log an error
-			return nullptr;
-		}
-		for (const TPair<int32, TArray<FString>>& LodIndexAndMeshUids : MeshUidsPerLodIndex)
+		const TArray<FString>* MeshUids = MeshUidsPerLodIndex.Find(LodIndex);
+		if (ensure(MeshUids))
 		{
-			if (Index == LodIndexAndMeshUids.Key)
+			FString FirstNodeUid = MeshUids->Num() > 0 ? (*MeshUids)[0] : FString();
+			for (int32 Index = 1; Index < MeshUids->Num(); ++Index)
 			{
-				const TArray<FString>& MeshUids = LodIndexAndMeshUids.Value;
-				if (MeshUids.Num() > 0)
+				if ((*MeshUids)[Index] < FirstNodeUid)
 				{
-					const FString& MeshUid = MeshUids[0];
-					const UInterchangeMeshNode* MeshNode = Cast<const UInterchangeMeshNode>(BaseNodeContainer->GetNode(MeshUid));
-					if (MeshNode)
-					{
-						OutFirstMeshNodeUid = MeshUid;
-						return MeshNode;
-					}
-					const UInterchangeSceneNode* SceneNode = Cast<const UInterchangeSceneNode>(BaseNodeContainer->GetNode(MeshUid));
-					if (SceneNode)
+					FirstNodeUid = (*MeshUids)[Index];
+				}
+			}
+
+			if (!FirstNodeUid.IsEmpty())
+			{
+				if (const UInterchangeMeshNode* MeshNode = Cast<const UInterchangeMeshNode>(BaseNodeContainer->GetNode(FirstNodeUid)))
+				{
+					OutFirstMeshNodeUid = FirstNodeUid;
+					return MeshNode;
+				}
+
+				if (const UInterchangeSceneNode* SceneNode = Cast<const UInterchangeSceneNode>(BaseNodeContainer->GetNode(FirstNodeUid)))
+				{
+					FString MeshNodeUid;
+					if (SceneNode->GetCustomAssetInstanceUid(MeshNodeUid))
 					{
-						FString MeshNodeUid;
-						if (SceneNode->GetCustomAssetInstanceUid(MeshNodeUid))
-						{
-							OutSceneNodeCount = MeshUids.Num();
-							OutFirstMeshNodeUid = MeshNodeUid;
-							return SceneNode;
-						}
+						OutSceneNodeCount = MeshUids->Num();
+						OutFirstMeshNodeUid = MeshNodeUid;
+						return SceneNode;
 					}
 				}
-				//We found the lod but there is no valid Mesh node to return the Uid
-				break;
 			}
 		}
+
+		//Log an error
 		return nullptr;
 	};
 
