MaterialX:
- Fixed missing ChiangHairBSDF node matching with the corresponding material function
- Remove the substrate add operator inside the MX_Surface material function as it is incompatible with Hair BSDF

#virtualized

diff --git a/Engine/Plugins/Interchange/Runtime/Source/Common/Public/MaterialX/InterchangeMaterialXDefinitions.h b/Engine/Plugins/Interchange/Runtime/Source/Common/Public/MaterialX/InterchangeMaterialXDefinitions.h
--- a/Engine/Plugins/Interchange/Runtime/Source/Common/Public/MaterialX/InterchangeMaterialXDefinitions.h
+++ b/Engine/Plugins/Interchange/Runtime/Source/Common/Public/MaterialX/InterchangeMaterialXDefinitions.h
@@ -811,6 +811,7 @@ MATERIALX_NAMESPACE_BEGIN
 		// PBR
 		// BSDF
 		static constexpr const char* BurleyDiffuseBSDF= "burley_diffuse_bsdf";
+		static constexpr const char* ChiangHairBSDF = "chiang_hair_bsdf";
 		static constexpr const char* ConductorBSDF = "conductor_bsdf";
 		static constexpr const char* DielectricBSDF = "dielectric_bsdf";
 		static constexpr const char* GeneralizedSchlickBSDF= "generalized_schlick_bsdf";
diff --git a/Engine/Plugins/Interchange/Runtime/Source/Import/Private/MaterialX/MaterialXUtils/MaterialXManager.cpp b/Engine/Plugins/Interchange/Runtime/Source/Import/Private/MaterialX/MaterialXUtils/MaterialXManager.cpp
--- a/Engine/Plugins/Interchange/Runtime/Source/Import/Private/MaterialX/MaterialXUtils/MaterialXManager.cpp
+++ b/Engine/Plugins/Interchange/Runtime/Source/Import/Private/MaterialX/MaterialXUtils/MaterialXManager.cpp
@@ -237,6 +237,7 @@ FMaterialXManager::FMaterialXManager()
 		{mx::Category::SheenBSDF,						FMaterialXMaterialFunction{TInPlaceType<EInterchangeMaterialXBSDF>{}, EInterchangeMaterialXBSDF::Sheen}},
 		{mx::Category::SubsurfaceBSDF ,					FMaterialXMaterialFunction{TInPlaceType<EInterchangeMaterialXBSDF>{}, EInterchangeMaterialXBSDF::Subsurface}},
 		{mx::Category::TranslucentBSDF,					FMaterialXMaterialFunction{TInPlaceType<EInterchangeMaterialXBSDF>{}, EInterchangeMaterialXBSDF::Translucent}},
+		{mx::Category::ChiangHairBSDF,					FMaterialXMaterialFunction{TInPlaceType<EInterchangeMaterialXBSDF>{}, EInterchangeMaterialXBSDF::ChiangHair}},
 		// EDF Nodes
 		{mx::Category::ConicalEDF,						FMaterialXMaterialFunction{TInPlaceType<EInterchangeMaterialXEDF>{}, EInterchangeMaterialXEDF::Conical}},
 		{mx::Category::MeasuredEDF,						FMaterialXMaterialFunction{TInPlaceType<EInterchangeMaterialXEDF>{}, EInterchangeMaterialXEDF::Measured}},
