Interchange: Add nullptr check to prevent crash when trying to parse USD LodSubtreeAPI skeletal meshes.

diff --git a/Engine/Plugins/Interchange/Runtime/Source/Pipelines/Private/InterchangePipelineMeshesUtilities.cpp b/Engine/Plugins/Interchange/Runtime/Source/Pipelines/Private/InterchangePipelineMeshesUtilities.cpp
--- a/Engine/Plugins/Interchange/Runtime/Source/Pipelines/Private/InterchangePipelineMeshesUtilities.cpp
+++ b/Engine/Plugins/Interchange/Runtime/Source/Pipelines/Private/InterchangePipelineMeshesUtilities.cpp
@@ -540,7 +540,15 @@ UInterchangePipelineMeshesUtilities* UInterchangePipelineMeshesUtilities::Create
 		SkinnedMeshNode->GetSkeletonDependency(0, JointNodeUid);
 		while (!JointNodeUid.Equals(UInterchangeBaseNode::InvalidNodeUid()) && !SkeletonRootNodeUids.Contains(JointNodeUid))
 		{
-			JointNodeUid = BaseNodeContainer->GetNode(JointNodeUid)->GetParentUid();
+			const UInterchangeBaseNode* JointNode = BaseNodeContainer->GetNode(JointNodeUid);
+			if (ensureMsgf(JointNode, TEXT("Node set as SkeletonDependency must exist in the container!")))
+			{
+				JointNodeUid = JointNode->GetParentUid();
+			}
+			else
+			{
+				break;
+			}
 		}
 		//Add the MeshGeometry to the map per joint uid
 		if (SkeletonRootNodeUids.Contains(JointNodeUid))
