Interchange: Fix actors with identity transform being set with the inverse transform of their parents in some cases.

diff --git a/Engine/Plugins/Interchange/Runtime/Source/Import/Private/Scene/InterchangeActorHelper.cpp b/Engine/Plugins/Interchange/Runtime/Source/Import/Private/Scene/InterchangeActorHelper.cpp
--- a/Engine/Plugins/Interchange/Runtime/Source/Import/Private/Scene/InterchangeActorHelper.cpp
+++ b/Engine/Plugins/Interchange/Runtime/Source/Import/Private/Scene/InterchangeActorHelper.cpp
@@ -149,7 +149,7 @@ namespace UE::Interchange::ActorHelper
 
 			if (ParentActor || !SpawnedActor->IsAttachedTo(ParentActor))
 			{
-				SpawnedActor->AttachToActor(ParentActor, FAttachmentTransformRules::KeepWorldTransform);
+				SpawnedActor->AttachToActor(ParentActor, FAttachmentTransformRules::KeepRelativeTransform);
 			}
 		}
 
