Interchange: Fix actor AssetUserData being overwritten and losing collected metadata.

diff --git a/Engine/Plugins/Interchange/Runtime/Source/Import/Private/Scene/InterchangeActorHelper.cpp b/Engine/Plugins/Interchange/Runtime/Source/Import/Private/Scene/InterchangeActorHelper.cpp
--- a/Engine/Plugins/Interchange/Runtime/Source/Import/Private/Scene/InterchangeActorHelper.cpp
+++ b/Engine/Plugins/Interchange/Runtime/Source/Import/Private/Scene/InterchangeActorHelper.cpp
@@ -237,10 +237,17 @@ namespace UE::Interchange::ActorHelper
 		const FSoftObjectPath SceneImportAssetPath((const UObject*)SceneImportAsset);
 		if (USceneComponent* RootComponent = Actor->GetRootComponent())
 		{
-			UInterchangeAssetUserData* AssetUserData = NewObject<UInterchangeAssetUserData>(Actor);
+			// This component may already have AssetUserData provided by UInterchangeGenericAssetsPipeline::AddMetaData(), so
+			// we have to be careful as AddAssetUserData() will remove it when adding a new instance of the same class
+			UInterchangeAssetUserData* AssetUserData = RootComponent->GetAssetUserData<UInterchangeAssetUserData>();
+			if (!AssetUserData)
+			{
+				AssetUserData = NewObject<UInterchangeAssetUserData>(Actor);
+				RootComponent->AddAssetUserData(AssetUserData);
+			}
+
 			AssetUserData->MetaData.Add(UE::Interchange::InterchangeReset::Constants::SceneImportAssetPathKey, SceneImportAssetPath.ToString());
 			AssetUserData->MetaData.Add(UE::Interchange::InterchangeReset::Constants::FactoryNodeUidPathKey, FactoryNode->GetUniqueID());
-			RootComponent->AddAssetUserData(AssetUserData);
 		}
 	}
 	
