[BUGFIX] Interchange Scene Import - Level Instance Actor scene hierarchy type is hitting a check
- Skipping the Level Instance Actors that are marked as garbage while iterating through objects to update level instances on reimport.

#rnx

diff --git a/Engine/Plugins/Interchange/Runtime/Source/Pipelines/Private/InterchangeGenericScenesPipeline.cpp b/Engine/Plugins/Interchange/Runtime/Source/Pipelines/Private/InterchangeGenericScenesPipeline.cpp
--- a/Engine/Plugins/Interchange/Runtime/Source/Pipelines/Private/InterchangeGenericScenesPipeline.cpp
+++ b/Engine/Plugins/Interchange/Runtime/Source/Pipelines/Private/InterchangeGenericScenesPipeline.cpp
@@ -1417,6 +1417,11 @@ void UInterchangeGenericLevelPipeline::ExecutePostBroadcastPipeline(const UInter
 				{
 					if (CurrentLevelInstanceActor->GetWorldAsset().Get() == WorldAsset)
 					{
+						if (CurrentLevelInstanceActor->HasAnyFlags(RF_MirroredGarbage))
+						{
+							continue;
+						}
+
 						const UWorld* CurrentLevelInstanceActorWorld = CurrentLevelInstanceActor->GetWorld();
 						if (!CurrentLevelInstanceActorWorld || CurrentLevelInstanceActorWorld->IsCleanedUp())
 						{
