Interchange MaterialX:
- Add support for <volumematerial> and <volumeshader> nodes
- Add MX_SimpleVolume material function that handle basic volume material using Substrate slab with a simple volume, internally storing the Albedo, Mean Free Path and Anisotropy
- Add MX_VolumetricFOGCloud  material function that handle volumetic material using Substrate FOG-Cloud Slab, internally storing the Extinction on the Volume
- Add an option in the MaterialX pipeline to switch between the 2 types of volume materials
- Fixed missing match between <randomfloat> and its material function counterpart
- Add a helper to log messages and display them in the Interchange window

#virtualized

diff --git a/Engine/Plugins/Interchange/Runtime/Source/Common/Public/MaterialX/InterchangeMaterialXDefinitions.h b/Engine/Plugins/Interchange/Runtime/Source/Common/Public/MaterialX/InterchangeMaterialXDefinitions.h
--- a/Engine/Plugins/Interchange/Runtime/Source/Common/Public/MaterialX/InterchangeMaterialXDefinitions.h
+++ b/Engine/Plugins/Interchange/Runtime/Source/Common/Public/MaterialX/InterchangeMaterialXDefinitions.h
@@ -50,6 +50,9 @@ enum class EInterchangeMaterialXShaders : uint8
 	/** Shader used for displacement. */
 	Displacement,
 
+	/** Construct a volume shader describing a participating medium. */
+	Volume,
+
 	MaxShaderCount UMETA(hidden)
 };
 
@@ -864,6 +867,8 @@ MATERIALX_NAMESPACE_BEGIN
 		static constexpr const char* FacingRatio = "facingratio";
 		static constexpr const char* GoochShade = "gooch_shade";
 		static constexpr const char* ViewDirection = "viewdirection";
+		// Volume shaders
+		static constexpr const char* Volume = "volume";
 	}
 
 	namespace NodeDefinition
@@ -879,6 +884,7 @@ MATERIALX_NAMESPACE_BEGIN
 		static constexpr const char* DisplacementFloat = "ND_displacement_float";
 		static constexpr const char* DisplacementVector3 = "ND_displacement_vector3";
 		static constexpr const char* MixSurfaceShader = "ND_mix_surfaceshader";
+		static constexpr const char* Volume = "ND_volume";
 	}
 
 	namespace Library
diff --git a/Engine/Plugins/Interchange/Runtime/Source/Import/Private/MaterialX/MaterialXUtils/MaterialXBase.cpp b/Engine/Plugins/Interchange/Runtime/Source/Import/Private/MaterialX/MaterialXUtils/MaterialXBase.cpp
--- a/Engine/Plugins/Interchange/Runtime/Source/Import/Private/MaterialX/MaterialXUtils/MaterialXBase.cpp
+++ b/Engine/Plugins/Interchange/Runtime/Source/Import/Private/MaterialX/MaterialXUtils/MaterialXBase.cpp
@@ -40,6 +40,10 @@ void FMaterialXBase::SetNodeDefinition(const char* NodeDef)
 	this->NodeDefinition = NodeDef;
 }
 
+void FMaterialXBase::SetTranslator(const UInterchangeTranslatorBase* InTranslator)
+{
+	Translator = InTranslator;
+}
 
 void FMaterialXBase::AddTexCoordToTiledImageNodes(MaterialX::ElementPtr Graph)
 {
@@ -226,7 +230,7 @@ MaterialX::InputPtr FMaterialXBase::GetInput(MaterialX::NodePtr Node, const char
 	return Input;
 }
 
-FLinearColor FMaterialXBase::GetLinearColor(MaterialX::InputPtr Input)
+FLinearColor FMaterialXBase::GetLinearColor(MaterialX::InputPtr Input) const
 {
 	//we assume that the default color space is linear
 	FLinearColor LinearColor = FLinearColor::Black;
@@ -258,7 +262,7 @@ FLinearColor FMaterialXBase::GetLinearColor(MaterialX::InputPtr Input)
 	}
 	else
 	{
-		UE_LOG(LogInterchangeImport, Warning, TEXT("<%hs>-<%hs>: Colorspace %s is not supported yet, falling back to linear"),
+		MTLX_LOG("MaterialXGetLinearColor", "<{0}>-<{1}>: Colorspace {2} is not supported yet, falling back to linear space",
 			   Input->getParent()->getName().c_str(),
 			   Input->getName().c_str(),
 			   *ColorSpace);
diff --git a/Engine/Plugins/Interchange/Runtime/Source/Import/Private/MaterialX/MaterialXUtils/MaterialXManager.cpp b/Engine/Plugins/Interchange/Runtime/Source/Import/Private/MaterialX/MaterialXUtils/MaterialXManager.cpp
--- a/Engine/Plugins/Interchange/Runtime/Source/Import/Private/MaterialX/MaterialXUtils/MaterialXManager.cpp
+++ b/Engine/Plugins/Interchange/Runtime/Source/Import/Private/MaterialX/MaterialXUtils/MaterialXManager.cpp
@@ -20,6 +20,8 @@
 #include "MaterialX/MaterialXUtils/MaterialXSurfaceMaterial.h"
 #include "MaterialX/MaterialXUtils/MaterialXDisplacementShader.h"
 #include "MaterialX/MaterialXUtils/MaterialXMixShader.h"
+#include "MaterialX/MaterialXUtils/MaterialXVolumeShader.h"
+#include "MaterialX/MaterialXUtils/MaterialXVolumeMaterial.h"
 
 #include "InterchangeMaterialDefinitions.h"
 #include "InterchangeImportLog.h"
@@ -215,7 +217,9 @@ FMaterialXManager::FMaterialXManager()
 		{mx::Category::PointLight, FMaterialXManager::FOnGetMaterialXInstance::CreateStatic(&FMaterialXPointLightShader::MakeInstance)},
 		{mx::Category::DirectionalLight, FMaterialXManager::FOnGetMaterialXInstance::CreateStatic(&FMaterialXDirectionalLightShader::MakeInstance)},
 		{mx::Category::SpotLight, FMaterialXManager::FOnGetMaterialXInstance::CreateStatic(&FMaterialXSpotLightShader::MakeInstance)},
-		{mx::SURFACE_MATERIAL_NODE_STRING.c_str(), FMaterialXManager::FOnGetMaterialXInstance::CreateStatic(&FMaterialXSurfaceMaterial::MakeInstance)}
+		{mx::SURFACE_MATERIAL_NODE_STRING.c_str(), FMaterialXManager::FOnGetMaterialXInstance::CreateStatic(&FMaterialXSurfaceMaterial::MakeInstance)},
+		{mx::VOLUME_MATERIAL_NODE_STRING.c_str(), FMaterialXManager::FOnGetMaterialXInstance::CreateStatic(&FMaterialXVolumeMaterial::MakeInstance)},
+		{mx::Category::Volume, FMaterialXManager::FOnGetMaterialXInstance::CreateStatic(&FMaterialXVolumeShader::MakeInstance)},
 	}
 	, MatchingMaterialFunctions {
 
@@ -227,6 +231,7 @@ FMaterialXManager::FMaterialXManager()
 		{mx::Category::Surface,							FMaterialXMaterialFunction{TInPlaceType<EInterchangeMaterialXShaders>{}, EInterchangeMaterialXShaders::Surface}},
 		{mx::Category::SurfaceUnlit,					FMaterialXMaterialFunction{TInPlaceType<EInterchangeMaterialXShaders>{}, EInterchangeMaterialXShaders::SurfaceUnlit}},
 		{mx::Category::Displacement,					FMaterialXMaterialFunction{TInPlaceType<EInterchangeMaterialXShaders>{}, EInterchangeMaterialXShaders::Displacement}},
+		{mx::Category::Volume,							FMaterialXMaterialFunction{TInPlaceType<EInterchangeMaterialXShaders>{}, EInterchangeMaterialXShaders::Volume}},
 
 		// BSDF Nodes
 		{mx::Category::BurleyDiffuseBSDF,				FMaterialXMaterialFunction{TInPlaceType<EInterchangeMaterialXBSDF>{}, EInterchangeMaterialXBSDF::BurleyDiffuse}},
@@ -269,6 +274,7 @@ FMaterialXManager::FMaterialXManager()
 		{mx::Category::Ramp,							FMaterialXMaterialFunction{TInPlaceType<FString>{}, MaterialFunctionPackage(UE::MaterialFunctions::Path::MxRamp)}},
 		{mx::Category::RampGradient,					FMaterialXMaterialFunction{TInPlaceType<FString>{}, MaterialFunctionPackage(UE::MaterialFunctions::Path::MxRampGradient)}},
 		{mx::Category::RandomColor,						FMaterialXMaterialFunction{TInPlaceType<FString>{}, MaterialFunctionPackage(UE::MaterialFunctions::Path::MxRandomColor)}},
+		{mx::Category::RandomFloat,						FMaterialXMaterialFunction{TInPlaceType<FString>{}, MaterialFunctionPackage(UE::MaterialFunctions::Path::MxRandomFloat)}},
 		{mx::Category::TiledCircles,					FMaterialXMaterialFunction{TInPlaceType<FString>{}, MaterialFunctionPackage(UE::MaterialFunctions::Path::MxTiledCircles)}},
 		{mx::Category::UnifiedNoise3D,					FMaterialXMaterialFunction{TInPlaceType<FString>{}, MaterialFunctionPackage(UE::MaterialFunctions::Path::MxUnifiedNoise3D)}},
 		// Adjustment																						
@@ -365,11 +371,18 @@ FMaterialXManager::FMaterialXManager()
 		MatchingInputNames.Add({ FKeyExpression{mx::Category::Mix, mx::NodeGroup::PBR, mx::Type::BSDF}, TEXT("mix") }, ExpressionInput(UE::Expressions::Inputs::Mix));
 		MatchingMaterialExpressions.Add({ mx::Category::Mix, mx::NodeGroup::PBR, mx::Type::BSDF}, UE::Expressions::Names::SubstrateHorizontalMixing);
 
+		// Mix surfaceshaders
 		MatchingInputNames.Add({ FKeyExpression{mx::Category::Mix, mx::NodeGroup::Compositing, mx::Type::SurfaceShader}, TEXT("bg") }, ExpressionInput(UE::Expressions::Inputs::Background));
 		MatchingInputNames.Add({ FKeyExpression{mx::Category::Mix, mx::NodeGroup::Compositing, mx::Type::SurfaceShader}, TEXT("fg") }, ExpressionInput(UE::Expressions::Inputs::Foreground));
 		MatchingInputNames.Add({ FKeyExpression{mx::Category::Mix, mx::NodeGroup::Compositing, mx::Type::SurfaceShader}, TEXT("mix") }, ExpressionInput(UE::Expressions::Inputs::Mix));
 		MatchingMaterialExpressions.Add({ mx::Category::Mix, mx::NodeGroup::Compositing, mx::Type::SurfaceShader }, UE::Expressions::Names::SubstrateHorizontalMixing);
 
+		// Mix volumeshaders
+		MatchingInputNames.Add({ FKeyExpression{mx::Category::Mix, mx::NodeGroup::Compositing, mx::Type::VolumeShader}, TEXT("bg") }, ExpressionInput(UE::Expressions::Inputs::Background));
+		MatchingInputNames.Add({ FKeyExpression{mx::Category::Mix, mx::NodeGroup::Compositing, mx::Type::VolumeShader}, TEXT("fg") }, ExpressionInput(UE::Expressions::Inputs::Foreground));
+		MatchingInputNames.Add({ FKeyExpression{mx::Category::Mix, mx::NodeGroup::Compositing, mx::Type::VolumeShader}, TEXT("mix") }, ExpressionInput(UE::Expressions::Inputs::Mix));
+		MatchingMaterialExpressions.Add({ mx::Category::Mix, mx::NodeGroup::Compositing, mx::Type::VolumeShader }, UE::Expressions::Names::SubstrateHorizontalMixing);
+
 		// Add
 		MatchingMaterialExpressions.Add({ mx::Category::Add, mx::NodeGroup::PBR, mx::Type::BSDF}, UE::Expressions::Names::SubstrateAdd);
 
@@ -564,6 +577,7 @@ bool FMaterialXManager::Translate(MaterialX::DocumentPtr Document, UInterchangeB
 					TSharedPtr<FMaterialXBase> ShaderTranslator = FMaterialXManager::GetInstance().GetShaderTranslator(Node->getCategory().c_str(), BaseNodeContainer);
 					if (ShaderTranslator)
 					{
+						ShaderTranslator->SetTranslator(Translator);
 						ShaderTranslator->Translate(Node);
 					}
 				}
diff --git a/Engine/Plugins/Interchange/Runtime/Source/Import/Private/MaterialX/MaterialXUtils/MaterialXMixShader.cpp b/Engine/Plugins/Interchange/Runtime/Source/Import/Private/MaterialX/MaterialXUtils/MaterialXMixShader.cpp
--- a/Engine/Plugins/Interchange/Runtime/Source/Import/Private/MaterialX/MaterialXUtils/MaterialXMixShader.cpp
+++ b/Engine/Plugins/Interchange/Runtime/Source/Import/Private/MaterialX/MaterialXUtils/MaterialXMixShader.cpp
@@ -41,17 +41,16 @@ UInterchangeBaseNode* FMaterialXMixShader::Translate(mx::NodePtr MixNode)
 		SetAttributeNewName(Input, TCHAR_TO_ANSI(UE::Expressions::Inputs::Mix));
 	}
 
-	UInterchangeShaderNode* MixShaderNode = CreateShaderNode(MixNode, MixNode->getName().c_str(), UE::Expressions::Names::SubstrateHorizontalMixing);
-
-	MixShaderNode = FMaterialXSurfaceShaderAbstract::Translate(MixShaderNode);
-
-	if(bIsSubstrateEnabled)
-	{
-		UInterchangeShaderPortsAPI::ConnectDefaultOuputToInput(ShaderGraphNode, SubstrateMaterial::Parameters::FrontMaterial.ToString(), MixShaderNode->GetUniqueID());
-	}
+	UInterchangeShaderNode* MixShaderNode = nullptr;
+	if (bIsSubstrateEnabled)
+    {
+		MixShaderNode = CreateShaderNode(MixNode, MixNode->getName().c_str(), UE::Expressions::Names::SubstrateHorizontalMixing);
+        MixShaderNode = FMaterialXSurfaceShaderAbstract::Translate(MixShaderNode);
+        UInterchangeShaderPortsAPI::ConnectDefaultOuputToInput(ShaderGraphNode, SubstrateMaterial::Parameters::FrontMaterial.ToString(), MixShaderNode->GetUniqueID());
+    }
 	else
 	{
-		UE_LOG(LogInterchangeImport, Error, TEXT("<%s> : mix of surface shaders is only available with Substrate."), UTF8_TO_TCHAR(MixNode->getName().c_str()));
+		MTLX_LOG("MaterialXMixShader", "\"{0}\": {1} of {2} is only available with Substrate.", MixNode->getName().c_str(), MixNode->getCategory().c_str(), MixNode->getType().c_str());
 	}
 
 	return MixShaderNode;
diff --git a/Engine/Plugins/Interchange/Runtime/Source/Import/Private/MaterialX/MaterialXUtils/MaterialXSurfaceMaterial.cpp b/Engine/Plugins/Interchange/Runtime/Source/Import/Private/MaterialX/MaterialXUtils/MaterialXSurfaceMaterial.cpp
--- a/Engine/Plugins/Interchange/Runtime/Source/Import/Private/MaterialX/MaterialXUtils/MaterialXSurfaceMaterial.cpp
+++ b/Engine/Plugins/Interchange/Runtime/Source/Import/Private/MaterialX/MaterialXUtils/MaterialXSurfaceMaterial.cpp
@@ -31,15 +31,7 @@ UInterchangeBaseNode* FMaterialXSurfaceMaterial::Translate(MaterialX::NodePtr Su
 	const FString ShaderGraphNodeUID = UInterchangeShaderNode::MakeNodeUid(UTF8_TO_TCHAR(SurfaceMaterialNode->getName().c_str()), FStringView{});
 	UInterchangeShaderGraphNode* ShaderGraphNode = const_cast<UInterchangeShaderGraphNode*>(Cast<UInterchangeShaderGraphNode>(NodeContainer.GetNode(ShaderGraphNodeUID)));
 
-	if(!ShaderGraphNode)
-	{
-		ShaderGraphNode = NewObject<UInterchangeShaderGraphNode>(&NodeContainer);
-		// We only add the ShaderGraph to the container if we found a supported surfaceshader
-		ShaderGraphNode->InitializeNode(ShaderGraphNodeUID, SurfaceMaterialNode->getName().c_str(), EInterchangeNodeContainerType::TranslatedAsset);
-		NodeContainer.ClearNodeParentUid(ShaderGraphNodeUID); //this seems unneccessary?
-		bIsShaderGraphInContainer = false;
-	}
-	else
+	if(ShaderGraphNode)
 	{
 		return ShaderGraphNode;
 	}
@@ -61,8 +53,18 @@ UInterchangeBaseNode* FMaterialXSurfaceMaterial::Translate(MaterialX::NodePtr Su
 			TSharedPtr<FMaterialXSurfaceShaderAbstract> ShaderTranslator = StaticCastSharedPtr<FMaterialXSurfaceShaderAbstract>(FMaterialXManager::GetInstance().GetShaderTranslator(ConnectedNode->getCategory().c_str(), NodeContainer));
 			if (ShaderTranslator.IsValid())
 			{
+				// We only add the ShaderGraph to the container if we found a supported surfaceshader
+				if (!ShaderGraphNode)
+				{
+					ShaderGraphNode = NewObject<UInterchangeShaderGraphNode>(&NodeContainer);
+					ShaderGraphNode->InitializeNode(ShaderGraphNodeUID, SurfaceMaterialNode->getName().c_str(), EInterchangeNodeContainerType::TranslatedAsset);
+					NodeContainer.AddNode(ShaderGraphNode);
+					bIsShaderGraphInContainer = false;
+				}
+
 				ShaderTranslator->ShaderGraphNode = ShaderGraphNode;
 				ShaderTranslator->SurfaceMaterialName = SurfaceMaterialNode->getName().c_str();
+				ShaderTranslator->SetTranslator(Translator);
 				// we set the nodedef here, because a displacement shader can either be a float or vector3, we want to be sure to retrieve the correct one
 				ShaderTranslator->SetNodeDefinition(ConnectedNode->getNodeDef()->getName().c_str());
 				bHasSurfaceShader |= true;
@@ -104,14 +106,9 @@ UInterchangeBaseNode* FMaterialXSurfaceMaterial::Translate(MaterialX::NodePtr Su
 		UInterchangeShaderPortsAPI::ConnectOuputToInputByName(SurfaceShaderNode, InputNormalName, DisplacementShaderNode->GetUniqueID(), PBRMR::Parameters::Normal.ToString());
 	}
 
-	// We only add the ShaderGraph to the container if we found a supported surfaceshader
-	if(bHasSurfaceShader && !bIsShaderGraphInContainer)
-	{
-		NodeContainer.AddNode(ShaderGraphNode);
-	}
-	else
+	if(!bHasSurfaceShader)
 	{
-		UE_LOG(LogInterchangeImport, Warning, TEXT("the surfaceshader of <%s> is not supported"), UTF8_TO_TCHAR(SurfaceMaterialNode->getName().c_str()));
+		MTLX_LOG("MaterialXSurfaceMaterial", "<{0}>: {1} \"{2}\" is not supported", SurfaceMaterialNode->getCategory().c_str(), SurfaceMaterialNode->getType().c_str(), SurfaceMaterialNode->getName().c_str());
 	}
 
 	return ShaderGraphNode;
diff --git a/Engine/Plugins/Interchange/Runtime/Source/Import/Private/MaterialX/MaterialXUtils/MaterialXSurfaceShaderAbstract.cpp b/Engine/Plugins/Interchange/Runtime/Source/Import/Private/MaterialX/MaterialXUtils/MaterialXSurfaceShaderAbstract.cpp
--- a/Engine/Plugins/Interchange/Runtime/Source/Import/Private/MaterialX/MaterialXUtils/MaterialXSurfaceShaderAbstract.cpp
+++ b/Engine/Plugins/Interchange/Runtime/Source/Import/Private/MaterialX/MaterialXUtils/MaterialXSurfaceShaderAbstract.cpp
@@ -252,7 +252,7 @@ bool FMaterialXSurfaceShaderAbstract::ConnectNodeGraphOutputToInput(MaterialX::I
 
 		if(!Output)
 		{
-			UE_LOG(LogInterchangeImport, Warning, TEXT("Couldn't find a connected output to (%s)."), *GetInputName(InputToNodeGraph));
+			MTLX_LOG("MaterialNodeGraphOutput", "Couldn't find a connected output to ({0}).", *GetInputName(InputToNodeGraph));
 			return false;
 		}
 
@@ -426,7 +426,7 @@ void FMaterialXSurfaceShaderAbstract::ConnectNodeCategoryOutputToInput(const Mat
 		{
 			if (!ConnectMatchingNodeOutputToInput({ UpstreamNode, ParentShaderNode, InputChannelName, OutputChannelName }))
 			{
-				UE_LOG(LogInterchangeImport, Warning, TEXT("<%s>: \"%s\" is not supported. [%s]"), UTF8_TO_TCHAR(UpstreamNode->getCategory().c_str()), UTF8_TO_TCHAR(UpstreamNode->getName().c_str()), *SurfaceMaterialName);
+				MTLX_LOG("MaterialXNodeCategoryOutput", "<{0}>: \"{1}\" is not supported. [{2}]", UpstreamNode->getCategory().c_str(), UpstreamNode->getName().c_str(), *SurfaceMaterialName);
 			}
 		}
 	}
@@ -1180,7 +1180,7 @@ void FMaterialXSurfaceShaderAbstract::ConnectBlurInputToOutput(const FConnectNod
 					}
 					else
 					{
-						UE_LOG(LogInterchangeImport, Warning, TEXT("<%s>: input 'size' must have a value."), UTF8_TO_TCHAR(Connect.UpstreamNode->getName().c_str()));
+						MTLX_LOG("MaterialXBlur", "<{0}>: input 'size' must have a value.", Connect.UpstreamNode->getName().c_str());
 					}
 				}
 			}
diff --git a/Engine/Plugins/Interchange/Runtime/Source/Import/Private/MaterialX/MaterialXUtils/MaterialXVolumeMaterial.cpp b/Engine/Plugins/Interchange/Runtime/Source/Import/Private/MaterialX/MaterialXUtils/MaterialXVolumeMaterial.cpp
new file mode 100644
--- /dev/null
+++ b/Engine/Plugins/Interchange/Runtime/Source/Import/Private/MaterialX/MaterialXUtils/MaterialXVolumeMaterial.cpp
@@ -0,0 +1,66 @@
+// Copyright Epic Games, Inc. All Rights Reserved.
+
+#if WITH_EDITOR
+
+#include "MaterialXVolumeMaterial.h"
+#include "InterchangeImportLog.h"
+#include "MaterialX/MaterialXUtils/MaterialXManager.h"
+#include "MaterialX/MaterialXUtils/MaterialXSurfaceShaderAbstract.h"
+
+#define LOCTEXT_NAMESPACE "InterchangeMaterialXVolumeMaterial"
+
+namespace mx = MaterialX;
+
+FMaterialXVolumeMaterial::FMaterialXVolumeMaterial(UInterchangeBaseNodeContainer& BaseNodeContainer)
+	: FMaterialXBase(BaseNodeContainer)
+{}
+
+TSharedRef<FMaterialXBase> FMaterialXVolumeMaterial::MakeInstance(UInterchangeBaseNodeContainer& BaseNodeContainer)
+{
+	return MakeShared<FMaterialXVolumeMaterial>(BaseNodeContainer);
+}
+
+UInterchangeBaseNode* FMaterialXVolumeMaterial::Translate(MaterialX::NodePtr VolumeMaterialNode)
+{
+	using namespace UE::Interchange::Materials;
+
+	if (!FMaterialXManager::GetInstance().IsSubstrateEnabled())
+	{
+		MTLX_LOG("MaterialXVolumeMaterial", "\"{0}\": {1} is only supported with Substrate enabled.", VolumeMaterialNode->getName().c_str(), VolumeMaterialNode->getCategory().c_str());
+		return nullptr;
+	}
+
+	const FString ShaderGraphNodeUID = UInterchangeShaderNode::MakeNodeUid(UTF8_TO_TCHAR(VolumeMaterialNode->getName().c_str()), FStringView{});
+	UInterchangeShaderGraphNode* ShaderGraphNode = const_cast<UInterchangeShaderGraphNode*>(Cast<UInterchangeShaderGraphNode>(NodeContainer.GetNode(ShaderGraphNodeUID)));
+
+	if(ShaderGraphNode)
+	{
+		return ShaderGraphNode;
+	}
+
+	if (mx::InputPtr Input = VolumeMaterialNode->getInput("volumeshader"))
+	{
+		if (mx::NodePtr VolumeShaderNode = Input->getConnectedNode())
+		{
+			if(TSharedPtr<FMaterialXSurfaceShaderAbstract> ShaderTranslator = StaticCastSharedPtr<FMaterialXSurfaceShaderAbstract>(FMaterialXManager::GetInstance().GetShaderTranslator(VolumeShaderNode->getCategory().c_str(), NodeContainer)))
+			{
+				ShaderGraphNode = NewObject<UInterchangeShaderGraphNode>(&NodeContainer);
+				ShaderGraphNode->InitializeNode(ShaderGraphNodeUID, VolumeMaterialNode->getName().c_str(), EInterchangeNodeContainerType::TranslatedAsset);
+				ShaderTranslator->ShaderGraphNode = ShaderGraphNode;
+				ShaderTranslator->SurfaceMaterialName = VolumeShaderNode->getName().c_str();
+				ShaderTranslator->SetTranslator(Translator);
+				ShaderTranslator->Translate(VolumeShaderNode);
+				NodeContainer.AddNode(ShaderGraphNode);
+			}
+			else
+			{
+				MTLX_LOG("MaterialXVolumeMaterial", "<{0}>: {1} \"{2}\" is not supported", VolumeShaderNode->getCategory().c_str(), VolumeShaderNode->getType().c_str(), VolumeShaderNode->getName().c_str());
+			}
+		}
+	}
+
+	return ShaderGraphNode;
+}
+
+#undef LOCTEXT_NAMESPACE 
+#endif
\ No newline at end of file
diff --git a/Engine/Plugins/Interchange/Runtime/Source/Import/Private/MaterialX/MaterialXUtils/MaterialXVolumeMaterial.h b/Engine/Plugins/Interchange/Runtime/Source/Import/Private/MaterialX/MaterialXUtils/MaterialXVolumeMaterial.h
new file mode 100644
--- /dev/null
+++ b/Engine/Plugins/Interchange/Runtime/Source/Import/Private/MaterialX/MaterialXUtils/MaterialXVolumeMaterial.h
@@ -0,0 +1,23 @@
+// Copyright Epic Games, Inc. All Rights Reserved.
+
+#pragma once
+
+#if WITH_EDITOR
+#include "MaterialX/MaterialXUtils/MaterialXBase.h"
+
+class FMaterialXVolumeMaterial : public FMaterialXBase
+{
+protected:
+
+	template<typename T, ESPMode>
+	friend class SharedPointerInternals::TIntrusiveReferenceController;
+
+	FMaterialXVolumeMaterial(UInterchangeBaseNodeContainer& BaseNodeContainer);
+
+public:
+
+	static TSharedRef<FMaterialXBase> MakeInstance(UInterchangeBaseNodeContainer& BaseNodeContainer);
+
+	virtual UInterchangeBaseNode* Translate(MaterialX::NodePtr VolumeMaterialNode) override;
+};
+#endif
\ No newline at end of file
diff --git a/Engine/Plugins/Interchange/Runtime/Source/Import/Private/MaterialX/MaterialXUtils/MaterialXVolumeShader.cpp b/Engine/Plugins/Interchange/Runtime/Source/Import/Private/MaterialX/MaterialXUtils/MaterialXVolumeShader.cpp
new file mode 100644
--- /dev/null
+++ b/Engine/Plugins/Interchange/Runtime/Source/Import/Private/MaterialX/MaterialXUtils/MaterialXVolumeShader.cpp
@@ -0,0 +1,36 @@
+// Copyright Epic Games, Inc. All Rights Reserved.
+
+#if WITH_EDITOR
+
+#include "MaterialXVolumeShader.h"
+#include "InterchangeImportLog.h"
+
+namespace mx = MaterialX;
+
+FMaterialXVolumeShader::FMaterialXVolumeShader(UInterchangeBaseNodeContainer& BaseNodeContainer)
+	: FMaterialXSurfaceShaderAbstract(BaseNodeContainer)
+{
+	NodeDefinition = mx::NodeDefinition::Volume;
+}
+
+TSharedRef<FMaterialXBase> FMaterialXVolumeShader::MakeInstance(UInterchangeBaseNodeContainer& BaseNodeContainer)
+{
+	TSharedRef<FMaterialXVolumeShader> Result = MakeShared<FMaterialXVolumeShader>(BaseNodeContainer);
+	Result->RegisterConnectNodeOutputToInputDelegates();
+	return Result;
+}
+
+UInterchangeBaseNode* FMaterialXVolumeShader::Translate(MaterialX::NodePtr VolumeNode)
+{
+	this->SurfaceShaderNode = VolumeNode;
+
+	UInterchangeShaderNode* FunctionVolumeShaderNode = FMaterialXSurfaceShaderAbstract::Translate(EInterchangeMaterialXShaders::Volume);
+
+	using namespace UE::Interchange::Materials;
+
+	// Outputs
+	UInterchangeShaderPortsAPI::ConnectDefaultOuputToInput(ShaderGraphNode, SubstrateMaterial::Parameters::FrontMaterial.ToString(), FunctionVolumeShaderNode->GetUniqueID());
+
+	return FunctionVolumeShaderNode;
+}
+#endif
\ No newline at end of file
diff --git a/Engine/Plugins/Interchange/Runtime/Source/Import/Private/MaterialX/MaterialXUtils/MaterialXVolumeShader.h b/Engine/Plugins/Interchange/Runtime/Source/Import/Private/MaterialX/MaterialXUtils/MaterialXVolumeShader.h
new file mode 100644
--- /dev/null
+++ b/Engine/Plugins/Interchange/Runtime/Source/Import/Private/MaterialX/MaterialXUtils/MaterialXVolumeShader.h
@@ -0,0 +1,25 @@
+// Copyright Epic Games, Inc. All Rights Reserved.
+
+#pragma once
+
+#if WITH_EDITOR
+
+#include "MaterialX/MaterialXUtils/MaterialXSurfaceShaderAbstract.h"
+
+/** Handles <volume> node from the pbrlib, acts more or less as a generic volume shader */
+class FMaterialXVolumeShader : public FMaterialXSurfaceShaderAbstract
+{
+protected:
+
+	template<typename T, ESPMode>
+	friend class SharedPointerInternals::TIntrusiveReferenceController;
+
+	FMaterialXVolumeShader(UInterchangeBaseNodeContainer& BaseNodeContainer);
+
+public:
+
+	static TSharedRef<FMaterialXBase> MakeInstance(UInterchangeBaseNodeContainer& BaseNodeContainer);
+
+	virtual UInterchangeBaseNode* Translate(MaterialX::NodePtr SurfaceNode) override;
+};
+#endif
\ No newline at end of file
diff --git a/Engine/Plugins/Interchange/Runtime/Source/Import/Public/MaterialX/MaterialXUtils/MaterialXBase.h b/Engine/Plugins/Interchange/Runtime/Source/Import/Public/MaterialX/MaterialXUtils/MaterialXBase.h
--- a/Engine/Plugins/Interchange/Runtime/Source/Import/Public/MaterialX/MaterialXUtils/MaterialXBase.h
+++ b/Engine/Plugins/Interchange/Runtime/Source/Import/Public/MaterialX/MaterialXUtils/MaterialXBase.h
@@ -4,10 +4,39 @@
 
 #if WITH_EDITOR
 
+#include "Import/Private/InterchangeImportLog.h"
 #include "InterchangeMaterialDefinitions.h"
+#include "InterchangeTranslatorBase.h"
+#include "Logging/StructuredLog.h"
 #include "MaterialXCore/Document.h"
 #include "MaterialX/InterchangeMaterialXDefinitions.h"
-#include "Nodes/InterchangeBaseNodeContainer.h"
+
+namespace UE::Interchange::MaterialX::Private
+{
+	static inline FText FormatFText(const FText& Text)
+	{
+		return Text;
+	}
+
+	template<typename... Args>
+	FText FormatFText(const FText& Text, Args&&... args)
+	{
+		TArray<FFormatArgumentValue> FormatArgs = { FFormatArgumentValue(FText::FromString(Forward<Args>(args)))... };
+		return FText::Format(Text, FormatArgs);
+	}
+}
+#define MTLX_LOG(Key, Format, ...) \
+	if (Translator) \
+	{ \
+		const FText LocText = NSLOCTEXT("MaterialX", Key, Format); \
+		UInterchangeResultWarning_Generic* Message = Translator->AddMessage<UInterchangeResultWarning_Generic>(); \
+		Message->Text = UE::Interchange::MaterialX::Private::FormatFText(LocText, ##__VA_ARGS__); \
+	} \
+	else \
+	{ \
+		UE_LOGFMT(LogInterchangeImport, Warning, Format, ##__VA_ARGS__); \
+	}
+													
 
 class FMaterialXBase : public TSharedFromThis<FMaterialXBase>
 {
@@ -17,6 +46,8 @@ protected:
 
 public:
 
+	enum class EMaterialXLog{Error, Warning};
+
 	virtual UInterchangeBaseNode* Translate(MaterialX::NodePtr Node) = 0;
 
 	virtual ~FMaterialXBase() = default;
@@ -34,6 +65,33 @@ public:
 	 */
 	void SetNodeDefinition(const char* NodeDefinition);
 
+	/** Set the translator, only used for logging messages*/
+	void SetTranslator(const UInterchangeTranslatorBase* InTranslator);
+
+	/** Log a message either in form of a message passed to the translator or as a console log */
+	template<EMaterialXLog Verbosity, typename... Args>
+    void LogMessage(const FText& Text, Args... Arguments)
+	{
+		//if(Translator)
+		{
+			//if constexpr (Verbosity == EMaterialXLog::Error)
+			//{
+				UInterchangeResultError_Generic* Message = Translator->AddMessage<UInterchangeResultError_Generic>();
+				Message->Text = Text;
+			//}
+			//else
+			//{
+			//	UInterchangeResultWarning_Generic* Message = Translator->AddMessage<UInterchangeResultWarning_Generic>();
+			//	Message->Text = Text;
+			//}
+		}
+		//else
+		//{
+		//	//const TCHAR* Fmt = *(Text.ToString());
+		//	//UE_LOG(LogInterchangeImport, Warning, Fmt TEXT("%s"), Forward<Args>(Arguments)...);
+		//}
+	}
+
 protected:
 
 	/**
@@ -121,12 +179,15 @@ protected:
 	 *
 	 * @return The linear color after color space conversion
 	 */
-	static FLinearColor GetLinearColor(MaterialX::InputPtr Input);
+	FLinearColor GetLinearColor(MaterialX::InputPtr Input) const;
 
 protected:
 
 	UInterchangeBaseNodeContainer& NodeContainer;
 
+	/** Used for logging */
+	const UInterchangeTranslatorBase* Translator;
+
 	/** It's up to the derive class to initialize it with the proper node definition*/
 	const char* NodeDefinition; 
 
diff --git a/Engine/Plugins/Interchange/Runtime/Source/Import/Public/MaterialX/MaterialXUtils/MaterialXSurfaceShaderAbstract.h b/Engine/Plugins/Interchange/Runtime/Source/Import/Public/MaterialX/MaterialXUtils/MaterialXSurfaceShaderAbstract.h
--- a/Engine/Plugins/Interchange/Runtime/Source/Import/Public/MaterialX/MaterialXUtils/MaterialXSurfaceShaderAbstract.h
+++ b/Engine/Plugins/Interchange/Runtime/Source/Import/Public/MaterialX/MaterialXUtils/MaterialXSurfaceShaderAbstract.h
@@ -31,6 +31,7 @@ protected:
 	static FString DefaultOutput;
 
 	friend class FMaterialXSurfaceMaterial;
+	friend class FMaterialXVolumeMaterial;
 
 	struct FConnectNode
 	{
diff --git a/Engine/Plugins/Interchange/Runtime/Source/Pipelines/Private/InterchangeMaterialXPipeline.cpp b/Engine/Plugins/Interchange/Runtime/Source/Pipelines/Private/InterchangeMaterialXPipeline.cpp
--- a/Engine/Plugins/Interchange/Runtime/Source/Pipelines/Private/InterchangeMaterialXPipeline.cpp
+++ b/Engine/Plugins/Interchange/Runtime/Source/Pipelines/Private/InterchangeMaterialXPipeline.cpp
@@ -38,6 +38,12 @@ namespace
 	constexpr const TCHAR* DisplacementFunctionsPath = TEXT("/InterchangeAssets/Functions/MX_Displacement.MX_Displacement");
 	constexpr const TCHAR* DisplacementSubstratePath = TEXT("/InterchangeAssets/Functions/MX_Displacement.MX_Displacement");
 	
+	// will act differently, we have 2 different volumes depending on the setting in the pipeline
+	constexpr const TCHAR* SimpleVolumeSlabSubstratePath = TEXT("/InterchangeAssets/Substrate/MX_SimpleVolume.MX_SimpleVolume");
+	constexpr const TCHAR* FOGCloudVolumeSlabSubstratePath = TEXT("/InterchangeAssets/Substrate/MX_VolumetricFogCloud.MX_VolumetricFogCloud");
+	constexpr const TCHAR* VolumeFunctionsPath = SimpleVolumeSlabSubstratePath;
+	constexpr const TCHAR* VolumeSubstratePath = SimpleVolumeSlabSubstratePath;
+	
 	MATERIALX_FUNCTIONS_SUBSTRATE_PATH(OrenNayarBSDF);
 	MATERIALX_FUNCTIONS_SUBSTRATE_PATH(BurleyDiffuseBSDF);
 	MATERIALX_FUNCTIONS_SUBSTRATE_PATH(DielectricBSDF);
@@ -423,6 +429,18 @@ UMaterialXPipelineSettings::UMaterialXPipelineSettings()
 					}
 				}
 			},
+			{
+				UMaterialXPipelineSettings::ToEnumKey(EInterchangeMaterialXShaders::Volume),
+				{
+					// Volume Inputs
+					TSet<FName>{
+						TEXT("vdf"),
+						TEXT("edf"),
+					},
+					// Volume Outputs
+					TSet<FName>{TEXT("Volume")},
+				}
+			},
 			// BSDF
 			{
 				UMaterialXPipelineSettings::ToEnumKey(EInterchangeMaterialXBSDF::OrenNayarDiffuse),
@@ -641,9 +659,7 @@ UMaterialXPipelineSettings::UMaterialXPipelineSettings()
 						TEXT("absorption")
 					},
 					// Absorption Outputs
-					TSet<FName>{
-						TEXT("Output")
-					}
+					TSet<FName>{TEXT("Output")}
 				}
 			},
 			{
@@ -656,9 +672,7 @@ UMaterialXPipelineSettings::UMaterialXPipelineSettings()
 						TEXT("anisotropy"),
 					},
 					// Anisotropic Outputs
-					TSet<FName>{
-						TEXT("Output")
-					}
+					TSet<FName>{TEXT("Output")}
 				}
 			}
 		};
@@ -669,19 +683,20 @@ UMaterialXPipelineSettings::UMaterialXPipelineSettings()
 UInterchangeMaterialXPipeline::UInterchangeMaterialXPipeline()
 	: MaterialXSettings(UMaterialXPipelineSettings::StaticClass()->GetDefaultObject<UMaterialXPipelineSettings>())
 {
-	if(HasAnyFlags(EObjectFlags::RF_ClassDefaultObject))
+	if (HasAnyFlags(EObjectFlags::RF_ClassDefaultObject))
 	{
 #if WITH_EDITOR
-		PathToEnumMapping =	{
+		PathToEnumMapping = {
 			{MATERIALX_MATERIALFUNCTION_PATH(OpenPBRSurface),		      UMaterialXPipelineSettings::ToEnumKey(EInterchangeMaterialXShaders::OpenPBRSurface)},
 			{MATERIALX_MATERIALFUNCTION_PATH(OpenPBRTransmissionSurface), UMaterialXPipelineSettings::ToEnumKey(EInterchangeMaterialXShaders::OpenPBRSurfaceTransmission)},
 			{MATERIALX_MATERIALFUNCTION_PATH(StandardSurface),			  UMaterialXPipelineSettings::ToEnumKey(EInterchangeMaterialXShaders::StandardSurface)},
 			{MATERIALX_MATERIALFUNCTION_PATH(TransmissionSurface),		  UMaterialXPipelineSettings::ToEnumKey(EInterchangeMaterialXShaders::StandardSurfaceTransmission)},
 			{MATERIALX_MATERIALFUNCTION_PATH(SurfaceUnlit),				  UMaterialXPipelineSettings::ToEnumKey(EInterchangeMaterialXShaders::SurfaceUnlit)},
 			{MATERIALX_MATERIALFUNCTION_PATH(Surface),					  UMaterialXPipelineSettings::ToEnumKey(EInterchangeMaterialXShaders::Surface)},
+			{MATERIALX_MATERIALFUNCTION_PATH(Volume),					  UMaterialXPipelineSettings::ToEnumKey(EInterchangeMaterialXShaders::Volume)},
 			{MATERIALX_MATERIALFUNCTION_PATH(UsdPreviewSurface),		  UMaterialXPipelineSettings::ToEnumKey(EInterchangeMaterialXShaders::UsdPreviewSurface)},
 			{MATERIALX_MATERIALFUNCTION_PATH(Displacement),				  UMaterialXPipelineSettings::ToEnumKey(EInterchangeMaterialXShaders::Displacement)},
-																		  
+
 			{MATERIALX_MATERIALFUNCTION_PATH(OrenNayarBSDF),			  UMaterialXPipelineSettings::ToEnumKey(EInterchangeMaterialXBSDF::OrenNayarDiffuse)},
 			{MATERIALX_MATERIALFUNCTION_PATH(BurleyDiffuseBSDF),		  UMaterialXPipelineSettings::ToEnumKey(EInterchangeMaterialXBSDF::BurleyDiffuse)},
 			{MATERIALX_MATERIALFUNCTION_PATH(DielectricBSDF),			  UMaterialXPipelineSettings::ToEnumKey(EInterchangeMaterialXBSDF::Dielectric)},
@@ -691,11 +706,11 @@ UInterchangeMaterialXPipeline::UInterchangeMaterialXPipeline()
 			{MATERIALX_MATERIALFUNCTION_PATH(GeneralizedSchlickBSDF),	  UMaterialXPipelineSettings::ToEnumKey(EInterchangeMaterialXBSDF::GeneralizedSchlick)},
 			{MATERIALX_MATERIALFUNCTION_PATH(TranslucentBSDF),			  UMaterialXPipelineSettings::ToEnumKey(EInterchangeMaterialXBSDF::Translucent)},
 			{MATERIALX_MATERIALFUNCTION_PATH(ChiangHairBSDF),			  UMaterialXPipelineSettings::ToEnumKey(EInterchangeMaterialXBSDF::ChiangHair)},
-																		  
+
 			{MATERIALX_MATERIALFUNCTION_PATH(UniformEDF),				  UMaterialXPipelineSettings::ToEnumKey(EInterchangeMaterialXEDF::Uniform)},
 			{MATERIALX_MATERIALFUNCTION_PATH(ConicalEDF),				  UMaterialXPipelineSettings::ToEnumKey(EInterchangeMaterialXEDF::Conical)},
 			{MATERIALX_MATERIALFUNCTION_PATH(MeasuredEDF),				  UMaterialXPipelineSettings::ToEnumKey(EInterchangeMaterialXEDF::Measured)},
-																		  
+
 			{MATERIALX_MATERIALFUNCTION_PATH(AbsorptionVDF),			  UMaterialXPipelineSettings::ToEnumKey(EInterchangeMaterialXVDF::Absorption)},
 			{MATERIALX_MATERIALFUNCTION_PATH(AnisotropicVDF),			  UMaterialXPipelineSettings::ToEnumKey(EInterchangeMaterialXVDF::Anisotropic)},
 		};
@@ -704,7 +719,6 @@ UInterchangeMaterialXPipeline::UInterchangeMaterialXPipeline()
 #endif // WITH_EDITOR
 	}
 
-
 	for (const TPair<EInterchangeMaterialXShaders, FSoftObjectPath>& Entry : MaterialXSettings->PredefinedSurfaceShaders)
 	{
 		PathToEnumMapping.FindOrAdd(Entry.Value.GetAssetPathString(), EInterchangeMaterialXSettings{ TInPlaceType<EInterchangeMaterialXShaders>{}, Entry.Key });
@@ -741,6 +755,27 @@ void UInterchangeMaterialXPipeline::ExecutePipeline(UInterchangeBaseNodeContaine
 	Super::ExecutePipeline(NodeContainer, InSourceDatas, ContentBasePath);
 
 #if WITH_EDITOR
+
+	// we only override the volume shader if it's not a data-driven path
+	if (bVolumetricMaterial)
+	{
+		PathToEnumMapping.FindOrAdd(FOGCloudVolumeSlabSubstratePath, UMaterialXPipelineSettings::ToEnumKey(EInterchangeMaterialXShaders::Volume));
+		FSoftObjectPath& VolumePath = MaterialXSettings->PredefinedSurfaceShaders.FindChecked(EInterchangeMaterialXShaders::Volume);
+		if (VolumePath.ToString() == SimpleVolumeSlabSubstratePath)
+		{
+			VolumePath = FOGCloudVolumeSlabSubstratePath;
+		}
+	}
+	else
+	{
+		PathToEnumMapping.FindOrAdd(SimpleVolumeSlabSubstratePath, UMaterialXPipelineSettings::ToEnumKey(EInterchangeMaterialXShaders::Volume));
+		FSoftObjectPath& VolumePath = MaterialXSettings->PredefinedSurfaceShaders.FindChecked(EInterchangeMaterialXShaders::Volume);
+		if (VolumePath.ToString() == FOGCloudVolumeSlabSubstratePath)
+		{
+			VolumePath = SimpleVolumeSlabSubstratePath;
+		}
+	}
+
 	auto UpdateMaterialXNodes = [this, NodeContainer](const FString& NodeUid, UInterchangeMaterialFunctionCallExpressionFactoryNode* FactorNode)
 	{
 		const FString MaterialFunctionMemberName = GET_MEMBER_NAME_CHECKED(UMaterialExpressionMaterialFunctionCall, MaterialFunction).ToString();
@@ -836,7 +871,9 @@ bool UMaterialXPipelineSettings::AreRequiredPackagesLoaded()
 		return bAllLoaded;
 	};
 
-	return ArePackagesLoaded(PredefinedSurfaceShaders) && ArePackagesLoaded(PredefinedBSDF) && ArePackagesLoaded(PredefinedEDF) && ArePackagesLoaded(PredefinedVDF);
+	// ensure to load the FOGCloudVolumetric Material Function, we only need to load it once compared to the other shaders
+	static bool bFOGCloudVolumetricAsset = ArePackagesLoaded(TMap<EInterchangeMaterialXShaders, FSoftObjectPath>{ { { EInterchangeMaterialXShaders::Volume, FSoftObjectPath{ FOGCloudVolumeSlabSubstratePath } } } });
+	return ArePackagesLoaded(PredefinedSurfaceShaders) && ArePackagesLoaded(PredefinedBSDF) && ArePackagesLoaded(PredefinedEDF) && ArePackagesLoaded(PredefinedVDF) && bFOGCloudVolumetricAsset;
 }
 
 #if WITH_EDITOR
@@ -849,7 +886,7 @@ void UMaterialXPipelineSettings::InitPredefinedAssets()
 			{UMaterialXPipelineSettings::ToEnumKey(EInterchangeMaterialXShaders::OpenPBRSurface), OpenPBRSurfaceFunctionsPath, OpenPBRSurfaceSubstratePath},
 			{UMaterialXPipelineSettings::ToEnumKey(EInterchangeMaterialXShaders::OpenPBRSurfaceTransmission), OpenPBRTransmissionSurfaceFunctionsPath, OpenPBRTransmissionSurfaceSubstratePath},
 			{UMaterialXPipelineSettings::ToEnumKey(EInterchangeMaterialXShaders::Surface), SurfaceFunctionsPath, SurfaceSubstratePath},
-			{UMaterialXPipelineSettings::ToEnumKey(EInterchangeMaterialXShaders::Surface), SurfaceFunctionsPath, SurfaceSubstratePath},
+			{UMaterialXPipelineSettings::ToEnumKey(EInterchangeMaterialXShaders::Volume), VolumeFunctionsPath, VolumeSubstratePath},
 			{UMaterialXPipelineSettings::ToEnumKey(EInterchangeMaterialXShaders::StandardSurface), StandardSurfaceFunctionsPath, StandardSurfaceSubstratePath},
 			{UMaterialXPipelineSettings::ToEnumKey(EInterchangeMaterialXShaders::StandardSurfaceTransmission), TransmissionSurfaceFunctionsPath, TransmissionSurfaceSubstratePath},
 			{UMaterialXPipelineSettings::ToEnumKey(EInterchangeMaterialXShaders::SurfaceUnlit), SurfaceUnlitFunctionsPath, SurfaceUnlitSubstratePath},
diff --git a/Engine/Plugins/Interchange/Runtime/Source/Pipelines/Public/InterchangeMaterialXPipeline.h b/Engine/Plugins/Interchange/Runtime/Source/Pipelines/Public/InterchangeMaterialXPipeline.h
--- a/Engine/Plugins/Interchange/Runtime/Source/Pipelines/Public/InterchangeMaterialXPipeline.h
+++ b/Engine/Plugins/Interchange/Runtime/Source/Pipelines/Public/InterchangeMaterialXPipeline.h
@@ -92,7 +92,7 @@ private:
 #endif // WITH_EDITOR
 };
 
-UCLASS(MinimalAPI, BlueprintType)
+UCLASS(MinimalAPI, config = Interchange, BlueprintType)
 class UInterchangeMaterialXPipeline : public UInterchangePipelineBase
 {
 	GENERATED_BODY()
@@ -102,6 +102,10 @@ class UInterchangeMaterialXPipeline : public UInterchangePipelineBase
 public:
 	TObjectPtr<UMaterialXPipelineSettings> MaterialXSettings;
 
+	/** If true, handle Volume Distribution Functions as volumetric materials, i.e an internal Substrate Volumetric FOG-Cloud BSDF is created, otherwise a simple volume Slab BSDF is used. */
+	UPROPERTY(EditAnywhere, config, Category = "Volume")
+	bool bVolumetricMaterial = false;
+
 protected:
 	UE_API virtual void AdjustSettingsForContext(const FInterchangePipelineContextParams& ContextParams) override;
 	UE_API virtual void ExecutePipeline(UInterchangeBaseNodeContainer* BaseNodeContainer, const TArray<UInterchangeSourceData*>& SourceDatas, const FString& ContentBasePath) override;
