Fixed freezing of Editor  when reimporting a skeletal mesh instanced in the current level
Root cause: Inadequate fix for UE-244416 at CL 41800602
Solution: Replace use of FScopedSkeletalMeshPostEditChange with FSkinnedMeshComponentRecreateRenderStateContext

#rnx

diff --git a/Engine/Plugins/Interchange/Runtime/Source/Import/Private/Mesh/InterchangeSkeletalMeshFactory.cpp b/Engine/Plugins/Interchange/Runtime/Source/Import/Private/Mesh/InterchangeSkeletalMeshFactory.cpp
--- a/Engine/Plugins/Interchange/Runtime/Source/Import/Private/Mesh/InterchangeSkeletalMeshFactory.cpp
+++ b/Engine/Plugins/Interchange/Runtime/Source/Import/Private/Mesh/InterchangeSkeletalMeshFactory.cpp
@@ -63,29 +63,6 @@
 
 #endif //WITH_EDITORONLY_DATA
 
-FScopedSkeletalMeshReimportUtility::FScopedSkeletalMeshReimportUtility(USkeletalMesh* InSkeletalMesh)
-{
-#if WITH_EDITOR
-	// Scope PostEditChange before locking the skeletal mesh's properties
-	ScopedPostEditChange = MakeUnique<FScopedSkeletalMeshPostEditChange>(InSkeletalMesh);
-	// Lock the skeletal mesh's properties to prevent any async access
-	LockPropertiesEvent = InSkeletalMesh->LockPropertiesUntil();
-#endif
-}
-
-FScopedSkeletalMeshReimportUtility::~FScopedSkeletalMeshReimportUtility()
-{
-#if WITH_EDITOR
-	// Unlock the skeletal mesh's properties first
-	if (ensure(LockPropertiesEvent))
-	{
-		LockPropertiesEvent->Trigger();
-	}
-	// Now apply PostEditChange to update skeletal mesh's resources
-	ScopedPostEditChange.Reset();
-#endif
-}
-
 #if WITH_EDITOR
 namespace UE
 {
@@ -1039,7 +1016,7 @@ UInterchangeFactoryBase::FImportAssetResult UInterchangeSkeletalMeshFactory::Beg
 	//Lock the skeletalmesh properties if the skeletal mesh already exist (re-import)
 	if (ExistingAsset)
 	{
-		ScopedReimportUtility = MakeUnique<FScopedSkeletalMeshReimportUtility>(SkeletalMesh);
+		ScopedReimportUtility = MakeUnique<FSkinnedMeshComponentRecreateRenderStateContext>(SkeletalMesh);
 	}
 
 	ImportAssetResult.ImportedObject = SkeletalMesh;
diff --git a/Engine/Plugins/Interchange/Runtime/Source/Import/Public/Mesh/InterchangeSkeletalMeshFactory.h b/Engine/Plugins/Interchange/Runtime/Source/Import/Public/Mesh/InterchangeSkeletalMeshFactory.h
--- a/Engine/Plugins/Interchange/Runtime/Source/Import/Public/Mesh/InterchangeSkeletalMeshFactory.h
+++ b/Engine/Plugins/Interchange/Runtime/Source/Import/Public/Mesh/InterchangeSkeletalMeshFactory.h
@@ -22,29 +22,6 @@ class UInterchangeSkeletalMeshFactoryNode;
 class USkeletalMesh;
 class USkeleton;
 
-// Utility class used on reimport only to: 
-// - Call PostEditchange at the end of the reimport which will reallocate the render resource
-// - Lock the skeletal mesh's properties while it is updated
-class FScopedSkeletalMeshReimportUtility
-{
-public:
-	/*
-	 * This constructor increment the skeletal mesh PostEditChangeStackCounter. If the stack counter is zero before the increment
-	 * the skeletal mesh component will be unregister from the world. The component will also release there rendering resources.
-	 * Parameters:
-	 * @param InbCallPostEditChange - if we are the first scope PostEditChange will be called.
-	 * @param InbReregisterComponents - if we are the first scope we will re register component from world and also component render data.
-	 */
-	FScopedSkeletalMeshReimportUtility(USkeletalMesh* InSkeletalMesh);
-	~FScopedSkeletalMeshReimportUtility();
-
-private:
-#if WITH_EDITOR
-	TUniquePtr<FScopedSkeletalMeshPostEditChange> ScopedPostEditChange;
-#endif
-	FEvent* LockPropertiesEvent = nullptr;
-};
-
 namespace UE::Interchange
 {
 	//Get the mesh node context for each MeshUids
@@ -131,7 +108,7 @@ public:
 	};
 
 private:
-	TUniquePtr<FScopedSkeletalMeshReimportUtility> ScopedReimportUtility;
+	TUniquePtr<FSkinnedMeshComponentRecreateRenderStateContext> ScopedReimportUtility;
 
 	TMap<int32, FLodPayloads> PayloadsPerLodIndex;
 
