[Interchange] uFBX - disabled MeshConversion module when WITH_ENGINE == 0 to  fix  issues with non-unity compilation of InterchangeWorker

#issue 970476 - Only IterchangeWorker's compilation error
#rnx

diff --git a/Engine/Plugins/Interchange/Runtime/Source/Parsers/Fbx/InterchangeFbxParser.Build.cs b/Engine/Plugins/Interchange/Runtime/Source/Parsers/Fbx/InterchangeFbxParser.Build.cs
--- a/Engine/Plugins/Interchange/Runtime/Source/Parsers/Fbx/InterchangeFbxParser.Build.cs
+++ b/Engine/Plugins/Interchange/Runtime/Source/Parsers/Fbx/InterchangeFbxParser.Build.cs
@@ -40,7 +40,6 @@ namespace UnrealBuildTool.Rules
 					"InterchangeCommonParser",
 					"InterchangeMessages",
 					"InterchangeNodes",
-					"MeshConversion",
 					"MeshDescription",
 					"SkeletalMeshDescription",
 					"StaticMeshDescription", 
@@ -49,16 +48,19 @@ namespace UnrealBuildTool.Rules
 
 			if (Target.bCompileAgainstEngine)
 			{
+				PrivateDefinitions.Add("WITH_MESH_DESCRIPTION_BUILDER");
+
 				PrivateDependencyModuleNames.AddRange(
 					new string[] {
 						"Engine",
-						"InterchangeEngine"
+						"InterchangeEngine",
+						"MeshConversion",
 					}
 				);
 
 				PrivateIncludePathModuleNames.AddRange(
 					new string[] {
-						"InterchangeImport"
+						"InterchangeImport",
 					}
 				);
 				
diff --git a/Engine/Plugins/Interchange/Runtime/Source/Parsers/Fbx/Private/Ufbx/UfbxMesh.cpp b/Engine/Plugins/Interchange/Runtime/Source/Parsers/Fbx/Private/Ufbx/UfbxMesh.cpp
--- a/Engine/Plugins/Interchange/Runtime/Source/Parsers/Fbx/Private/Ufbx/UfbxMesh.cpp
+++ b/Engine/Plugins/Interchange/Runtime/Source/Parsers/Fbx/Private/Ufbx/UfbxMesh.cpp
@@ -9,7 +9,10 @@
 #include "InterchangeMeshNode.h"
 #include "Nodes/InterchangeBaseNodeContainer.h"
 
+#ifdef WITH_MESH_DESCRIPTION_BUILDER
 #include "MeshDescriptionBuilder.h"
+#endif 
+
 #include "SkeletalMeshAttributes.h"
 #include "StaticMeshAttributes.h"
 #include "StaticMeshOperations.h"
@@ -33,6 +36,7 @@ namespace UE::Interchange::Private
 	class FFaceCollector
 	{
 	public:
+#ifdef WITH_MESH_DESCRIPTION_BUILDER
 		FFaceCollector(FMeshDescriptionBuilder& InBuilder, ufbx_mesh& InMesh, const int32 UvCount)
 			: Builder(InBuilder)
 			, Mesh(InMesh)
@@ -56,6 +60,7 @@ namespace UE::Interchange::Private
 				}
 			}
 		}
+#endif
 
 		int32 TriangulateFace(const ufbx_face& Face)
 		{
@@ -100,6 +105,7 @@ namespace UE::Interchange::Private
 					continue;
 				}
 				
+#ifdef WITH_MESH_DESCRIPTION_BUILDER
 				for (int32 Vertex = 0; Vertex < 3; ++Vertex)
 				{
 					FVertexInstanceID VertexInstanceID = Builder.AppendInstance(VertexMapping[SrcVertexIndices[Vertex]]);
@@ -142,13 +148,16 @@ namespace UE::Interchange::Private
 					                                 UvIndexMap[UvSet.vertex_uv.indices[TriangulatedIndices[2]]],
 					                                 UvLayerIndex);
 				}
+#endif
 			}
 		};
 
 		// Map ufbx vertices index to FMeshDescription's FVertexID
 		TArray<FVertexID> VertexMapping;
 
+#ifdef WITH_MESH_DESCRIPTION_BUILDER
 		FMeshDescriptionBuilder& Builder;
+#endif
 		const ufbx_mesh& Mesh;
 
 		// Temp buffer to keep triangulated face indices 
@@ -409,6 +418,7 @@ namespace UE::Interchange::Private
 		constexpr bool bKeepExistingAttribute = true;
 		Attributes.Register(bKeepExistingAttribute);
 
+#ifdef WITH_MESH_DESCRIPTION_BUILDER
 		FMeshDescriptionBuilder Builder;
 		Builder.SetMeshDescription(&MeshDescription);
 		Builder.SuspendMeshDescriptionIndexing();
@@ -470,6 +480,8 @@ namespace UE::Interchange::Private
 		HandleMesh(FaceCollector);
 
 		Builder.ResumeMeshDescriptionIndexing();
+#endif
+
 
 		FStaticMeshOperations::ApplyTransform(MeshDescription, MeshGlobalTransform, true);
 		return true;
