[Interchange] uFBX - fix compilation error with newer Windows SDK

- INFINITY macro was causing error C4756: overflow in constant arithmetic

diff --git a/Engine/Plugins/Interchange/Runtime/Source/Parsers/Fbx/InterchangeFbxParser.Build.cs b/Engine/Plugins/Interchange/Runtime/Source/Parsers/Fbx/InterchangeFbxParser.Build.cs
--- a/Engine/Plugins/Interchange/Runtime/Source/Parsers/Fbx/InterchangeFbxParser.Build.cs
+++ b/Engine/Plugins/Interchange/Runtime/Source/Parsers/Fbx/InterchangeFbxParser.Build.cs
@@ -8,6 +8,9 @@ namespace UnrealBuildTool.Rules
 	{
 		public InterchangeFbxParser(ReadOnlyTargetRules Target) : base(Target)
 		{
+			// see _UCRT_LEGACY_INFINITY below. Generate own pch to have this define used
+			PCHUsage = PCHUsageMode.NoPCHs;
+
 			if (Target.Type != TargetType.Editor && Target.Type != TargetType.Program)
 			{
 				throw new BuildException("InterchangeFbxParser module can be build only on editor or program target.");
@@ -68,9 +71,14 @@ namespace UnrealBuildTool.Rules
 
 			AddEngineThirdPartyPrivateStaticDependencies(Target, "FBX");
 
+			// ufbx
 			string UfbxDir = Path.GetFullPath(Path.Combine(ModuleDirectory, "..", "ThirdParty", "uFBX", "0.18.2"));
 			PublicSystemIncludePaths.Add(UfbxDir);
 			PublicDebugVisualizerPaths.Add(Path.Combine(UfbxDir, "misc", "ufbx.natvis"));
+			// Allow compilation of ufbx.c with newer Windows SDK(10.0.26100) where INFINITY is defined as (float)1e+300 which caused
+			// error C4756: overflow in constant arithmetic
+			PrivateDefinitions.Add("_UCRT_LEGACY_INFINITY");
+
 		}
 	}
 }
