Fixed N^2 behavior inside UInterchangeGenericLevelPipeline::ExecuteSceneNodePreImport.

Worst case for this issue was using a lot (100+) of StaticMesh nodes that have LODs. Or use a lot of SkeletalMesh (100+)

The solution was to cache factories inside maps and sets so it's faster to access them during iteration. As a result, memory usage is increased a bit.

Performance Comparison (Before vs After) :

x10000 SkeletalMesh (USDA)
- ExecuteSceneNodePreImport : 65s vs 17s
- Time before import dialog shows up : 103s vs 53s

x10000 Mesh + LOD (USDA)
- ExecuteSceneNodePreImport : 632s vs 12s
- Time before import dialog shows up : 673s vs 49s

Note that before CL-44137550 was commited, ExecuteSceneNodePreImport time was counted twice.
Thanks to [REDACTED_PERSON] for pointing out this issue.

#rnx UInterchangeGenericLevelPipeline::ExecuteSceneNodePreImport loop optimization

diff --git a/Engine/Plugins/Interchange/Runtime/Source/Pipelines/Private/InterchangeGenericScenesPipeline.cpp b/Engine/Plugins/Interchange/Runtime/Source/Pipelines/Private/InterchangeGenericScenesPipeline.cpp
--- a/Engine/Plugins/Interchange/Runtime/Source/Pipelines/Private/InterchangeGenericScenesPipeline.cpp
+++ b/Engine/Plugins/Interchange/Runtime/Source/Pipelines/Private/InterchangeGenericScenesPipeline.cpp
@@ -263,22 +263,11 @@ void UInterchangeGenericLevelPipeline::ExecutePipeline(UInterchangeBaseNodeConta
 	}
 
 	TArray<UInterchangeSceneNode*> SceneNodes;
+	TMap<FString, const UInterchangeStaticMeshFactoryNode*> LodGroupFactoryNodeLookUp;
+	TSet<FString> SkinnedMeshNodes;
 
-	//Find all translated node we need for this pipeline
-	BaseNodeContainer->IterateNodes([&SceneNodes](const FString& NodeUid, UInterchangeBaseNode* Node)
-	{
-		switch(Node->GetNodeContainerType())
-		{
-		case EInterchangeNodeContainerType::TranslatedScene:
-		{
-			if (UInterchangeSceneNode* SceneNode = Cast<UInterchangeSceneNode>(Node))
-			{
-				SceneNodes.Add(SceneNode);
-			}
-		}
-		break;
-		}
-	});
+	// Find all translated node we need for this pipeline
+	PopulateSceneNodesPreImport(SceneNodes, LodGroupFactoryNodeLookUp, SkinnedMeshNodes);
 
 #if WITH_EDITORONLY_DATA
 	const FString FilePath = FPaths::ConvertRelativePathToFull(InSourceDatas[0]->GetFilename());
@@ -354,7 +343,7 @@ void UInterchangeGenericLevelPipeline::ExecutePipeline(UInterchangeBaseNodeConta
 	CacheActiveJointUids();
 
 	//Cache any lod group data, this way we can add actor only for the lod group
-	TMap< const UInterchangeSceneNode*, TArray<const UInterchangeSceneNode*>> SceneNodesPerLodGroupNode;
+	TMap<const UInterchangeSceneNode*, TArray<const UInterchangeSceneNode*>> SceneNodesPerLodGroupNode;
 	for (const UInterchangeSceneNode* SceneNode : SceneNodes)
 	{
 		if (!SceneNode)
@@ -448,7 +437,8 @@ void UInterchangeGenericLevelPipeline::ExecutePipeline(UInterchangeBaseNodeConta
 				continue;
 			}
 
-			ExecuteSceneNodePreImport(GlobalOffsetTransform, SceneNode);
+			const FSceneNodePreImportData SceneNodeData(SceneNode, LodGroupFactoryNodeLookUp.FindRef(SceneNode->GetUniqueID()), SkinnedMeshNodes.Contains(SceneNode->GetUniqueID()));
+			ExecuteSceneNodePreImport(GlobalOffsetTransform, SceneNodeData);
 		}
 	}
 
@@ -469,10 +459,93 @@ void UInterchangeGenericLevelPipeline::ExecutePipeline(UInterchangeBaseNodeConta
 	}
 }
 
-void UInterchangeGenericLevelPipeline::ExecuteSceneNodePreImport(const FTransform& GlobalOffsetTransform, const UInterchangeSceneNode* SceneNode)
+void UInterchangeGenericLevelPipeline::PopulateSceneNodesPreImport(TArray<UInterchangeSceneNode*>& SceneNodes,
+TMap<FString, const UInterchangeStaticMeshFactoryNode*>& LodGroupFactoryNodeLookUp,
+TSet<FString>& SkinnedMeshNodes) const
+{
+	BaseNodeContainer->IterateNodes([this, &SceneNodes, &LodGroupFactoryNodeLookUp, &SkinnedMeshNodes](const FString& NodeUid, UInterchangeBaseNode* Node)
+	{
+		switch (Node->GetNodeContainerType())
+		{
+		case EInterchangeNodeContainerType::TranslatedScene:
+		{
+			if (UInterchangeSceneNode* SceneNode = Cast<UInterchangeSceneNode>(Node))
+			{
+				SceneNodes.Add(SceneNode);
+			}
+		}
+		case EInterchangeNodeContainerType::FactoryData:
+		{
+			if (UInterchangeStaticMeshFactoryNode* StaticMeshFactoryNode = Cast<UInterchangeStaticMeshFactoryNode>(Node))
+			{
+				TArray<FString> LodDataUids;
+				StaticMeshFactoryNode->GetLodDataUniqueIds(LodDataUids);
+				if (!LodDataUids.IsEmpty())
+				{
+					if (UInterchangeStaticMeshLodDataNode* LodDataNode = Cast<UInterchangeStaticMeshLodDataNode>(BaseNodeContainer->GetFactoryNode(LodDataUids[0])))
+					{
+						TArray<FString> LodDataMeshUids;
+						LodDataNode->GetMeshUids(LodDataMeshUids);
+						if (!LodDataMeshUids.IsEmpty())
+						{
+							FString CurrentNodeUID = LodDataMeshUids[0];
+
+							// Every ancestor of LodDataMeshUids[0] uses the same FactoryNode
+							while (CurrentNodeUID != UInterchangeBaseNode::InvalidNodeUid())
+							{
+								const UInterchangeBaseNode* CurrentNode = BaseNodeContainer->GetNode(CurrentNodeUID);
+								if (!CurrentNode)
+								{
+									break;
+								}
+								
+								const FString& ParentUID = CurrentNode->GetParentUid();
+								if (ParentUID.IsEmpty())
+								{
+									break;
+								}
+
+								LodGroupFactoryNodeLookUp.Add(ParentUID, StaticMeshFactoryNode);
+
+								CurrentNodeUID = ParentUID;
+							}
+						}
+					}
+				}
+			}
+			else if (UInterchangeSkeletalMeshFactoryNode* SkeletalMeshFactoryNode = Cast<UInterchangeSkeletalMeshFactoryNode>(Node))
+			{
+				int32 LodCount = SkeletalMeshFactoryNode->GetLodDataCount();
+				TArray<FString> LodDataUniqueIds;
+				SkeletalMeshFactoryNode->GetLodDataUniqueIds(LodDataUniqueIds);
+				for (int32 LodIndex = 0; LodIndex < LodCount; ++LodIndex)
+				{
+					FString LodUniqueId = LodDataUniqueIds[LodIndex];
+					const UInterchangeSkeletalMeshLodDataNode* LodDataNode = Cast<UInterchangeSkeletalMeshLodDataNode>(BaseNodeContainer->GetNode(LodUniqueId));
+					if (LodDataNode)
+					{
+						TArray<FString> MeshUids;
+						LodDataNode->GetMeshUids(MeshUids);
+
+						for (FString Uid : MeshUids)
+						{
+							SkinnedMeshNodes.Add(Uid);
+						}
+					}
+				}
+			}
+		}
+		break;
+		}
+	});
+}
+
+void UInterchangeGenericLevelPipeline::ExecuteSceneNodePreImport(const FTransform& GlobalOffsetTransform, const FSceneNodePreImportData& SceneNodeData)
 {
 	using namespace UE::Interchange;
 
+	const UInterchangeSceneNode* SceneNode = SceneNodeData.SceneNode;
+
 	if (!BaseNodeContainer || !SceneNode)
 	{
 		return;
@@ -509,37 +582,13 @@ void UInterchangeGenericLevelPipeline::ExecuteSceneNodePreImport(const FTransfor
 		{
 			TranslatedAssetNode = BaseNodeContainer->GetNode(AssetInstanceUid);
 		}
+
 		if (const UInterchangeMeshNode* MeshNode = Cast<UInterchangeMeshNode>(TranslatedAssetNode))
 		{
-			bool bIsSkinnedMesh = MeshNode->IsSkinnedMesh();
-			if(!bIsSkinnedMesh)
-			{
-				//In case we have a rigid mesh (static mesh having morph target...), we should find a skeletalmesh factory node holding this mesh node
-				BaseNodeContainer->BreakableIterateNodesOfType<UInterchangeSkeletalMeshFactoryNode>([this, &SceneNode, &bIsSkinnedMesh](const FString& NodeUid, UInterchangeSkeletalMeshFactoryNode* SkeletalMeshNode)
-				{
-					int32 LodCount = SkeletalMeshNode->GetLodDataCount();
-					TArray<FString> LodDataUniqueIds;
-					SkeletalMeshNode->GetLodDataUniqueIds(LodDataUniqueIds);
-					for (int32 LodIndex = 0; LodIndex < LodCount; ++LodIndex)
-					{
-						FString LodUniqueId = LodDataUniqueIds[LodIndex];
-						const UInterchangeSkeletalMeshLodDataNode* LodDataNode = Cast<UInterchangeSkeletalMeshLodDataNode>(BaseNodeContainer->GetNode(LodUniqueId));
-						if (LodDataNode)
-						{
-							TArray<FString> MeshUids;
-							LodDataNode->GetMeshUids(MeshUids);
-							if (MeshUids.Contains(SceneNode->GetUniqueID()))
-							{
-								bIsSkinnedMesh = true;
-								return bIsSkinnedMesh;
-							}
-						}
-					}
-					return false;
-				});
-			}
-			//Skinned mesh are added when the bRootJointNode is true.
-			//In this case we dont want to add an empty staticmesh actor.
+			bool bIsSkinnedMesh = SceneNodeData.bIsSkinnedMesh || MeshNode->IsSkinnedMesh();
+			
+			// Skinned mesh are added when the bRootJointNode is true.
+			// In this case we dont want to add an empty staticmesh actor.
 			if (bIsSkinnedMesh)
 			{
 				return;
@@ -568,46 +617,43 @@ void UInterchangeGenericLevelPipeline::ExecuteSceneNodePreImport(const FTransfor
 			return SpecializeTypes.Contains(UE::Interchange::FSceneNodeStaticData::GetLodGroupSpecializeTypeString());
 		}();
 
-	UInterchangeStaticMeshFactoryNode* LodGroupStaticMeshFactoryNode = nullptr;
+	const UInterchangeStaticMeshFactoryNode* LodGroupStaticMeshFactoryNode = SceneNodeData.LodGroupFactoryNode;
 	if (bIsLodGroup)
-	{
-		//Find the staticmesh factory node create for this lod group
-		BaseNodeContainer->BreakableIterateNodesOfType<UInterchangeStaticMeshFactoryNode>([this, &LodGroupStaticMeshFactoryNode, &SceneNode, &TranslatedAssetNode](const FString& NodeUid, UInterchangeStaticMeshFactoryNode* StaticMeshFactoryNode)
+	{	
+		if (LodGroupStaticMeshFactoryNode)
+		{
+			TArray<FString> LodDataUids;
+			LodGroupStaticMeshFactoryNode->GetLodDataUniqueIds(LodDataUids);
+			if (!LodDataUids.IsEmpty())
 			{
-				TArray<FString> LodDataUids;
-				StaticMeshFactoryNode->GetLodDataUniqueIds(LodDataUids);
-				if (!LodDataUids.IsEmpty())
+				if (UInterchangeStaticMeshLodDataNode* LodDataNode = Cast<UInterchangeStaticMeshLodDataNode>(BaseNodeContainer->GetFactoryNode(LodDataUids[0])))
 				{
-					if (UInterchangeStaticMeshLodDataNode* LodDataNode = Cast<UInterchangeStaticMeshLodDataNode>(BaseNodeContainer->GetFactoryNode(LodDataUids[0])))
+					TArray<FString> LodDataMeshUids;
+					LodDataNode->GetMeshUids(LodDataMeshUids);
+					if (!LodDataMeshUids.IsEmpty())
 					{
-						TArray<FString> LodDataMeshUids;
-						LodDataNode->GetMeshUids(LodDataMeshUids);
-						if (!LodDataMeshUids.IsEmpty())
+						if (BaseNodeContainer->GetIsAncestor(LodDataMeshUids[0], SceneNode->GetUniqueID()))
 						{
-							if (BaseNodeContainer->GetIsAncestor(LodDataMeshUids[0], SceneNode->GetUniqueID()))
+							//Set the first lod meshuid has the TranslatedAssetNode so it create the mesh actor factory node,
+							//and setup it properly
+							if (const UInterchangeSceneNode* LodDataMeshNode = Cast<UInterchangeSceneNode>(BaseNodeContainer->GetNode(LodDataMeshUids[0])))
 							{
-								LodGroupStaticMeshFactoryNode = StaticMeshFactoryNode;
-								//Set the first lod meshuid has the TranslatedAssetNode so it create the mesh actor factory node,
-								//and setup it properly
-								if (const UInterchangeSceneNode* LodDataMeshNode = Cast<UInterchangeSceneNode>(BaseNodeContainer->GetNode(LodDataMeshUids[0])))
-								{
-									FString AssetUid;
-									if (LodDataMeshNode->GetCustomAssetInstanceUid(AssetUid))
-									{
-										TranslatedAssetNode = BaseNodeContainer->GetNode(AssetUid);
-									}
-								}
-								else
+								FString AssetUid;
+								if (LodDataMeshNode->GetCustomAssetInstanceUid(AssetUid))
 								{
-									TranslatedAssetNode = BaseNodeContainer->GetNode(LodDataMeshUids[0]);
+									TranslatedAssetNode = BaseNodeContainer->GetNode(AssetUid);
 								}
-								return true;
+							}
+							else
+							{
+								TranslatedAssetNode = BaseNodeContainer->GetNode(LodDataMeshUids[0]);
 							}
 						}
 					}
 				}
-				return false;
-			});
+			}
+		}
+
 		if (!LodGroupStaticMeshFactoryNode || !TranslatedAssetNode)
 		{
 			//Skip this lod group if there is no associated mesh
diff --git a/Engine/Plugins/Interchange/Runtime/Source/Pipelines/Public/InterchangeGenericScenesPipeline.h b/Engine/Plugins/Interchange/Runtime/Source/Pipelines/Public/InterchangeGenericScenesPipeline.h
--- a/Engine/Plugins/Interchange/Runtime/Source/Pipelines/Public/InterchangeGenericScenesPipeline.h
+++ b/Engine/Plugins/Interchange/Runtime/Source/Pipelines/Public/InterchangeGenericScenesPipeline.h
@@ -17,6 +17,7 @@ class ALevelInstance;
 class FPreviewScene;
 class UInterchangeActorFactoryNode;
 class UInterchangeLevelFactoryNode;
+class UInterchangeStaticMeshFactoryNode;
 class UInterchangeLevelInstanceActorFactoryNode;
 class UInterchangeSceneNode;
 class UInterchangeSceneVariantSetsNode;
@@ -82,10 +83,21 @@ protected:
 	}
 	/** END UInterchangePipelineBase overrides */
 
+	struct FSceneNodePreImportData
+	{
+		FSceneNodePreImportData(const UInterchangeSceneNode* InSceneNode, const UInterchangeStaticMeshFactoryNode* InLodGroupFactoryNode, bool bInIsSkinnedMesh)
+			: SceneNode(InSceneNode), LodGroupFactoryNode(InLodGroupFactoryNode), bIsSkinnedMesh(bInIsSkinnedMesh) {
+		}
+
+		const UInterchangeSceneNode* SceneNode;
+		const UInterchangeStaticMeshFactoryNode* LodGroupFactoryNode;
+		bool bIsSkinnedMesh;
+	};
+
 	/**
 	 * PreImport step called for each translated SceneNode.
 	 */
-	UE_API virtual void ExecuteSceneNodePreImport(const FTransform& GlobalOffsetTransform, const UInterchangeSceneNode* SceneNode);
+	UE_API virtual void ExecuteSceneNodePreImport(const FTransform& GlobalOffsetTransform, const FSceneNodePreImportData& SceneNodeData);
 
 	/**
 	 * PreImport step called for each translated SceneVariantSetNode.
@@ -101,6 +113,10 @@ protected:
 	 * Use to set up the given factory node's attributes after its initialization.
 	 */
 	UE_API virtual void SetUpFactoryNode(UInterchangeActorFactoryNode* ActorFactoryNode, const UInterchangeSceneNode* SceneNode, const UInterchangeBaseNode* TranslatedAssetNode) const;
+
+	UE_API virtual void PopulateSceneNodesPreImport(TArray<UInterchangeSceneNode*>& SceneNodes,
+		TMap<FString, const UInterchangeStaticMeshFactoryNode*>& StaticMeshFactoryNodesBySceneNodeUid,
+		TSet<FString>& SkinnedMeshNodes) const;
 	
 	/** Disable this option to not convert Standard(Perspective) to Physical Cameras*/
 	UPROPERTY(EditAnywhere, BlueprintReadWrite, Category = "Common Skeletal Meshes and Animations")
