[BUGFIX] Interchange Skeletal Mesh Factory
- Skin Weight Attributes are properly handled in case there are existing profiles.

#rnx

diff --git a/Engine/Plugins/Interchange/Runtime/Source/Import/Private/Mesh/InterchangeSkeletalMeshFactory.cpp b/Engine/Plugins/Interchange/Runtime/Source/Import/Private/Mesh/InterchangeSkeletalMeshFactory.cpp
--- a/Engine/Plugins/Interchange/Runtime/Source/Import/Private/Mesh/InterchangeSkeletalMeshFactory.cpp
+++ b/Engine/Plugins/Interchange/Runtime/Source/Import/Private/Mesh/InterchangeSkeletalMeshFactory.cpp
@@ -2176,13 +2176,13 @@ void UInterchangeSkeletalMeshFactory::SetupObject_GameThread(const FSetupObjectP
 					if (ExistingMeshDescriptionSrc.Vertices().Num() == PointNumberDest)
 					{
 						FSkeletalMeshAttributes AttributesDest(MeshDescriptionDest);
-						if (AttributesDest.GetSkinWeightProfileNames().Find(ProfileName))
+						if (AttributesDest.GetSkinWeightProfileNames().Contains(ProfileName))
 						{
 							AttributesDest.UnregisterSkinWeightAttribute(ProfileName);
 						}
 
 						FSkeletalMeshConstAttributes AttributesExistingSrc(ExistingMeshDescriptionSrc);
-						if (AttributesExistingSrc.GetSkinWeightProfileNames().Find(ProfileName))
+						if (AttributesExistingSrc.GetSkinWeightProfileNames().Contains(ProfileName))
 						{
 							AttributesDest.RegisterSkinWeightAttribute(ProfileName);
 							//Fill the SkinWeightsAttribute
