Interchange:
- Use material instances for OBJ and FBX Translators by using MaterialInstanceNode instead of shader graph nodes
- Add a Phong Surface Mask Material, for special cases where FBX use textures as masks
- Add PhongTranslucent parent material for material instances
- Modify the Phong and CommonMaterial with a static switch parameter, in case instead of Normal maps the format uses a BumpMap texture

EngineTest:
- Update the tests to reflect now that FBX and OBJ always import Material Instances

#virtualized

diff --git a/Engine/Plugins/Interchange/Runtime/Source/Import/Private/Mesh/InterchangeOBJTranslator.cpp b/Engine/Plugins/Interchange/Runtime/Source/Import/Private/Mesh/InterchangeOBJTranslator.cpp
--- a/Engine/Plugins/Interchange/Runtime/Source/Import/Private/Mesh/InterchangeOBJTranslator.cpp
+++ b/Engine/Plugins/Interchange/Runtime/Source/Import/Private/Mesh/InterchangeOBJTranslator.cpp
@@ -4,28 +4,21 @@
 
 #include "Algo/BinarySearch.h"
 #include "Algo/Find.h"
-#include "Async/Async.h"
-#include "Async/Future.h"
-#include "InterchangeImportLog.h"
-#include "InterchangeManager.h"
 #include "InterchangeMaterialDefinitions.h"
 #include "InterchangeMeshDefinitions.h"
 #include "InterchangeMeshNode.h"
-#include "InterchangeShaderGraphNode.h"
+#include "InterchangeMaterialInstanceNode.h"
 #include "InterchangeTexture2DNode.h"
 #include "InterchangeTextureNode.h"
 #include "InterchangeTranslatorHelper.h"
 #include "MaterialDomain.h"
 #include "Materials/Material.h"
-#include "Misc/ConfigCacheIni.h"
 #include "Misc/FileHelper.h"
 #include "Misc/Paths.h"
-#include "Misc/ScopedSlowTask.h"
-#include "Modules/ModuleManager.h"
 #include "Nodes/InterchangeBaseNodeContainer.h"
+#include "Nodes/InterchangeSourceNode.h"
 #include "StaticMeshAttributes.h"
 #include "StaticMeshOperations.h"
-#include "UObject/GCObjectScopeGuard.h"
 #include "UVMapSettings.h"
 #include "Material/InterchangeMaterialNodesHelper.h"
 
@@ -1033,7 +1026,7 @@ namespace ObjTranslatorUtils
 	 * Build and return a UID name for a shader graph node.
 	 * @todo: move this to be a static method on the ShaderGraphNode class itself?
 	 */
-	static FString MakeShaderGraphNodeName(const FString& Name)
+	static FString MakeMaterialInstanceNodeName(const FString& Name)
 	{
 		return TEXT("MAT_") + (Name.IsEmpty() ? FString(TEXT("Null")) : Name);
 	}
@@ -1049,7 +1042,7 @@ namespace ObjTranslatorUtils
 	}
 
 
-	static const UInterchangeTexture2DNode* CreateTexture2DNode(UInterchangeBaseNodeContainer& BaseNodeContainer, FString TexturePath, FString& InOutTextureName) 
+	static const UInterchangeTexture2DNode* CreateTexture2DNode(UInterchangeBaseNodeContainer& BaseNodeContainer, const FString& TexturePath, FString& InOutTextureName) 
 	{
 		FString TextureName = InOutTextureName;
 		for (int32 Index = 0;;Index ++)
@@ -1081,154 +1074,87 @@ namespace ObjTranslatorUtils
 		}
 	}
 
-	static UInterchangeShaderNode* CreateMaterialTextureSampleNode(UInterchangeBaseNodeContainer& BaseNodeContainer, UInterchangeShaderGraphNode* ShaderGraphNode, const FObjData::FTextureParameter& TextureParameter) 
+	static const UInterchangeTexture2DNode* CreateTexture2DNode(UInterchangeBaseNodeContainer& BaseNodeContainer, const FObjData::FTextureParameter& TextureParameter)
 	{
+		const UInterchangeTexture2DNode* TextureNode = nullptr;
 		if (!TextureParameter.IsEmpty())
 		{
-			// @todo: better way to handle textures than implementing the TexturePayload interface in the OBJTranslator
-			// and creating a temporary texture translator there.
-
-			using namespace UE::Interchange::Materials::Standard::Nodes;
-
 			FString TextureName = MakeTextureNodeName(FPaths::GetBaseFilename(TextureParameter.Path));
-
-			if (const UInterchangeTexture2DNode* TextureNode = CreateTexture2DNode(BaseNodeContainer, TextureParameter.Path, TextureName))
-			{
-				UInterchangeShaderNode* TextureSampleShaderNode = UInterchangeShaderNode::Create(&BaseNodeContainer, TextureName, ShaderGraphNode->GetUniqueID());
-				TextureSampleShaderNode->SetCustomShaderType(TextureSample::Name.ToString());
-				TextureSampleShaderNode->AddStringInput(TextureSample::Inputs::Texture.ToString(), TextureNode->GetUniqueID(), /*bIsAParameter =*/ true);
-				return TextureSampleShaderNode;
-			}
+			TextureNode = CreateTexture2DNode(BaseNodeContainer, TextureParameter.Path, TextureName);
 		}
-		return nullptr;
+		return TextureNode;
 	}
 
 	// multiply texture by color, handling cases when any of the arguments is missing
-	static bool AddTexturedColoredInput(UInterchangeBaseNodeContainer& BaseNodeContainer, UInterchangeShaderGraphNode* ShaderGraphNode, const FString& InputName, const FVector3f& Color, const FObjData::FTextureParameter& TexturePath)
+	static bool AddTexturedColoredInput(UInterchangeBaseNodeContainer& BaseNodeContainer, UInterchangeMaterialInstanceNode* MaterialInstanceNode, const FString& InputName, const FVector3f& Color, const FObjData::FTextureParameter& TexturePath)
 	{
-		UInterchangeShaderNode* TextureSampleShaderNode = CreateMaterialTextureSampleNode(BaseNodeContainer, ShaderGraphNode, TexturePath);
-
-		if (TextureSampleShaderNode)
+		if (const UInterchangeTexture2DNode* TextureNode = CreateTexture2DNode(BaseNodeContainer, TexturePath))
 		{
-			using namespace UE::Interchange::Materials::Standard::Nodes;
-
-			const FString MultiplierNodeName = InputName + TEXT("Multiply");
-			UInterchangeShaderNode* MultiplierNode = UInterchangeShaderNode::Create(&BaseNodeContainer, MultiplierNodeName, ShaderGraphNode->GetUniqueID());
-			MultiplierNode->SetCustomShaderType(Multiply::Name.ToString());
-
-			UInterchangeShaderPortsAPI::ConnectDefaultOuputToInput(MultiplierNode, Multiply::Inputs::A.ToString(), TextureSampleShaderNode->GetUniqueID());
-
-			const FString WeightNodeName = InputName + TEXT("MapWeight");
-			UInterchangeShaderNode* WeightNode = UInterchangeShaderNode::Create(&BaseNodeContainer, WeightNodeName, MultiplierNode->GetUniqueID());
-			WeightNode->SetCustomShaderType(ScalarParameter::Name.ToString());
-
-			const float WeightValue = 1.0f;
-			WeightNode->AddFloatAttribute(UInterchangeShaderPortsAPI::MakeInputParameterKey(ScalarParameter::Attributes::DefaultValue.ToString()), WeightValue);
-
-			UInterchangeShaderPortsAPI::ConnectDefaultOuputToInput(MultiplierNode, Multiply::Inputs::B.ToString(), WeightNode->GetUniqueID());
-
-			UInterchangeShaderPortsAPI::ConnectDefaultOuputToInput(MultiplierNode, Multiply::Inputs::A.ToString(), TextureSampleShaderNode->GetUniqueID());
-
 			const FString InputNameLabel = InputName + TEXT("Map");
-			TextureSampleShaderNode->SetDisplayLabel(InputNameLabel);
-			UInterchangeShaderPortsAPI::ConnectDefaultOuputToInput(ShaderGraphNode, InputName, MultiplierNode->GetUniqueID());
+			return MaterialInstanceNode->AddTextureParameterValue(InputNameLabel, TextureNode->GetUniqueID());
 		}
 		else if (IsColorInitialized(Color))
 		{
-			UE::Interchange::Materials::Private::FMaterialNodesHelper::SetupVectorParameter(BaseNodeContainer, ShaderGraphNode, InputName, FLinearColor(Color));
-		}
-		else
-		{
-			return false;
-		}
-		return true;
-	}
-
-	static UInterchangeShaderNode* MakeTexturedWeighted(UInterchangeBaseNodeContainer& BaseNodeContainer, UInterchangeShaderGraphNode* ShaderGraphNode, FString MapName, const float& Weight, UInterchangeShaderNode* TextureSampleShaderNode) 
-	{
-		if (!TextureSampleShaderNode)
-		{
-			return nullptr;
+			return MaterialInstanceNode->AddVectorParameterValue(InputName, Color);
 		}
 
-		using namespace UE::Interchange::Materials::Standard::Nodes;
-
-		const FString MultiplierNodeName = MapName + TEXT("Multiply");
-		UInterchangeShaderNode* MultiplierNode = UInterchangeShaderNode::Create(&BaseNodeContainer, MultiplierNodeName, ShaderGraphNode->GetUniqueID());
-		MultiplierNode->SetCustomShaderType(Multiply::Name.ToString());
-
-		UInterchangeShaderPortsAPI::ConnectDefaultOuputToInput(MultiplierNode, Multiply::Inputs::A.ToString(), TextureSampleShaderNode->GetUniqueID());
-
-		const FString WeightNodeName = MapName + TEXT("MapWeight");
-		UInterchangeShaderNode* WeightNode = UInterchangeShaderNode::Create(&BaseNodeContainer, WeightNodeName, MultiplierNode->GetUniqueID());
-		WeightNode->SetCustomShaderType(ScalarParameter::Name.ToString());
-
-		const float WeightValue = IsScalarInitialized(Weight) ? Weight : 1.0f;
-		WeightNode->AddFloatAttribute(UInterchangeShaderPortsAPI::MakeInputParameterKey(ScalarParameter::Attributes::DefaultValue.ToString()), WeightValue);
-
-		UInterchangeShaderPortsAPI::ConnectDefaultOuputToInput(MultiplierNode, Multiply::Inputs::B.ToString(), WeightNode->GetUniqueID());
-
-		return MultiplierNode;
+		return false;
 	}
 
-	static bool AddTexturedWeightedInput(UInterchangeBaseNodeContainer& BaseNodeContainer, UInterchangeShaderGraphNode* ShaderGraphNode, const FString& InputName, const float& Weight, UInterchangeShaderNode* TextureSampleShader) 
+	// multiply texture by scalar, handling cases when any of the arguments is missing
+	static bool AddTexturedWeightedInput(UInterchangeBaseNodeContainer& BaseNodeContainer, UInterchangeMaterialInstanceNode* MaterialInstanceNode, const FString& InputName, const float& Weight, const FObjData::FTextureParameter& TexturePath) 
 	{
-		if (UInterchangeShaderNode* TextureSampleWeighted = MakeTexturedWeighted(BaseNodeContainer, ShaderGraphNode, InputName, Weight, TextureSampleShader))
+		bool bResult = false;
+		if (const UInterchangeTexture2DNode* TextureNode = CreateTexture2DNode(BaseNodeContainer, TexturePath))
 		{
-			TextureSampleWeighted->SetDisplayLabel(InputName);
-			UInterchangeShaderPortsAPI::ConnectDefaultOuputToInput(ShaderGraphNode, InputName, TextureSampleWeighted->GetUniqueID());
+			const FString InputNameLabel = InputName + TEXT("MapWeight");
+			bResult |= MaterialInstanceNode->AddTextureParameterValue(InputNameLabel, TextureNode->GetUniqueID());
+
+			const float WeightValue = IsScalarInitialized(Weight) ? Weight : 1.0f;
+            bResult |= MaterialInstanceNode->AddScalarParameterValue(InputName, Weight);
 		}
 		else if (IsScalarInitialized(Weight))
 		{
-			UE::Interchange::Materials::Private::FMaterialNodesHelper::SetupScalarParameter(BaseNodeContainer, ShaderGraphNode, InputName, Weight);
-		}
-		else
-		{
-			return false;
+			bResult |= MaterialInstanceNode->AddScalarParameterValue(InputName, Weight);
 		}
-		return true;
-	}
 
-	// multiply texture by scalar, handling cases when any of the arguments is missing
-	static bool AddTexturedWeightedInput(UInterchangeBaseNodeContainer& BaseNodeContainer, UInterchangeShaderGraphNode* ShaderGraphNode, const FString& InputType, const float& Weight, const FObjData::FTextureParameter& TexturePath) 
-	{
-		UInterchangeShaderNode* TextureSampleShaderNode = CreateMaterialTextureSampleNode(BaseNodeContainer, ShaderGraphNode, TexturePath);
-		
-		return AddTexturedWeightedInput(BaseNodeContainer, ShaderGraphNode, InputType, Weight, TextureSampleShaderNode);
+		return bResult;
 	}
 
-	static bool HandlePbrAttributes(UInterchangeBaseNodeContainer& BaseNodeContainer, const FObjData::FMaterialData& MaterialData, UInterchangeShaderGraphNode* ShaderGraphNode)
+	static bool HandlePbrAttributes(UInterchangeBaseNodeContainer& BaseNodeContainer, const FObjData::FMaterialData& MaterialData, UInterchangeMaterialInstanceNode* MaterialInstanceNode)
 	{
 		bool bPbr = !MaterialData.RoughnessTexture.IsEmpty() || IsScalarInitialized(MaterialData.Roughness)
-		         || !MaterialData.MetallicTexture.IsEmpty()  || IsScalarInitialized(MaterialData.Metallic);
+		         || !MaterialData.MetallicTexture.IsEmpty()  || IsScalarInitialized(MaterialData.Metallic)
+				 || IsScalarInitialized(MaterialData.ClearCoatThickness)
+			     || IsScalarInitialized(MaterialData.Anisotropy);
 
 		if (!bPbr)
 		{
 			return false;
 		}
 
-		AddTexturedColoredInput(BaseNodeContainer, ShaderGraphNode, Materials::PBRMR::Parameters::BaseColor.ToString(), MaterialData.DiffuseColor, MaterialData.DiffuseTexture);
-		AddTexturedWeightedInput(BaseNodeContainer, ShaderGraphNode, Materials::PBRMR::Parameters::Metallic.ToString(), MaterialData.Metallic, MaterialData.MetallicTexture);
+		AddTexturedColoredInput(BaseNodeContainer, MaterialInstanceNode, Materials::PBRMR::Parameters::BaseColor.ToString(), MaterialData.DiffuseColor, MaterialData.DiffuseTexture);
+		AddTexturedWeightedInput(BaseNodeContainer, MaterialInstanceNode, Materials::PBRMR::Parameters::Metallic.ToString(), MaterialData.Metallic, MaterialData.MetallicTexture);
 
 		const float SpecularValue = (MaterialData.SpecularColor.X + MaterialData.SpecularColor.Y + MaterialData.SpecularColor.Z) / 3;
-		AddTexturedWeightedInput(BaseNodeContainer, ShaderGraphNode, Materials::PBRMR::Parameters::Specular.ToString(), SpecularValue, MaterialData.SpecularTexture);
+		AddTexturedWeightedInput(BaseNodeContainer, MaterialInstanceNode, Materials::PBRMR::Parameters::Specular.ToString(), SpecularValue, MaterialData.SpecularTexture);
 
 		// Roughness (Glossiness, Shininess, Specular Power/Exponent in non-PBR terms)
-		AddTexturedWeightedInput(BaseNodeContainer, ShaderGraphNode, Materials::PBRMR::Parameters::Roughness.ToString(), MaterialData.Roughness, MaterialData.RoughnessTexture);
+		AddTexturedWeightedInput(BaseNodeContainer, MaterialInstanceNode, Materials::PBRMR::Parameters::Roughness.ToString(), MaterialData.Roughness, MaterialData.RoughnessTexture);
 
 		if (IsScalarInitialized(MaterialData.ClearCoatThickness) && !FMath::IsNearlyZero(MaterialData.ClearCoatThickness))
 		{
-			UE::Interchange::Materials::Private::FMaterialNodesHelper::SetupScalarParameter(BaseNodeContainer, ShaderGraphNode, Materials::ClearCoat::Parameters::ClearCoat.ToString(), MaterialData.ClearCoatThickness);
+			MaterialInstanceNode->AddScalarParameterValue(Materials::ClearCoat::Parameters::ClearCoat.ToString(), MaterialData.ClearCoatThickness);
 
 			if (IsScalarInitialized(MaterialData.ClearCoatRoughness))
 			{
-				UE::Interchange::Materials::Private::FMaterialNodesHelper::SetupScalarParameter(BaseNodeContainer, ShaderGraphNode, Materials::ClearCoat::Parameters::ClearCoatRoughness.ToString(), MaterialData.ClearCoatRoughness);
+				MaterialInstanceNode->AddScalarParameterValue(Materials::ClearCoat::Parameters::ClearCoatRoughness.ToString(), MaterialData.ClearCoatRoughness);
 			}
 		}
 
 		if (IsScalarInitialized(MaterialData.Anisotropy)  && !FMath::IsNearlyZero(MaterialData.Anisotropy))
 		{
-			UE::Interchange::Materials::Private::FMaterialNodesHelper::SetupScalarParameter(BaseNodeContainer, ShaderGraphNode, Materials::PBRMR::Parameters::Anisotropy.ToString(), MaterialData.Anisotropy);
+			MaterialInstanceNode->AddScalarParameterValue(Materials::PBRMR::Parameters::Anisotropy.ToString(), MaterialData.Anisotropy);
 		}
 
 		return true;
@@ -1347,8 +1273,8 @@ bool UInterchangeOBJTranslator::Translate(UInterchangeBaseNodeContainer& BaseNod
 			for (const FObjData::FFaceGroupData& FaceGroup : Group.Value.FaceGroups)
 			{
 				if (!FaceGroup.MaterialName.IsEmpty())
-				{
-					const FString MaterialDependencyUid = UInterchangeShaderGraphNode::MakeNodeUid(MakeShaderGraphNodeName(FaceGroup.MaterialName));
+				{					
+					const FString MaterialDependencyUid = UInterchangeMaterialInstanceNode::MakeNodeUid(MakeMaterialInstanceNodeName(FaceGroup.MaterialName), TEXT(""));
 					MeshNode->SetSlotMaterialDependencyUid(FaceGroup.MaterialName, MaterialDependencyUid);
 				}
 			}
@@ -1359,10 +1285,14 @@ bool UInterchangeOBJTranslator::Translate(UInterchangeBaseNodeContainer& BaseNod
 		{
 			const FString& MaterialName = Material.Key;
 			const FObjData::FMaterialData& MaterialData = Material.Value;
-			FString NodeUid = UInterchangeShaderGraphNode::MakeNodeUid(MakeShaderGraphNodeName(MaterialName));
+			FString NodeUid = UInterchangeMaterialInstanceNode::MakeNodeUid(MakeMaterialInstanceNodeName(MaterialName), TEXT(""));
 
-			UInterchangeShaderGraphNode * ShaderGraphNode = NewObject<UInterchangeShaderGraphNode>(&BaseNodeContainer);
-			BaseNodeContainer.SetupNode(ShaderGraphNode, NodeUid, MaterialName, EInterchangeNodeContainerType::TranslatedAsset);
+			UInterchangeMaterialInstanceNode* MaterialInstanceNode = NewObject<UInterchangeMaterialInstanceNode>(&BaseNodeContainer);
+			MaterialInstanceNode->SetCustomParent(TEXT("/InterchangeAssets/Materials/PBRSurfaceMaterial_MR.PBRSurfaceMaterial_MR"));
+			BaseNodeContainer.SetupNode(MaterialInstanceNode, NodeUid, MaterialName, EInterchangeNodeContainerType::TranslatedAsset);
+
+			bool bPBR = true;
+			bool bOpacity = false;
 
 			bool bUnlit = false;
 			bool bSpecular = true;
@@ -1371,8 +1301,9 @@ bool UInterchangeOBJTranslator::Translate(UInterchangeBaseNodeContainer& BaseNod
 			bool bTransmission = false;
 			bool bRefraction = false;
 
-			if (!HandlePbrAttributes(BaseNodeContainer, MaterialData, ShaderGraphNode))
+			if (!HandlePbrAttributes(BaseNodeContainer, MaterialData, MaterialInstanceNode))
 			{
+				bPBR = false;
 				// It's not a Pbr material - in this case parse the illumination model
 
 				// What is relevant to us about the illumination model in OBJ spec
@@ -1444,50 +1375,40 @@ bool UInterchangeOBJTranslator::Translate(UInterchangeBaseNodeContainer& BaseNod
 
 				if(bUnlit)
 				{
-					AddTexturedColoredInput(BaseNodeContainer, ShaderGraphNode, Materials::Unlit::Parameters::UnlitColor.ToString(), MaterialData.DiffuseColor, MaterialData.DiffuseTexture);
+					MaterialInstanceNode->SetCustomParent(TEXT("/InterchangeAssets/Materials/UnlitMaterial.UnlitMaterial"));
+					AddTexturedColoredInput(BaseNodeContainer, MaterialInstanceNode, Materials::Unlit::Parameters::UnlitColor.ToString(), MaterialData.DiffuseColor, MaterialData.DiffuseTexture);
 				}
 				else
 				{
-					AddTexturedColoredInput(BaseNodeContainer, ShaderGraphNode, Materials::Phong::Parameters::DiffuseColor.ToString(), (MaterialData.DiffuseColor.GetMin() < 0) ? FVector3f(0.8f, 0.8f, 0.8f) : MaterialData.DiffuseColor, MaterialData.DiffuseTexture);
+					MaterialInstanceNode->SetCustomParent(TEXT("/InterchangeAssets/Materials/PhongSurfaceMaterial.PhongSurfaceMaterial"));
+					AddTexturedColoredInput(BaseNodeContainer, MaterialInstanceNode, Materials::Phong::Parameters::DiffuseColor.ToString(), (MaterialData.DiffuseColor.GetMin() < 0) ? FVector3f(0.8f, 0.8f, 0.8f) : MaterialData.DiffuseColor, MaterialData.DiffuseTexture);
 
 					if (bSpecular)
 					{
-						AddTexturedColoredInput(BaseNodeContainer, ShaderGraphNode, Materials::Phong::Parameters::SpecularColor.ToString(), MaterialData.SpecularColor, MaterialData.SpecularTexture);
+						AddTexturedColoredInput(BaseNodeContainer, MaterialInstanceNode, Materials::Phong::Parameters::SpecularColor.ToString(), MaterialData.SpecularColor, MaterialData.SpecularTexture);
 
 						// OBJ Specular Exponent('Ns') is in range 0 to 1000(for sharpest highlight) and Interchange Shininess is up to 100
 						const float SpecularExponentToShininessFactor = 0.1f;
 						const FString ShininessParameterName = Materials::Phong::Parameters::Shininess.ToString();
-						if (UInterchangeShaderNode* SpecularExponentTextureSampleShader = CreateMaterialTextureSampleNode(BaseNodeContainer, ShaderGraphNode, MaterialData.SpecularExponentTexture))
+						if (const UInterchangeTexture2DNode* SpecularExponentTextureNode= CreateTexture2DNode(BaseNodeContainer, MaterialData.SpecularExponentTexture))
 						{
 							using namespace UE::Interchange::Materials::Standard::Nodes;
 
 							float SpecularExponentScale = (IsScalarInitialized(MaterialData.SpecularExponent) ? MaterialData.SpecularExponent : 1.0f) * SpecularExponentToShininessFactor;
-
-							const FString MultiplierNodeName = ShininessParameterName + TEXT("_Multiply");
-							UInterchangeShaderNode* MultiplierNode = UInterchangeShaderNode::Create(&BaseNodeContainer, MultiplierNodeName, ShaderGraphNode->GetUniqueID());
-							MultiplierNode->SetCustomShaderType(Multiply::Name.ToString());
-
-							UInterchangeShaderPortsAPI::ConnectDefaultOuputToInput(MultiplierNode, Multiply::Inputs::A.ToString(), SpecularExponentTextureSampleShader->GetUniqueID());
-							MultiplierNode->AddFloatAttribute(UInterchangeShaderPortsAPI::MakeInputValueKey( Multiply::Inputs::B.ToString() ), SpecularExponentScale);
-
-							UInterchangeShaderPortsAPI::ConnectDefaultOuputToInput(ShaderGraphNode, ShininessParameterName, MultiplierNode->GetUniqueID());
+							MaterialInstanceNode->AddTextureParameterValue(ShininessParameterName + TEXT("Map"), SpecularExponentTextureNode->GetUniqueID());
+							MaterialInstanceNode->AddScalarParameterValue(ShininessParameterName, SpecularExponentScale);
 						}
 						else if (IsScalarInitialized(MaterialData.SpecularExponent))
 						{
-							UE::Interchange::Materials::Private::FMaterialNodesHelper::SetupScalarParameter(BaseNodeContainer, ShaderGraphNode, ShininessParameterName, MaterialData.SpecularExponent);
+							MaterialInstanceNode->AddScalarParameterValue(ShininessParameterName, MaterialData.SpecularExponent);
 						}
 					}
 
 					const FString AmbientColorPropertyName = Materials::Phong::Parameters::AmbientColor.ToString();
-					if (bReflectionMetal)
-					{
-						// Set ambient to white as it drives the metallic parameter of UE Pbr material.
-						UE::Interchange::Materials::Private::FMaterialNodesHelper::SetupVectorParameter(BaseNodeContainer, ShaderGraphNode, AmbientColorPropertyName, FLinearColor::White);
-					}
-					else
+					if (!bReflectionMetal)
 					{
 						const FVector3f AmbientColor = MaterialData.AmbientColor.GetMin() < 0.f ? FVector3f(0.8f, 0.8f, 0.8f) : MaterialData.AmbientColor;
-						AddTexturedColoredInput(BaseNodeContainer, ShaderGraphNode, AmbientColorPropertyName, AmbientColor, MaterialData.AmbientTexture);
+						AddTexturedColoredInput(BaseNodeContainer, MaterialInstanceNode, AmbientColorPropertyName, AmbientColor, MaterialData.AmbientTexture);
 					}
 				}
 			}
@@ -1499,24 +1420,33 @@ bool UInterchangeOBJTranslator::Translate(UInterchangeBaseNodeContainer& BaseNod
 				const FString OpacityInputName = Materials::Common::Parameters::Opacity.ToString();
 
 				// Transparency texture replaces Transparency scalar
-				if (UInterchangeShaderNode* TransparencySample = CreateMaterialTextureSampleNode(BaseNodeContainer, ShaderGraphNode, MaterialData.TransparencyTexture))
+                if (const UInterchangeTexture2DNode* TransparencyTextureNode = CreateTexture2DNode(BaseNodeContainer, MaterialData.TransparencyTexture))
 				{
-					// Invert transparency to get Opacity
-					const FString OneMinusNodeName = OpacityInputName + TEXT("OneMinus");
-					UInterchangeShaderNode* OneMinusNode = UInterchangeShaderNode::Create( &BaseNodeContainer, OneMinusNodeName, ShaderGraphNode->GetUniqueID() );
-					OneMinusNode->SetCustomShaderType(Materials::Standard::Nodes::OneMinus::Name.ToString());
-					UInterchangeShaderPortsAPI::ConnectDefaultOuputToInput(OneMinusNode, Materials::Standard::Nodes::OneMinus::Inputs::Input.ToString(), TransparencySample->GetUniqueID());
-					UInterchangeShaderPortsAPI::ConnectDefaultOuputToInput(ShaderGraphNode, OpacityInputName, OneMinusNode->GetUniqueID());
+					bOpacity = true;
+                    MaterialInstanceNode->AddTextureParameterValue(OpacityInputName + TEXT("Map"), TransparencyTextureNode->GetUniqueID());
 				}
 				// Opacity texture is weighted by scalar opacity("d")
-				else if (!AddTexturedWeightedInput(BaseNodeContainer, ShaderGraphNode, OpacityInputName, MaterialData.Opacity, MaterialData.OpacityTexture))
+				else if (!AddTexturedWeightedInput(BaseNodeContainer, MaterialInstanceNode, OpacityInputName, MaterialData.Opacity, MaterialData.OpacityTexture))
 				{
 					// When no textures are present compute opacity value depending on which is defined - opacity or transparency
 					// OBJ has two ways to define Opacity/Transparency - "d"(for opacity) and "Tr"(for "inverse opacity" - transparency)
 					float OpacityScalar = IsScalarInitialized(MaterialData.Opacity) ? MaterialData.Opacity : (IsScalarInitialized(MaterialData.Transparency) ? 1.f - MaterialData.Transparency : -1.f);
 					if (IsScalarInitialized(OpacityScalar))
 					{
-						UE::Interchange::Materials::Private::FMaterialNodesHelper::SetupScalarParameter(BaseNodeContainer, ShaderGraphNode, OpacityInputName, OpacityScalar);
+						bOpacity = true;
+						MaterialInstanceNode->AddScalarParameterValue(OpacityInputName, OpacityScalar);
+					}
+				}
+
+				if (bOpacity)
+				{
+					if (bPBR)
+					{
+						MaterialInstanceNode->SetCustomParent(TEXT("/InterchangeAssets/Materials/ThinTranslucentMaterial_MR.ThinTranslucentMaterial_MR"));
+					}
+					else
+					{
+						MaterialInstanceNode->SetCustomParent(TEXT("/InterchangeAssets/Materials/PhongSurfaceTranslucentMaterial.PhongSurfaceTranslucentMaterial"));
 					}
 				}
 			}
@@ -1527,7 +1457,7 @@ bool UInterchangeOBJTranslator::Translate(UInterchangeBaseNodeContainer& BaseNod
 			}
 
 			// Emissive
-			AddTexturedColoredInput(BaseNodeContainer, ShaderGraphNode, Materials::Common::Parameters::EmissiveColor.ToString(), MaterialData.EmissiveColor, MaterialData.EmissiveTexture);
+			AddTexturedColoredInput(BaseNodeContainer, MaterialInstanceNode, Materials::Common::Parameters::EmissiveColor.ToString(), MaterialData.EmissiveColor, MaterialData.EmissiveTexture);
 				
 			// Ior
 			if (IsScalarInitialized(MaterialData.IndexOfRefraction) || bRefraction)
@@ -1535,42 +1465,33 @@ bool UInterchangeOBJTranslator::Translate(UInterchangeBaseNodeContainer& BaseNod
 				const float DefaultIndexOfRefraction = 1.52f;  // Use glass index of refraction by default
 				const float IndexOfRefraction = IsScalarInitialized(MaterialData.IndexOfRefraction) ? MaterialData.IndexOfRefraction : DefaultIndexOfRefraction;
 				const FString ParameterName = Materials::Common::Parameters::IndexOfRefraction.ToString();
-				UE::Interchange::Materials::Private::FMaterialNodesHelper::SetupScalarParameter(BaseNodeContainer, ShaderGraphNode, ParameterName, IndexOfRefraction);
+				MaterialInstanceNode->AddScalarParameterValue(ParameterName, IndexOfRefraction);
 			}
 
 			if ((IsColorInitialized(MaterialData.TransmissionFilter) && !(MaterialData.TransmissionFilter - FVector3f::One()).IsNearlyZero()) || bTransmission)
 			{
 				FLinearColor TransmissionColor = IsColorInitialized(MaterialData.TransmissionFilter) ? MaterialData.TransmissionFilter : FLinearColor::White;
 				const FString ParameterName = Materials::ThinTranslucent::Parameters::TransmissionColor.ToString();
-				UE::Interchange::Materials::Private::FMaterialNodesHelper::SetupVectorParameter(BaseNodeContainer, ShaderGraphNode, ParameterName, TransmissionColor);
+				MaterialInstanceNode->AddVectorParameterValue(ParameterName, TransmissionColor);
 			}
 
 			// Normal/bump
-			if (UInterchangeShaderNode* NormalmapTextureSampleShader = CreateMaterialTextureSampleNode(BaseNodeContainer, ShaderGraphNode, MaterialData.NormalmapTexture))
+			if (const UInterchangeTexture2DNode* NormalmapTextureNode = CreateTexture2DNode(BaseNodeContainer, MaterialData.NormalmapTexture))
 			{
-				UInterchangeShaderPortsAPI::ConnectDefaultOuputToInput(ShaderGraphNode, Materials::Common::Parameters::Normal.ToString(), NormalmapTextureSampleShader->GetUniqueID());
+				MaterialInstanceNode->AddTextureParameterValue(Materials::Common::Parameters::Normal.ToString() + TEXT("Map"), NormalmapTextureNode->GetUniqueID());
 			}
-			if (!MaterialData.BumpmapTexture.IsEmpty())
+			else if (!MaterialData.BumpmapTexture.IsEmpty())
 			{
 				using namespace Materials::Standard::Nodes;
 
 				FString TexturePath = MaterialData.BumpmapTexture.Path;
 				FString TextureName = MakeTextureNodeName(FPaths::GetBaseFilename(TexturePath));
 
-				if (const UInterchangeTexture2DNode* TextureNode = CreateTexture2DNode(BaseNodeContainer, TexturePath, TextureName))
+				if (const UInterchangeTexture2DNode* BumpmapTextureNode = CreateTexture2DNode(BaseNodeContainer, TexturePath, TextureName))
 				{
-					// NormalFromHeightmap needs TextureObject(not just a sample as it takes multiple samples from it)
-					UInterchangeShaderNode* TextureObjectNode = UInterchangeShaderNode::Create(&BaseNodeContainer, TextureName, ShaderGraphNode->GetUniqueID());
-					TextureObjectNode->SetCustomShaderType(TextureObject::Name.ToString());
-					TextureObjectNode->AddStringAttribute(UInterchangeShaderPortsAPI::MakeInputValueKey(TextureObject::Inputs::Texture.ToString()), TextureNode->GetUniqueID());
-
-					UInterchangeShaderNode* HeightMapNode = UInterchangeShaderNode::Create( &BaseNodeContainer, NormalFromHeightMap::Name.ToString(), ShaderGraphNode->GetUniqueID() );
-					HeightMapNode->SetCustomShaderType(NormalFromHeightMap::Name.ToString());
-
-					UInterchangeShaderPortsAPI::ConnectDefaultOuputToInput(HeightMapNode, NormalFromHeightMap::Inputs::HeightMap.ToString(), TextureObjectNode->GetUniqueID());
-					HeightMapNode->AddFloatAttribute(UInterchangeShaderPortsAPI::MakeInputValueKey(NormalFromHeightMap::Inputs::Intensity.ToString()), MaterialData.BumpmapTexture.BumpMultiplier);
-
-					UInterchangeShaderPortsAPI::ConnectDefaultOuputToInput(ShaderGraphNode, Materials::Common::Parameters::Normal.ToString(), HeightMapNode->GetUniqueID());
+					MaterialInstanceNode->AddTextureParameterValue(Materials::Common::Parameters::Normal.ToString() + TEXT("Map"), BumpmapTextureNode->GetUniqueID());
+					MaterialInstanceNode->AddScalarParameterValue(TEXT("NormalMapWeight"), MaterialData.BumpmapTexture.BumpMultiplier);
+					MaterialInstanceNode->AddStaticSwitchParameterValue(TEXT("bHasBump"), true);
 				}
 			}
 		}
diff --git a/Engine/Plugins/Interchange/Runtime/Source/Import/Public/Mesh/InterchangeOBJTranslator.h b/Engine/Plugins/Interchange/Runtime/Source/Import/Public/Mesh/InterchangeOBJTranslator.h
--- a/Engine/Plugins/Interchange/Runtime/Source/Import/Public/Mesh/InterchangeOBJTranslator.h
+++ b/Engine/Plugins/Interchange/Runtime/Source/Import/Public/Mesh/InterchangeOBJTranslator.h
@@ -13,7 +13,6 @@
 
 #define UE_API INTERCHANGEIMPORT_API
 
-class UInterchangeShaderGraphNode;
 struct FObjData;
 
 
diff --git a/Engine/Plugins/Interchange/Runtime/Source/Parsers/Fbx/Private/FbxMaterial.cpp b/Engine/Plugins/Interchange/Runtime/Source/Parsers/Fbx/Private/FbxMaterial.cpp
--- a/Engine/Plugins/Interchange/Runtime/Source/Parsers/Fbx/Private/FbxMaterial.cpp
+++ b/Engine/Plugins/Interchange/Runtime/Source/Parsers/Fbx/Private/FbxMaterial.cpp
@@ -3,14 +3,13 @@
 #include "FbxMaterial.h"
 
 #include "FbxAPI.h"
-#include "FbxConvert.h"
 #include "FbxHelper.h"
 #include "FbxInclude.h"
 #include "Fbx/InterchangeFbxMessages.h"
 #include "InterchangeMaterialDefinitions.h"
+#include "InterchangeMaterialInstanceNode.h"
 #include "InterchangeResultsContainer.h"
 #include "InterchangeSceneNode.h"
-#include "InterchangeShaderGraphNode.h"
 #include "InterchangeTexture2DNode.h"
 #include "InterchangeTextureNode.h"
 #include "Misc/Paths.h"
@@ -26,12 +25,14 @@ namespace UE
 	{
 		namespace Private
 		{
-			UInterchangeShaderGraphNode* FFbxMaterial::CreateShaderGraphNode(UInterchangeBaseNodeContainer& NodeContainer, const FString& NodeUid, const FString& NodeName)
+			UInterchangeMaterialInstanceNode* FFbxMaterial::CreateMaterialInstanceNode(UInterchangeBaseNodeContainer& NodeContainer, const FString& NodeUID, const FString& NodeName)
 			{
-				UInterchangeShaderGraphNode* ShaderGraphNode = NewObject<UInterchangeShaderGraphNode>(&NodeContainer);
-				NodeContainer.SetupNode(ShaderGraphNode, NodeUid, NodeName, EInterchangeNodeContainerType::TranslatedAsset);
+				UInterchangeMaterialInstanceNode* MaterialInstanceNode = NewObject<UInterchangeMaterialInstanceNode>(&NodeContainer);
+				// By default we create a Phong material, depending later of the different input hooked we'll change it
+				MaterialInstanceNode->SetCustomParent(TEXT("/InterchangeAssets/Materials/PhongSurfaceMaterial.PhongSurfaceMaterial"));
+				NodeContainer.SetupNode(MaterialInstanceNode, NodeUID, NodeName, EInterchangeNodeContainerType::TranslatedAsset);
 
-				return ShaderGraphNode;
+				return MaterialInstanceNode;
 			}
 
 			const UInterchangeTexture2DNode* FFbxMaterial::CreateTexture2DNode(UInterchangeBaseNodeContainer& NodeContainer, const FString& TextureFilePath)
@@ -63,10 +64,8 @@ namespace UE
 				return NewTextureNode;
 			}
 
-			const UInterchangeShaderNode* FFbxMaterial::CreateTextureSampler(FbxFileTexture* FbxTexture, UInterchangeBaseNodeContainer& NodeContainer, const FString& ShaderUniqueID, const FString& InputName)
+			const UInterchangeTexture2DNode* FFbxMaterial::CreateTexture2DNode(FbxFileTexture* FbxTexture, UInterchangeBaseNodeContainer& NodeContainer, UInterchangeMaterialInstanceNode* MaterialInstanceNode)
 			{
-				using namespace Materials::Standard::Nodes;
-
 				if (!FbxTexture)
 				{
 					return nullptr;
@@ -75,18 +74,18 @@ namespace UE
 				const FString TextureFilename = [this, FbxTexture]
 				{
 					// try opening from absolute path
-					FString AbsolueTextureFilepath = UTF8_TO_TCHAR(FbxTexture->GetFileName());
-					FPaths::NormalizeFilename(AbsolueTextureFilepath);
-					if( FPaths::FileExists(AbsolueTextureFilepath) )
+					FString AbsoluteTextureFilepath = UTF8_TO_TCHAR(FbxTexture->GetFileName());
+					FPaths::NormalizeFilename(AbsoluteTextureFilepath);
+					if (FPaths::FileExists(AbsoluteTextureFilepath))
 					{
-						return AbsolueTextureFilepath;
+						return AbsoluteTextureFilepath;
 					}
 
 					FString FileBasePath = FPaths::GetPath(Parser.GetSourceFilename());
 					// try fbx file base path + relative path
 					FString RelativeToFBXTextureFilepath = FileBasePath / UTF8_TO_TCHAR(FbxTexture->GetRelativeFileName());
 					FPaths::NormalizeFilename(RelativeToFBXTextureFilepath);
-					if( FPaths::FileExists(RelativeToFBXTextureFilepath) )
+					if (FPaths::FileExists(RelativeToFBXTextureFilepath))
 					{
 						return RelativeToFBXTextureFilepath;
 					}
@@ -94,7 +93,7 @@ namespace UE
 					// Some fbx files do not store the actual absolute filename as absolute and it is actually relative.  Try to get it relative to the FBX file we are importing
 					FString AbosluteAsRelativeToFBXTextureFilepath = FileBasePath / UTF8_TO_TCHAR(FbxTexture->GetFileName());
 					FPaths::NormalizeFilename(AbosluteAsRelativeToFBXTextureFilepath);
-					if( FPaths::FileExists(AbosluteAsRelativeToFBXTextureFilepath) )
+					if (FPaths::FileExists(AbosluteAsRelativeToFBXTextureFilepath))
 					{
 						return AbosluteAsRelativeToFBXTextureFilepath;
 					}
@@ -102,53 +101,24 @@ namespace UE
 					return FString();
 				}();
 
-				const FString NodeName = TEXT("Sampler_") + InputName;
-
-				// Return already created node if applicable.
-				const FString SamplerNodeUid = UInterchangeShaderNode::MakeNodeUid(NodeName, ShaderUniqueID);
-				if (const UInterchangeShaderNode* SamplerNode = Cast<UInterchangeShaderNode>(NodeContainer.GetNode(SamplerNodeUid)))
-				{
-					return SamplerNode;
-				}
-
-				UInterchangeShaderNode* TextureSampleShader = UInterchangeShaderNode::Create(&NodeContainer, NodeName, ShaderUniqueID);
-				TextureSampleShader->SetDisplayLabel(InputName);
-				TextureSampleShader->SetCustomShaderType(TextureSample::Name.ToString());
-
 				// Return incomplete texture sampler if texture file does not exist
 				if (TextureFilename.IsEmpty() || !FPaths::FileExists(TextureFilename))
 				{
-					const UInterchangeBaseNode* ShaderGraphNode = NodeContainer.GetNode(ShaderUniqueID);
-					if(!GIsAutomationTesting)
+					if (!GIsAutomationTesting)
 					{
 						UInterchangeResultTextureDisplay_TextureFileDoNotExist* Message = Parser.AddMessage<UInterchangeResultTextureDisplay_TextureFileDoNotExist>();
 						Message->TextureName = TextureFilename;
-						Message->MaterialName = ShaderGraphNode ? ShaderGraphNode->GetDisplayLabel() : TEXT("Unknown");
+						Message->MaterialName = MaterialInstanceNode->GetDisplayLabel();
 					}
 
-					return TextureSampleShader;
+					return nullptr;
 				}
 
 				const UInterchangeTexture2DNode* TextureNode = CreateTexture2DNode(NodeContainer, TextureFilename);
-
-				TextureSampleShader->AddStringAttribute(UInterchangeShaderPortsAPI::MakeInputValueKey(TextureSample::Inputs::Texture.ToString()), TextureNode->GetUniqueID());
-
-				if (!FMath::IsNearlyEqual(FbxTexture->GetScaleU(), 1.0) || !FMath::IsNearlyEqual(FbxTexture->GetScaleV(), 1.0))
-				{
-					UInterchangeShaderNode* TextureCoordinateShader = UInterchangeShaderNode::Create(&NodeContainer, InputName + TEXT("_Coordinate"), ShaderUniqueID);
-					TextureCoordinateShader->SetCustomShaderType(TextureCoordinate::Name.ToString());
-
-					TextureCoordinateShader->AddFloatAttribute(UInterchangeShaderPortsAPI::MakeInputValueKey(TextureCoordinate::Inputs::UTiling.ToString()), (float)FbxTexture->GetScaleU());
-					TextureCoordinateShader->AddFloatAttribute(UInterchangeShaderPortsAPI::MakeInputValueKey(TextureCoordinate::Inputs::VTiling.ToString()), (float)FbxTexture->GetScaleV());
-
-					UInterchangeShaderPortsAPI::ConnectDefaultOuputToInput(TextureSampleShader, TextureSample::Inputs::Coordinates.ToString(), TextureCoordinateShader->GetUniqueID());
-				}
-
-				return TextureSampleShader;
+				return TextureNode;
 			}
 
-			bool FFbxMaterial::ConvertPropertyToShaderNode(UInterchangeBaseNodeContainer& NodeContainer, UInterchangeShaderGraphNode* ShaderGraphNode, FbxProperty& Property, float Factor, FName InputName,
-														   const TVariant<FLinearColor, float>& DefaultValue, bool bInverse)
+			bool FFbxMaterial::ConvertPropertyToMaterialInstanceNode(UInterchangeBaseNodeContainer& NodeContainer, UInterchangeMaterialInstanceNode* MaterialInstanceNode, FbxProperty& Property, float Factor, FName InputName, const TVariant<FLinearColor, float>& DefaultValue, bool bInverse)
 			{
 				using namespace Materials::Standard::Nodes;
 
@@ -158,12 +128,12 @@ namespace UE
 
 				if (TextureCount == 0)
 				{
-					const FString InputAttributeKey = UInterchangeShaderPortsAPI::MakeInputValueKey(InputName.ToString());
+					const FString InputAttributeKey = InputName.ToString();
 
 					if (DataType == eFbxDouble || DataType == eFbxFloat || DataType == eFbxInt)
 					{
 						const float PropertyValue = Property.Get<float>() * Factor;
-						ShaderGraphNode->AddFloatAttribute(InputAttributeKey, bInverse ? 1.f - PropertyValue : PropertyValue);
+						MaterialInstanceNode->AddFloatAttribute(InputAttributeKey, bInverse ? 1.f - PropertyValue : PropertyValue);
 					}
 					else if (DataType == eFbxDouble3 || DataType == eFbxDouble4)
 					{
@@ -173,65 +143,23 @@ namespace UE
 
 						if (DefaultValue.IsType<FLinearColor>())
 						{
-							ShaderGraphNode->AddLinearColorAttribute(InputAttributeKey, PropertyValue);
+							MaterialInstanceNode->AddVectorParameterValue(InputAttributeKey, PropertyValue);
 						}
 						else if (DefaultValue.IsType<float>())
 						{
 							// We're connecting a linear color to a float input. Ideally, we'd go through a desaturate, but for now we'll just take the red channel and ignore the rest.
-							ShaderGraphNode->AddFloatAttribute(InputAttributeKey, PropertyValue.R);
+							MaterialInstanceNode->AddScalarParameterValue(InputAttributeKey, PropertyValue.R);
 						}
 					}
 
 					return true;
 				}
 
-				UInterchangeShaderNode* NodeToConnectTo = ShaderGraphNode;
-
-				if (bInverse)
-				{
-					const FString OneMinusNodeName = InputName.ToString() + TEXT("OneMinus");
-					UInterchangeShaderNode* OneMinusNode = UInterchangeShaderNode::Create(&NodeContainer, OneMinusNodeName, ShaderGraphNode->GetUniqueID());
-					OneMinusNode->SetCustomShaderType(OneMinus::Name.ToString());
-
-					UInterchangeShaderPortsAPI::ConnectDefaultOuputToInput(ShaderGraphNode, InputName.ToString(), OneMinusNode->GetUniqueID());
-
-					NodeToConnectTo = OneMinusNode;
-					InputToConnectTo = OneMinus::Inputs::Input.ToString();
-				}
-
-				{
-					FString LerpNodeName = InputName.ToString() + TEXT("Lerp");
-					UInterchangeShaderNode* LerpNode = UInterchangeShaderNode::Create(&NodeContainer, LerpNodeName, ShaderGraphNode->GetUniqueID());
-					LerpNode->SetCustomShaderType(Lerp::Name.ToString());
-
-					if (DefaultValue.IsType<float>())
-					{
-						LerpNode->AddFloatAttribute(UInterchangeShaderPortsAPI::MakeInputValueKey(Lerp::Inputs::A.ToString()), DefaultValue.Get<float>());
-					}
-					else if (DefaultValue.IsType<FLinearColor>())
-					{
-						LerpNode->AddLinearColorAttribute(UInterchangeShaderPortsAPI::MakeInputValueKey(Lerp::Inputs::A.ToString()), DefaultValue.Get<FLinearColor>());
-					}
-					
-					const FString WeightNodeName = InputName.ToString() + TEXT("MapWeight");
-					UInterchangeShaderNode* WeightNode = UInterchangeShaderNode::Create(&NodeContainer, WeightNodeName, LerpNode->GetUniqueID());
-					WeightNode->SetCustomShaderType(ScalarParameter::Name.ToString());
-										
-					WeightNode->AddFloatAttribute(UInterchangeShaderPortsAPI::MakeInputParameterKey(ScalarParameter::Attributes::DefaultValue.ToString()), Factor);
-
-					UInterchangeShaderPortsAPI::ConnectDefaultOuputToInput(LerpNode, Lerp::Inputs::Factor.ToString() , WeightNode->GetUniqueID());
-
-					UInterchangeShaderPortsAPI::ConnectDefaultOuputToInput(NodeToConnectTo, InputToConnectTo, LerpNode->GetUniqueID());
-
-					NodeToConnectTo = LerpNode;
-					InputToConnectTo = Lerp::Inputs::B.ToString();
-				}
-
 				// Handles max one texture per property.
 				FbxFileTexture* FbxTexture = Property.GetSrcObject<FbxFileTexture>(0);
-				if (const UInterchangeShaderNode* TextureSampleShader = CreateTextureSampler(FbxTexture, NodeContainer, ShaderGraphNode->GetUniqueID(), InputName.ToString() + TEXT("Map")))
+				if (const UInterchangeTexture2DNode* TextureNode = CreateTexture2DNode(FbxTexture, NodeContainer, MaterialInstanceNode))
 				{
-					UInterchangeShaderPortsAPI::ConnectDefaultOuputToInput(NodeToConnectTo, InputToConnectTo, TextureSampleShader->GetUniqueID());
+					MaterialInstanceNode->AddTextureParameterValue(InputName.ToString() + TEXT("Map"), TextureNode->GetUniqueID());
 				}
 				else
 				{
@@ -239,7 +167,7 @@ namespace UE
 					{
 						UInterchangeResultTextureDisplay_TextureFileDoNotExist* Message = Parser.AddMessage<UInterchangeResultTextureDisplay_TextureFileDoNotExist>();
 						Message->TextureName = FbxTexture ? UTF8_TO_TCHAR(FbxTexture->GetFileName()) : TEXT("Undefined");
-						Message->MaterialName = ShaderGraphNode->GetDisplayLabel();
+						Message->MaterialName = MaterialInstanceNode->GetDisplayLabel();
 					}
 
 					return false;
@@ -248,48 +176,7 @@ namespace UE
 				return true;
 			}
 
-			void FFbxMaterial::ConvertShininessToShaderNode(FbxSurfaceMaterial& SurfaceMaterial, UInterchangeBaseNodeContainer& NodeContainer, UInterchangeShaderGraphNode* ShaderGraphNode)
-			{
-				using namespace UE::Interchange::Materials;
-				using namespace UE::Interchange::Materials::Standard::Nodes;
-
-				FBXSDK_NAMESPACE::FbxProperty MaterialProperty = SurfaceMaterial.FindProperty(FbxSurfaceMaterial::sShininess);
-
-				if (!MaterialProperty.IsValid())
-				{
-					return;
-				}
-
-				const FString InputName = Phong::Parameters::Shininess.ToString();
-
-				if (MaterialProperty.GetSrcObjectCount<FBXSDK_NAMESPACE::FbxTexture>() > 0)
-				{
-					FbxFileTexture* FbxTexture = MaterialProperty.GetSrcObject<FbxFileTexture>(0);
-					if (const UInterchangeShaderNode* TextureSampleShader = CreateTextureSampler(FbxTexture, NodeContainer, ShaderGraphNode->GetUniqueID(), TEXT("ShininessMap")))
-					{
-						FString MultiplyNodeName = Phong::Parameters::Shininess.ToString() + TEXT("_Multiply");
-						UInterchangeShaderNode* MultiplyNode = UInterchangeShaderNode::Create(&NodeContainer, MultiplyNodeName, ShaderGraphNode->GetUniqueID());
-						MultiplyNode->SetCustomShaderType(Multiply::Name.ToString());
-						
-						UInterchangeShaderPortsAPI::ConnectDefaultOuputToInput(ShaderGraphNode, InputName, MultiplyNode->GetUniqueID());
-
-						UInterchangeShaderNode* WeightNode = UInterchangeShaderNode::Create(&NodeContainer, TEXT("ShininessMapWeight"), MultiplyNode->GetUniqueID());
-						WeightNode->SetCustomShaderType(ScalarParameter::Name.ToString());
-						WeightNode->AddFloatAttribute(UInterchangeShaderPortsAPI::MakeInputParameterKey(ScalarParameter::Attributes::DefaultValue.ToString()), 1.f);
-
-						UInterchangeShaderPortsAPI::ConnectDefaultOuputToInput(MultiplyNode, Multiply::Inputs::B.ToString(), WeightNode->GetUniqueID());
-
-						UInterchangeShaderPortsAPI::ConnectDefaultOuputToInput(MultiplyNode, Multiply::Inputs::A.ToString(), TextureSampleShader->GetUniqueID());
-					}
-				}
-				else if (!FBXSDK_NAMESPACE::FbxProperty::HasDefaultValue(MaterialProperty))
-				{
-					const FString InputValueKey = UInterchangeShaderPortsAPI::MakeInputValueKey(InputName);
-					ShaderGraphNode->AddFloatAttribute(InputValueKey, MaterialProperty.Get<float>());
-				}
-			}
-
-			const UInterchangeShaderGraphNode* FFbxMaterial::AddShaderGraphNode(FbxSurfaceMaterial* SurfaceMaterial, UInterchangeBaseNodeContainer& NodeContainer)
+			const UInterchangeMaterialInstanceNode* FFbxMaterial::AddMaterialInstanceNode(FbxSurfaceMaterial* SurfaceMaterial, UInterchangeBaseNodeContainer& NodeContainer)
 			{
 				using namespace UE::Interchange::Materials;
 
@@ -301,15 +188,15 @@ namespace UE
 				//Create a material node
 				FString MaterialName = Parser.GetFbxHelper()->GetFbxObjectName(SurfaceMaterial);
 				FString NodeUid = TEXT("\\Material\\") + MaterialName;
-				const UInterchangeShaderGraphNode* ExistingShaderGraphNode = Cast<const UInterchangeShaderGraphNode>(NodeContainer.GetNode(NodeUid));
-				if (ExistingShaderGraphNode)
+				const UInterchangeMaterialInstanceNode* ExistingMaterialInstanceNode = Cast<const UInterchangeMaterialInstanceNode>(NodeContainer.GetNode(NodeUid));
+				if (ExistingMaterialInstanceNode)
 				{
-					return ExistingShaderGraphNode;
+					return ExistingMaterialInstanceNode;
 				}
 
 
-				UInterchangeShaderGraphNode* ShaderGraphNode = CreateShaderGraphNode(NodeContainer, NodeUid, MaterialName);
-				if (ShaderGraphNode == nullptr)
+				UInterchangeMaterialInstanceNode* MaterialInstanceNode = CreateMaterialInstanceNode(NodeContainer, NodeUid, MaterialName);
+				if (MaterialInstanceNode == nullptr)
 				{
 					FFormatNamedArguments Args
 					{
@@ -320,46 +207,47 @@ namespace UE
 					return nullptr;
 				}
 
-				ProcessCustomAttributes(Parser, SurfaceMaterial, ShaderGraphNode);
-
-				TFunction<bool(FBXSDK_NAMESPACE::FbxProperty&)> ShouldConvertProperty = [&](FBXSDK_NAMESPACE::FbxProperty& MaterialProperty) -> bool
-				{
-					bool bShouldConvertProperty = false;
+				ProcessCustomAttributes(Parser, SurfaceMaterial, MaterialInstanceNode);
 
-					if (MaterialProperty.IsValid())
+				auto ShouldConvertProperty = [](FBXSDK_NAMESPACE::FbxProperty& MaterialProperty) -> bool
 					{
-						// FbxProperty::HasDefaultValue(..) can return true while the property has textures attached to it.
-						bShouldConvertProperty = MaterialProperty.GetSrcObjectCount<FBXSDK_NAMESPACE::FbxTexture>() > 0
-							|| !FBXSDK_NAMESPACE::FbxProperty::HasDefaultValue(MaterialProperty);
-					}
+						bool bShouldConvertProperty = false;
 
-					return bShouldConvertProperty;
-				};
+						if (MaterialProperty.IsValid())
+						{
+							// FbxProperty::HasDefaultValue(..) can return true while the property has textures attached to it.
+							bShouldConvertProperty = MaterialProperty.GetSrcObjectCount<FBXSDK_NAMESPACE::FbxTexture>() > 0
+								|| !FBXSDK_NAMESPACE::FbxProperty::HasDefaultValue(MaterialProperty);
+						}
 
-				TFunction<float(const char*)> GetFactor = [&](const char* FactorName)
-				{
-					FBXSDK_NAMESPACE::FbxProperty Property = SurfaceMaterial->FindProperty(FactorName);
-					return Property.IsValid() ? (float)Property.Get<FbxDouble>() : 1.;
-				};
+						return bShouldConvertProperty;
+					};
 
-				TFunction<bool(FName, const char*, float , TVariant<FLinearColor, float>&, bool)>  ConnectInput;
-				ConnectInput = [&](FName InputName, const char* FbxPropertyName, float Factor, TVariant<FLinearColor, float>& DefaultValue, bool bInverse) -> bool
-				{
-					FBXSDK_NAMESPACE::FbxProperty MaterialProperty = SurfaceMaterial->FindProperty(FbxPropertyName);
-					if (ShouldConvertProperty(MaterialProperty))
+				auto GetFactor = [&SurfaceMaterial](const char* FactorName)
 					{
-						return ConvertPropertyToShaderNode(NodeContainer, ShaderGraphNode, MaterialProperty, Factor, InputName, DefaultValue, bInverse);
-					}
+						FBXSDK_NAMESPACE::FbxProperty Property = SurfaceMaterial->FindProperty(FactorName);
+						return Property.IsValid() ? (float)Property.Get<FbxDouble>() : 1.;
+					};
 
-					return false;
-				};
+				auto ConnectInput = [&](FName InputName, const char* FbxPropertyName, float Factor, TVariant<FLinearColor, float>& DefaultValue, bool bInverse) -> bool
+					{
+						FBXSDK_NAMESPACE::FbxProperty MaterialProperty = SurfaceMaterial->FindProperty(FbxPropertyName);
+						if (ShouldConvertProperty(MaterialProperty))
+						{
+							return ConvertPropertyToMaterialInstanceNode(NodeContainer, MaterialInstanceNode, MaterialProperty, Factor, InputName, DefaultValue, bInverse);
+						}
+
+						return false;
+					};
+
+				bool bHasInput = false;
 
 				// Diffuse
 				{
 					const float Factor = GetFactor(FbxSurfaceMaterial::sDiffuseFactor);
 					TVariant<FLinearColor, float> DefaultValue;
 					DefaultValue.Set<FLinearColor>(FLinearColor::Black);
-					ConnectInput(Phong::Parameters::DiffuseColor, FbxSurfaceMaterial::sDiffuse, Factor, DefaultValue, false);
+					bHasInput |= ConnectInput(Phong::Parameters::DiffuseColor, FbxSurfaceMaterial::sDiffuse, Factor, DefaultValue, false);
 				}
 
 				// Ambient
@@ -367,7 +255,7 @@ namespace UE
 					const float Factor = GetFactor(FbxSurfaceMaterial::sAmbientFactor);
 					TVariant<FLinearColor, float> DefaultValue;
 					DefaultValue.Set<FLinearColor>(FLinearColor::Black);
-					ConnectInput(Phong::Parameters::AmbientColor, FbxSurfaceMaterial::sAmbient, Factor, DefaultValue, false);
+					bHasInput |= ConnectInput(Phong::Parameters::AmbientColor, FbxSurfaceMaterial::sAmbient, Factor, DefaultValue, false);
 				}
 
 				// Emissive
@@ -375,53 +263,21 @@ namespace UE
 					const float Factor = GetFactor(FbxSurfaceMaterial::sEmissiveFactor);
 					TVariant<FLinearColor, float> DefaultValue;
 					DefaultValue.Set<FLinearColor>(FLinearColor::Black);
-					ConnectInput(Phong::Parameters::EmissiveColor, FbxSurfaceMaterial::sEmissive, Factor, DefaultValue, false);
+					bHasInput |= ConnectInput(Phong::Parameters::EmissiveColor, FbxSurfaceMaterial::sEmissive, Factor, DefaultValue, false);
 				}
 
 				// Normal
 				{
-					// FbxSurfaceMaterial can have either a normal map or a bump map, check for both
-					FBXSDK_NAMESPACE::FbxProperty MaterialProperty = SurfaceMaterial->FindProperty(FbxSurfaceMaterial::sNormalMap);
-					if (MaterialProperty.IsValid() && MaterialProperty.GetSrcObjectCount<FbxTexture>() > 0)
+					const float FactorNormal = GetFactor(FbxSurfaceMaterial::sNormalMap);
+					TVariant<FLinearColor, float> DefaultValue;
+					DefaultValue.Set<FLinearColor>(FLinearColor{ FVector::UpVector });
+					bool bHasNormal = ConnectInput(Phong::Parameters::Normal, FbxSurfaceMaterial::sNormalMap, FactorNormal, DefaultValue, false);
+					bHasInput |= bHasNormal;
+					if (!bHasNormal)
 					{
-						TVariant<FLinearColor, float> DefaultValue;
-						DefaultValue.Set<FLinearColor>(FLinearColor(FVector::UpVector));
-						ConvertPropertyToShaderNode(NodeContainer, ShaderGraphNode, MaterialProperty, 1.f, Phong::Parameters::Normal, DefaultValue);
-					}
-					else
-					{
-						const float Factor = GetFactor(FbxSurfaceMaterial::sBumpFactor);
-						MaterialProperty = SurfaceMaterial->FindProperty(FbxSurfaceMaterial::sBump);
-						if (MaterialProperty.IsValid() && MaterialProperty.GetSrcObjectCount<FbxTexture>() > 0)
-						{
-							using namespace Materials::Standard::Nodes;
-
-							if (FbxFileTexture* FbxTexture = MaterialProperty.GetSrcObject<FbxFileTexture>(0))
-							{
-								const FString TexturePath = FbxTexture->GetFileName();
-								if (!TexturePath.IsEmpty() && FPaths::FileExists(TexturePath))
-								{
-									if (const UInterchangeTexture2DNode* TextureNode = CreateTexture2DNode(NodeContainer, TexturePath))
-									{
-										const FString TextureName = FPaths::GetBaseFilename(TexturePath);
-
-										// NormalFromHeightmap needs TextureObject(not just a sample as it takes multiple samples from it)
-										UInterchangeShaderNode* TextureObjectNode = UInterchangeShaderNode::Create(&NodeContainer, TEXT("NormalMap"), ShaderGraphNode->GetUniqueID());
-										TextureObjectNode->SetDisplayLabel(TEXT("NormalMap"));
-										TextureObjectNode->SetCustomShaderType(TextureObject::Name.ToString());
-										TextureObjectNode->AddStringAttribute(UInterchangeShaderPortsAPI::MakeInputValueKey(TextureObject::Inputs::Texture.ToString()), TextureNode->GetUniqueID());
-
-										UInterchangeShaderNode* HeightMapNode = UInterchangeShaderNode::Create(&NodeContainer, NormalFromHeightMap::Name.ToString(), ShaderGraphNode->GetUniqueID());
-										HeightMapNode->SetCustomShaderType(NormalFromHeightMap::Name.ToString());
-
-										UInterchangeShaderPortsAPI::ConnectDefaultOuputToInput(HeightMapNode, NormalFromHeightMap::Inputs::HeightMap.ToString(), TextureObjectNode->GetUniqueID());
-										HeightMapNode->AddFloatAttribute(UInterchangeShaderPortsAPI::MakeInputValueKey(NormalFromHeightMap::Inputs::Intensity.ToString()), Factor);
-
-										UInterchangeShaderPortsAPI::ConnectDefaultOuputToInput(ShaderGraphNode, Materials::Common::Parameters::Normal.ToString(), HeightMapNode->GetUniqueID());
-									}
-								}
-							}
-						}
+						const float FactorBump = GetFactor(FbxSurfaceMaterial::sBumpFactor);
+						DefaultValue.Set<float>(0.f);
+						bHasInput |= ConnectInput(Phong::Parameters::Normal, FbxSurfaceMaterial::sBump, FactorBump, DefaultValue, false);
 					}
 				}
 
@@ -435,6 +291,7 @@ namespace UE
 
 						if (MaterialProperty.GetSrcObjectCount<FbxFileTexture>() > 0)
 						{
+							MaterialInstanceNode->SetCustomParent(TEXT("/InterchangeAssets/Materials/PhongSurfaceTransparentMaterial.PhongSurfaceTransparentMaterial"));
 							TVariant<FLinearColor, float> DefaultValue;
 							DefaultValue.Set<float>(0.f); // Opaque
 							FName InputName = Phong::Parameters::Opacity;
@@ -442,49 +299,62 @@ namespace UE
 							if (MaterialProperty.Get<FbxDouble>() == 0.)
 							{
 								InputName = Phong::Parameters::OpacityMask;
-								ShaderGraphNode->SetCustomOpacityMaskClipValue(0.333f);
+								MaterialInstanceNode->SetCustomParent(TEXT("/InterchangeAssets/Materials/PhongSurfaceMaskMaterial.PhongSurfaceMaskMaterial"));
 							}
-							ConvertPropertyToShaderNode(NodeContainer, ShaderGraphNode, MaterialProperty, Factor, InputName, DefaultValue, true);
+							bHasInput |= ConvertPropertyToMaterialInstanceNode(NodeContainer, MaterialInstanceNode, MaterialProperty, Factor, InputName, DefaultValue, true);
 						}
 						else
 						{
 							const float OpacityScalar = 1.f - (MaterialProperty.Get<float>() * Factor);
 							if (OpacityScalar < 1.f)
 							{
-								const FString InputName = UInterchangeShaderPortsAPI::MakeInputValueKey(Phong::Parameters::Opacity.ToString());
-								ShaderGraphNode->AddFloatAttribute(InputName, OpacityScalar);
+								MaterialInstanceNode->SetCustomParent(TEXT("/InterchangeAssets/Materials/PhongSurfaceTransparentMaterial.PhongSurfaceTransparentMaterial"));
+								MaterialInstanceNode->AddScalarParameterValue(Phong::Parameters::Opacity.ToString(), OpacityScalar);
 							}
 						}
 					}
 				}
 
+				// if it has no specular no shininess then just take a Lambert material
+				bool bIsPhong = false;
+
 				// Specular
 				{
 					const float Factor = GetFactor(FbxSurfaceMaterial::sSpecularFactor);
 					TVariant<FLinearColor, float> DefaultValue;
 					DefaultValue.Set<FLinearColor>(FLinearColor::Black);
-					ConnectInput(Phong::Parameters::SpecularColor, FbxSurfaceMaterial::sSpecular, Factor, DefaultValue, false);
+					bIsPhong |= ConnectInput(Phong::Parameters::SpecularColor, FbxSurfaceMaterial::sSpecular, Factor, DefaultValue, false);
+					bHasInput |= bIsPhong;
+				}
+
+				//Shininess
+				{
+					const float Factor = GetFactor(FbxSurfaceMaterial::sShininess);
+					TVariant<FLinearColor, float> DefaultValue;
+					DefaultValue.Set<float>(20.f);
+					bIsPhong |= ConnectInput(Phong::Parameters::Shininess, FbxSurfaceMaterial::sShininess, Factor, DefaultValue, false);
+					bHasInput |= bIsPhong;
 				}
 
-				// Shininess
-				ConvertShininessToShaderNode(*SurfaceMaterial, NodeContainer, ShaderGraphNode);
+				if (!bIsPhong)
+				{
+					MaterialInstanceNode->SetCustomParent(TEXT("/InterchangeAssets/Materials/LambertSurfaceMaterial.LambertSurfaceMaterial"));
+				}
 
 				// If no valid property found, create a material anyway
-				// Use random color because there may be multiple materials, then they can be different
 				TArray<FString> InputNames;
-				UInterchangeShaderPortsAPI::GatherInputs(ShaderGraphNode, InputNames);
-				if (InputNames.Num() == 0)
+				if (!bHasInput)
 				{
 					FLinearColor BaseColor;
 					BaseColor.R = 0.7f;
 					BaseColor.G = 0.7f;
 					BaseColor.B = 0.7f;
 
-					const FString InputValueKey = UInterchangeShaderPortsAPI::MakeInputValueKey(Phong::Parameters::DiffuseColor.ToString());
-					ShaderGraphNode->AddLinearColorAttribute(InputValueKey, BaseColor);
+					const FString InputValueKey = Phong::Parameters::DiffuseColor.ToString();
+					MaterialInstanceNode->AddVectorParameterValue(InputValueKey, BaseColor);
 				}
 
-				return ShaderGraphNode;
+				return MaterialInstanceNode;
 			}
 
 			void FFbxMaterial::AddAllTextures(FbxScene* SDKScene, UInterchangeBaseNodeContainer& NodeContainer)
@@ -524,7 +394,7 @@ namespace UE
 				{
 					if (FbxSurfaceMaterial* SurfaceMaterial = ParentFbxNode->GetMaterial(MaterialIndex))
 					{
-						const UInterchangeShaderGraphNode* ShaderGraphNode = AddShaderGraphNode(SurfaceMaterial, NodeContainer);
+						const UInterchangeMaterialInstanceNode* MaterialInstanceNode = AddMaterialInstanceNode(SurfaceMaterial, NodeContainer);
 						
 						int32& SlotMaterialCount = UniqueSlotNames.FindOrAdd(SurfaceMaterial);
 						FString MaterialSlotName = Parser.GetFbxHelper()->GetFbxObjectName(SurfaceMaterial);
@@ -532,7 +402,7 @@ namespace UE
 						{
 							MaterialSlotName += TEXT("_Section") + FString::FromInt(SlotMaterialCount);
 						}
-						SceneNode->SetSlotMaterialDependencyUid(MaterialSlotName, ShaderGraphNode->GetUniqueID());
+						SceneNode->SetSlotMaterialDependencyUid(MaterialSlotName, MaterialInstanceNode->GetUniqueID());
 						SlotMaterialCount++;
 					}
 				}
@@ -545,7 +415,7 @@ namespace UE
 				{
 					if (FbxSurfaceMaterial* SurfaceMaterial = SDKScene->GetMaterial(MaterialIndex))
 					{
-						AddShaderGraphNode(SurfaceMaterial, NodeContainer);
+						AddMaterialInstanceNode(SurfaceMaterial, NodeContainer);
 					}
 				}
 			}
diff --git a/Engine/Plugins/Interchange/Runtime/Source/Parsers/Fbx/Private/FbxMaterial.h b/Engine/Plugins/Interchange/Runtime/Source/Parsers/Fbx/Private/FbxMaterial.h
--- a/Engine/Plugins/Interchange/Runtime/Source/Parsers/Fbx/Private/FbxMaterial.h
+++ b/Engine/Plugins/Interchange/Runtime/Source/Parsers/Fbx/Private/FbxMaterial.h
@@ -9,8 +9,7 @@
 /** Forward declarations */
 class UInterchangeBaseNode;
 class UInterchangeBaseNodeContainer;
-class UInterchangeShaderGraphNode;
-class UInterchangeShaderNode;
+class UInterchangeMaterialInstanceNode;
 class UInterchangeTexture2DNode;
 class UInterchangeSceneNode;
 
@@ -52,15 +51,14 @@ namespace UE
 				void AddAllNodeMaterials(UInterchangeSceneNode* SceneNode, FbxNode* ParentFbxNode, UInterchangeBaseNodeContainer& NodeContainer);
 
 			protected:
-				UInterchangeShaderGraphNode* CreateShaderGraphNode(UInterchangeBaseNodeContainer& NodeContainer, const FString& NodeUID, const FString& NodeName);
+				UInterchangeMaterialInstanceNode* CreateMaterialInstanceNode(UInterchangeBaseNodeContainer& NodeContainer, const FString& NodeUID, const FString& NodeName);
 				const UInterchangeTexture2DNode* CreateTexture2DNode(UInterchangeBaseNodeContainer& NodeContainer, const FString& TextureFilePath);
-				bool ConvertPropertyToShaderNode(UInterchangeBaseNodeContainer& NodeContainer, UInterchangeShaderGraphNode* ShaderGraphNode,
-												 FbxProperty& Property, float Factor, FName PropertyName, const TVariant<FLinearColor, float>& DefaultValue, bool bInverse = false);
+				const UInterchangeTexture2DNode* CreateTexture2DNode(FbxFileTexture* FbxTexture, UInterchangeBaseNodeContainer& NodeContainer, UInterchangeMaterialInstanceNode* MaterialInstanceNode);
+				bool ConvertPropertyToMaterialInstanceNode(UInterchangeBaseNodeContainer& NodeContainer, UInterchangeMaterialInstanceNode* MaterialInstanceNode,
+														   FbxProperty& Property, float Factor, FName InputName, const TVariant<FLinearColor, float>& DefaultValue, bool bInverse = false);
 
 			private:
-				const UInterchangeShaderGraphNode* AddShaderGraphNode(FbxSurfaceMaterial* SurfaceMaterial, UInterchangeBaseNodeContainer& NodeContainer);
-				void ConvertShininessToShaderNode(FbxSurfaceMaterial& SurfaceMaterial, UInterchangeBaseNodeContainer& NodeContainer, UInterchangeShaderGraphNode* ShaderGraphNode);
-				const UInterchangeShaderNode* CreateTextureSampler(FbxFileTexture* FbxTexture, UInterchangeBaseNodeContainer& NodeContainer, const FString& ShaderUniqueID, const FString& InputName);
+				const UInterchangeMaterialInstanceNode* AddMaterialInstanceNode(FbxSurfaceMaterial* SurfaceMaterial, UInterchangeBaseNodeContainer& NodeContainer);
 
 				FFbxParser& Parser;
 			};
diff --git a/Engine/Source/Editor/MaterialEditor/Private/MaterialExpressionClasses.cpp b/Engine/Source/Editor/MaterialEditor/Private/MaterialExpressionClasses.cpp
--- a/Engine/Source/Editor/MaterialEditor/Private/MaterialExpressionClasses.cpp
+++ b/Engine/Source/Editor/MaterialEditor/Private/MaterialExpressionClasses.cpp
@@ -111,10 +111,10 @@ void MaterialExpressionClasses::InitMaterialExpressionClasses()
 						}
 
 						// Hide node types that are tagged private
-						if(Class->HasMetaData(TEXT("Private")))
-						{
-							continue;
-						}
+						//if(Class->HasMetaData(TEXT("Private")))
+						//{
+						//	continue;
+						//}
 
 						// Exclude comments from the expression list, as well as the base parameter expression and local variable usage, as they should not be used directly
 						if ( Class != UMaterialExpressionComment::StaticClass() 
