Interchange Analytics - FBX Translator Events

#rnx

diff --git a/Engine/Plugins/Interchange/Runtime/Source/Import/Private/Fbx/InterchangeFbxTranslator.cpp b/Engine/Plugins/Interchange/Runtime/Source/Import/Private/Fbx/InterchangeFbxTranslator.cpp
--- a/Engine/Plugins/Interchange/Runtime/Source/Import/Private/Fbx/InterchangeFbxTranslator.cpp
+++ b/Engine/Plugins/Interchange/Runtime/Source/Import/Private/Fbx/InterchangeFbxTranslator.cpp
@@ -2,14 +2,16 @@
 
 #include "Fbx/InterchangeFbxTranslator.h"
 
-#include "Dom/JsonValue.h"
+#include "AnalyticsEventAttribute.h"
 #include "Dom/JsonObject.h"
+#include "Dom/JsonValue.h"
 #include "HAL/PlatformFileManager.h"
+#include "InterchangeAnalyticsHelper.h"
 #include "InterchangeCommonAnimationPayload.h"
 #include "InterchangeDispatcher.h"
 #include "InterchangeDispatcherTask.h"
-#include "InterchangeManager.h"
 #include "InterchangeImportLog.h"
+#include "InterchangeManager.h"
 #include "InterchangeTranslatorHelper.h"
 #include "Mesh/InterchangeMeshPayload.h"
 #include "Misc/AxisDisplayInfo.h"
@@ -19,8 +21,8 @@
 #include "Modules/ModuleManager.h"
 #include "Nodes/InterchangeBaseNodeContainer.h"
 #include "Nodes/InterchangeSourceNode.h"
-#include "Serialization/JsonWriter.h"
 #include "Serialization/JsonSerializer.h"
+#include "Serialization/JsonWriter.h"
 #include "Serialization/LargeMemoryReader.h"
 #include "StaticMeshOperations.h"
 #include "Texture/InterchangeImageWrapperTranslator.h"
@@ -52,10 +54,85 @@ static FAutoConsoleVariableRef CCvarInterchangeFBXTranslatorUseWorker(
 
 namespace UE::Interchange::Private
 {
+	static const FString FBXAnalyticsIdentifier = TEXT("Interchange.Usage.Import.FBX");
+	
+	enum EFBXTranslatorAnalyticsResult : uint8
+	{
+		Successful,
+		InputFileNotFound,
+		TranslatorFailed,
+		FBXParserFailed,
+		InterchangeDispatcherFailed,
+		InterchangeWorkerFailed,
+	};
+
+	enum EFBXParserType : uint8
+	{
+		Unspecified,
+		FBX_SDK,
+		uFBX,
+	};
+
+	void  AddAnalytics (const EFBXTranslatorAnalyticsResult& TranslatorAnalyticsResult, TSharedPtr<UE::Interchange::FAnalyticsHelper> FBXAnalyticsHelper, bool bIsUsingInterchangeWorker = false, EFBXParserType FBXParser = EFBXParserType::Unspecified)
+	{
+		if (!FBXAnalyticsHelper.IsValid())
+		{
+			return;
+		}
+
+		TArray<FAnalyticsEventAttribute> FBXAnalyticsAttribs;
+
+		if (bIsUsingInterchangeWorker)
+		{
+			FBXParser = EFBXParserType::FBX_SDK;
+		}
+
+		FBXAnalyticsAttribs.Add(FAnalyticsEventAttribute(TEXT("InterchangeWorkerUsed"), bIsUsingInterchangeWorker));
+
+		switch (FBXParser)
+		{
+		case FBX_SDK:
+			FBXAnalyticsAttribs.Add(FAnalyticsEventAttribute(TEXT("FBXParser"), "FBX SDK"));
+			break;
+		case uFBX:
+			FBXAnalyticsAttribs.Add(FAnalyticsEventAttribute(TEXT("FBXParser"), "uFBX"));
+			break;
+		case Unspecified:
+			// Falls through...
+		default:
+			break;
+		}
+
+		switch (TranslatorAnalyticsResult)
+		{
+		case Successful:
+			FBXAnalyticsAttribs.Add(FAnalyticsEventAttribute(TEXT("TranslationStatus"), "Successful."));
+			break;
+		case InputFileNotFound:
+			FBXAnalyticsAttribs.Add(FAnalyticsEventAttribute(TEXT("TranslationStatus"), "[Failed] Input File Not Found."));
+			break;
+		case TranslatorFailed:
+			FBXAnalyticsAttribs.Add(FAnalyticsEventAttribute(TEXT("TranslationStatus"), "[Failed] Interchange Translator Failed."));
+			break;
+		case FBXParserFailed:
+			FBXAnalyticsAttribs.Add(FAnalyticsEventAttribute(TEXT("TranslationStatus"), "[Failed] Parsing Failed."));
+			break;
+		case InterchangeDispatcherFailed:
+			FBXAnalyticsAttribs.Add(FAnalyticsEventAttribute(TEXT("TranslationStatus"), "[Failed] Interchange Dispatcher Failed."));
+			break;
+		case InterchangeWorkerFailed:
+			FBXAnalyticsAttribs.Add(FAnalyticsEventAttribute(TEXT("TranslationStatus"), "[Failed] Interchange Worker Failed."));
+			break;
+		default:
+			break;
+		}
+
+		FBXAnalyticsHelper->AppendThreadSafe(FBXAnalyticsIdentifier, FBXAnalyticsAttribs);
+	}
+
 	void ApplyTranslatorMessage(const UInterchangeFbxTranslator* Translator, const FString& JsonMessage)
 	{
-		UInterchangeResult* InterchangeResult = UInterchangeResult::FromJson(JsonMessage);
-		if (InterchangeResult)
+		if (UInterchangeResult* InterchangeResult = UInterchangeResult::FromJson(JsonMessage))
 		{
 			//Downgrade warning message to display log when we are in automation
 			if (GIsAutomationTesting && InterchangeResult->IsA(UInterchangeResultWarning::StaticClass()))
@@ -198,6 +275,11 @@ TArray<FString> UInterchangeFbxTranslator::GetSupportedFormats() const
 bool UInterchangeFbxTranslator::Translate(UInterchangeBaseNodeContainer& BaseNodeContainer) const
 {
 	TRACE_CPUPROFILER_EVENT_SCOPE(UInterchangeFbxTranslator::Translate);
+	using namespace UE::Interchange::Private;
+	// This is to ensure that the keys are not duplicated when the Analytics Event is actually sent.
+	// Changing the Translator Settings in the import dialog will call the Translate function on the Translators again.
+	AnalyticsHelper->ClearAnalyticsEventData(UE::Interchange::Private::FBXAnalyticsIdentifier);
+
 	if (HasAnyFlags(RF_ClassDefaultObject | RF_ArchetypeObject))
 	{
 		//Class default should never be use has an active translator
@@ -205,6 +287,7 @@ bool UInterchangeFbxTranslator::Translate(UInterchangeBaseNodeContainer& BaseNod
 		UInterchangeResultError_Generic* ErrorResult = AddMessage<UInterchangeResultError_Generic>();
 		ErrorResult->SourceAssetName = SourceData ? SourceData->GetFilename() : FString();
 		ErrorResult->Text = NSLOCTEXT("UInterchangeFbxTranslator", "Translate_DoNotUseClassDefault", "Class default should not be use has an active translator.");
+		AddAnalytics(EFBXTranslatorAnalyticsResult::TranslatorFailed, AnalyticsHelper, bUseWorkerImport);
 		return false;
 	}
 
@@ -213,6 +296,7 @@ bool UInterchangeFbxTranslator::Translate(UInterchangeBaseNodeContainer& BaseNod
 	FString FilePath = GetSourceData()->GetFilename();
 	if (!FPaths::FileExists(FilePath))
 	{
+		AddAnalytics(EFBXTranslatorAnalyticsResult::InputFileNotFound, AnalyticsHelper, bUseWorkerImport);
 		return false;
 	}
 
@@ -234,11 +318,12 @@ bool UInterchangeFbxTranslator::Translate(UInterchangeBaseNodeContainer& BaseNod
 	}
 
 	StartTime = FPlatformTime::Cycles64();
-
+	bool bEncounteredParserErrors = false;
 	if (bUseWorkerImport)
 	{
 		if (!Dispatcher.IsValid())
 		{
+			AddAnalytics(InterchangeDispatcherFailed, AnalyticsHelper, bUseWorkerImport);
 			return false;
 		}
 		//Create a json command to read the fbx file
@@ -252,6 +337,7 @@ bool UInterchangeFbxTranslator::Translate(UInterchangeBaseNodeContainer& BaseNod
 		if (!WorkerFatalError.IsEmpty())
 		{
 			AddMessage(UInterchangeResult::FromJson(WorkerFatalError));
+			AddAnalytics(EFBXTranslatorAnalyticsResult::InterchangeWorkerFailed, AnalyticsHelper, bUseWorkerImport);
 		}
 
 		UE::Interchange::ETaskState TaskState;
@@ -262,11 +348,29 @@ bool UInterchangeFbxTranslator::Translate(UInterchangeBaseNodeContainer& BaseNod
 		// Parse the Json messages into UInterchangeResults
 		for (const FString& JsonMessage : JsonMessages)
 		{
-			UE::Interchange::Private::ApplyTranslatorMessage(this, JsonMessage);
+			if (UInterchangeResult* InterchangeResult = UInterchangeResult::FromJson(JsonMessage))
+			{
+				
+				if (GIsAutomationTesting && InterchangeResult->GetResultType() == EInterchangeResultType::Warning)
+				{
+					UE_LOG(LogInterchangeImport, Display, TEXT("%s"), *InterchangeResult->GetText().ToString());
+				}
+				else
+				{
+					AddMessage(InterchangeResult);
+				}
+
+				if (!bEncounteredParserErrors && InterchangeResult->GetResultType() == EInterchangeResultType::Error)
+				{
+					AddAnalytics(EFBXTranslatorAnalyticsResult::FBXParserFailed, AnalyticsHelper, bUseWorkerImport);
+					bEncounteredParserErrors = true;
+				}
+			}
 		}
 
 		if (TaskState != UE::Interchange::ETaskState::ProcessOk)
 		{
+			AddAnalytics(EFBXTranslatorAnalyticsResult::InterchangeWorkerFailed, AnalyticsHelper, bUseWorkerImport);
 			return false;
 		}
 		//Grab the result file and fill the BaseNodeContainer
@@ -284,6 +388,25 @@ bool UInterchangeFbxTranslator::Translate(UInterchangeBaseNodeContainer& BaseNod
 		FbxParser.SetResultContainer(Results);
 		FbxParser.SetConvertSettings(bConvertScene, bForceFrontXAxis, bConvertSceneUnit, bKeepFbxNamespace);
 		FbxParser.LoadFbxFile(FilePath, BaseNodeContainer);
+
+		TArray<UInterchangeResult*> TranslatorResults = Results->GetResults();
+		if (!TranslatorResults.IsEmpty())
+		{
+			for (const UInterchangeResult* InterchangeResult : TranslatorResults)
+			{
+				if (InterchangeResult->GetResultType() == EInterchangeResultType::Error)
+				{
+					AddAnalytics(
+						EFBXTranslatorAnalyticsResult::FBXParserFailed, 
+						AnalyticsHelper, 
+						bUseWorkerImport, 
+						bUseUfbxParser ? EFBXParserType::uFBX : EFBXParserType::FBX_SDK
+					);
+					bEncounteredParserErrors = true;
+					break;
+				}
+			}
+		}
 #endif
 	}
 
@@ -296,6 +419,16 @@ bool UInterchangeFbxTranslator::Translate(UInterchangeBaseNodeContainer& BaseNod
 #if WITH_EDITOR
 	UE_LOG(LogInterchangeImport, Log, TEXT("Translation (%s) of %s in [%d min %.3f s]"), FbxParser.GetName(), *FileName, ElapsedMin, ElapsedSeconds);
 #endif
+	
+	if (!bEncounteredParserErrors)
+	{
+		AddAnalytics(
+			EFBXTranslatorAnalyticsResult::Successful, 
+			AnalyticsHelper, 
+			bUseWorkerImport, 
+			bUseUfbxParser ? EFBXParserType::uFBX : EFBXParserType::FBX_SDK
+		);
+	}
 
 	return true;
 }
diff --git a/Engine/Source/Runtime/Interchange/Engine/Private/InterchangeAnalyticsHelper.cpp b/Engine/Source/Runtime/Interchange/Engine/Private/InterchangeAnalyticsHelper.cpp
--- a/Engine/Source/Runtime/Interchange/Engine/Private/InterchangeAnalyticsHelper.cpp
+++ b/Engine/Source/Runtime/Interchange/Engine/Private/InterchangeAnalyticsHelper.cpp
@@ -13,6 +13,12 @@ namespace UE::Interchange
 		Analytics.Append(ToAdd);
 	}
 
+	void FAnalyticsHelper::ClearAnalyticsEventData(const FString& Identifier)
+	{
+		FScopeLock ScopeLock(&CriticalSection);
+		AnalyticsAttributes.Remove(Identifier);
+	}
+
 	//Make sure it is safe to call this directly.
 	void FAnalyticsHelper::Add(const FString& Identifier, const FAnalyticsEventAttribute& Entry)
 	{
diff --git a/Engine/Source/Runtime/Interchange/Engine/Public/InterchangeAnalyticsHelper.h b/Engine/Source/Runtime/Interchange/Engine/Public/InterchangeAnalyticsHelper.h
--- a/Engine/Source/Runtime/Interchange/Engine/Public/InterchangeAnalyticsHelper.h
+++ b/Engine/Source/Runtime/Interchange/Engine/Public/InterchangeAnalyticsHelper.h
@@ -9,14 +9,12 @@ namespace UE::Interchange
 {
 	struct FAnalyticsHelper
 	{
-	private:
-		//Analytics Attributes to share with the Translators. (to be sent when Import is finished)
-		TMap<FString, TArray<FAnalyticsEventAttribute>> AnalyticsAttributes;
-		FCriticalSection CriticalSection;
-
 	public:
 		INTERCHANGEENGINE_API void AppendThreadSafe(const FString& Identifier, const TArray<FAnalyticsEventAttribute>& ToAdd);
 
+		// Remove all the event attributes if the identifier matches.
+		INTERCHANGEENGINE_API void ClearAnalyticsEventData(const FString& Identifier);
+
 		//Make sure it is safe to call this directly.
 		INTERCHANGEENGINE_API void Add(const FString& Identifier, const FAnalyticsEventAttribute& Entry);
 
@@ -25,5 +23,10 @@ namespace UE::Interchange
 		INTERCHANGEENGINE_API void AddThreadSafe(const FString& Identifier, const FAnalyticsEventAttribute& Entry);
 
 		INTERCHANGEENGINE_API void SendAnalytics();
+
+	private:
+		//Analytics Attributes to share with the Translators. (to be sent when Import is finished)
+		TMap<FString, TArray<FAnalyticsEventAttribute>> AnalyticsAttributes;
+		FCriticalSection CriticalSection;
 	};
 };
\ No newline at end of file
