[Interchange] uFBX improvements

- Morph animation prototype
- multiple skeletons support
- Rigid animation
- Mesh LOD supported

diff --git a/Engine/Plugins/Interchange/Runtime/Source/Parsers/Fbx/Private/Ufbx/UfbxAnimation.cpp b/Engine/Plugins/Interchange/Runtime/Source/Parsers/Fbx/Private/Ufbx/UfbxAnimation.cpp
--- a/Engine/Plugins/Interchange/Runtime/Source/Parsers/Fbx/Private/Ufbx/UfbxAnimation.cpp
+++ b/Engine/Plugins/Interchange/Runtime/Source/Parsers/Fbx/Private/Ufbx/UfbxAnimation.cpp
@@ -2,26 +2,306 @@
 
 #include "UfbxAnimation.h"
 #include "UfbxParser.h"
+#include "UfbxScene.h"
 #include "UfbxConvert.h"
 
 #include "InterchangeCommonAnimationPayload.h"
+#include "Algo/AnyOf.h"
+#include "Nodes/InterchangeBaseNodeContainer.h"
+
+#define LOCTEXT_NAMESPACE "InterchangeFbxParser"
 
 namespace UE::Interchange::Private
 {
 
-	namespace Impl
+	namespace AnimUtils
 	{
 
 		template<typename ValueList>
 		auto GetAnimationKeyValue(const ValueList& Keys, const bool& IsConstant, const int32 TransformIndex, const double& Frequency) -> decltype(Keys[0].value)
 		{
-			// todo: investigate key start time handling, e.g. TransformIndex-FMath::CeilToInt(Keys[0].time*Frequency - UE_KINDA_SMALL_NUMBER)
+			// #ufbx_todo: investigate key start time handling, e.g. TransformIndex-FMath::CeilToInt(Keys[0].time*Frequency - UE_KINDA_SMALL_NUMBER)
 			int32 KeyIndex = IsConstant ? 0
 			: FMath::Clamp<int32>(TransformIndex, 0, Keys.count-1);
 
 			return Keys[KeyIndex].value;
 		};
 
+		// inspired by UE::Interchange::Private::ImportCurve - Engine/Plugins/Interchange/Runtime/Source/Parsers/Fbx/Private/FbxAnimation.cpp
+		void ConvertCurve(const ufbx_anim_curve& SourceFloatCurves, const float ScaleValue, TArray<FInterchangeCurveKey>& DestinationFloatCurve)
+		{
+			const float DefaultCurveWeight = FbxAnimCurveDef::sDEFAULT_WEIGHT;
+			
+			int32 KeyCount = SourceFloatCurves.keyframes.count;
+			for (int32 KeyIndex = 0; KeyIndex < KeyCount; ++KeyIndex)
+			{
+				
+				const ufbx_keyframe& Key = SourceFloatCurves.keyframes[KeyIndex];
+				
+				const float KeyTimeValue = static_cast<float>(Key.time);
+				const float Value = Key.value * ScaleValue;
+				FInterchangeCurveKey& InterchangeCurveKey = DestinationFloatCurve.AddDefaulted_GetRef();
+				InterchangeCurveKey.Time = KeyTimeValue;
+				InterchangeCurveKey.Value = Value;
+
+				const bool bIncludeOverrides = true;
+
+				// #ufbx_todo: 
+				// FbxAnimCurveDef::EWeightedMode KeyTangentWeightMode = Key.GetTangentWeightMode();
+
+				EInterchangeCurveInterpMode NewInterpMode = EInterchangeCurveInterpMode::Linear;
+				EInterchangeCurveTangentMode NewTangentMode = EInterchangeCurveTangentMode::Auto;
+				EInterchangeCurveTangentWeightMode NewTangentWeightMode = EInterchangeCurveTangentWeightMode::WeightedNone;
+
+				//#ufbx_todo:  tangents - investigate, implement and refine how tangents are done in ufbx compared to FBX SDK
+				float RightTangent = [Key, ScaleValue]
+				{
+					if (FMath::Abs(Key.right.dx) > UE_SMALL_NUMBER)
+					{
+						return Key.right.dy/Key.right.dx * ScaleValue;
+					}
+
+					return 0.f;
+				}();
+
+				float LeftTangent = [Key, ScaleValue]
+				{
+					if (FMath::Abs(Key.left.dx) > UE_SMALL_NUMBER)
+					{
+						 return Key.left.dy/Key.left.dx * ScaleValue; 
+					}
+
+					return 0.f;
+				}();
+
+				float RightTangentWeight = 0.0f;
+				float LeftTangentWeight = 0.0f; //This one is dependent on the previous key.
+				bool bLeftWeightActive = false;
+				bool bRightWeightActive = false;
+
+				const bool bPreviousKeyValid = KeyIndex > 0;
+				const bool bNextKeyValid = KeyIndex < KeyCount - 1;
+				float PreviousValue = 0.0f;
+				float PreviousKeyTimeValue = 0.0f;
+				float NextValue = 0.0f;
+				float NextKeyTimeValue = 0.0f;
+				if (bPreviousKeyValid)
+				{
+					const ufbx_keyframe& PreviousKey = SourceFloatCurves.keyframes[KeyIndex - 1];
+					FbxTime PreviousKeyTime = PreviousKey.time;
+					PreviousKeyTimeValue = static_cast<float>(PreviousKeyTime.GetSecondDouble());
+					PreviousValue = PreviousKey.value * ScaleValue;
+					//The left tangent is driven by the previous key. If the previous key have a the NextLeftweight or both flag weighted mode, it mean the next key is weighted on the left side
+
+					// #ufbx_todo:
+					// bLeftWeightActive = (PreviousKey.GetTangentWeightMode() & FbxAnimCurveDef::eWeightedNextLeft) > 0;
+					// if (bLeftWeightActive)
+					// {
+					// 	LeftTangentWeight = PreviousKey.GetDataFloat(FbxAnimCurveDef::eNextLeftWeight);
+					// }
+				}
+				if (bNextKeyValid)
+				{
+					const ufbx_keyframe& NextKey = SourceFloatCurves.keyframes[KeyIndex + 1];
+					FbxTime NextKeyTime = NextKey.time;
+					NextKeyTimeValue = static_cast<float>(NextKeyTime.GetSecondDouble());
+					NextValue = NextKey.value * ScaleValue;
+
+					// #ufbx_todo:
+					// bRightWeightActive = (KeyTangentWeightMode & FbxAnimCurveDef::eWeightedRight) > 0;
+					if (bRightWeightActive)
+					{
+						//The right tangent weight should be use only if we are not the last key since the last key do not have a right tangent.
+						//Use the current key to gather the right tangent weight
+						// #ufbx_todo: 
+						// RightTangentWeight = Key.GetDataFloat(FbxAnimCurveDef::eRightWeight);
+					}
+				}
+
+				// When this flag is true, the tangent is flat if the value has the same value as the previous or next key.
+				// FbxAnimCurveDef::ETangentMode KeyTangentMode = Key.GetTangentMode(bIncludeOverrides);
+				const bool bTangentGenericClamp = false; //(KeyTangentMode & FbxAnimCurveDef::eTangentGenericClamp);
+				//Time independent tangent this is consider has a spline tangent key
+				const bool bTangentGenericTimeIndependent = true; //(KeyTangentMode & FbxAnimCurveDef::ETangentMode::eTangentGenericTimeIndependent);
+				// When this flag is true, the tangent is flat if the value is outside of the [previous key, next key] value range.
+				//Clamp progressive is (eTangentGenericClampProgressive |eTangentGenericTimeIndependent)
+				const bool bTangentGenericClampProgressive = true; //(KeyTangentMode & FbxAnimCurveDef::ETangentMode::eTangentGenericClampProgressive) == FbxAnimCurveDef::ETangentMode::eTangentGenericClampProgressive;
+				// if (KeyTangentMode & FbxAnimCurveDef::eTangentGenericBreak)
+				// {
+				// 	NewTangentMode = EInterchangeCurveTangentMode::Break;
+				// }
+				// else if (KeyTangentMode & FbxAnimCurveDef::eTangentUser)
+				// {
+				// 	NewTangentMode = EInterchangeCurveTangentMode::User;
+				// }
+
+				
+				switch (Key.interpolation)
+				{
+					// #ufbx_todo: distinguish two const interps
+				case UFBX_INTERPOLATION_CONSTANT_NEXT://! Constant value until next key.
+				case UFBX_INTERPOLATION_CONSTANT_PREV://! Constant value until next key.
+					NewInterpMode = EInterchangeCurveInterpMode::Constant;
+					break;
+				case UFBX_INTERPOLATION_LINEAR://! Linear progression to next key.
+					NewInterpMode = EInterchangeCurveInterpMode::Linear;
+					break;
+				case UFBX_INTERPOLATION_CUBIC://! Cubic progression to next key.
+					NewInterpMode = EInterchangeCurveInterpMode::Cubic;
+					// get tangents
+					{
+						bool bIsFlatTangent = false;
+						bool bIsComputedTangent = false;
+						if (bTangentGenericClampProgressive)
+						{
+							if (bPreviousKeyValid && bNextKeyValid)
+							{
+								const float PreviousNextHalfDelta = (NextValue - PreviousValue) * 0.5f;
+								const float PreviousNextAverage = PreviousValue + PreviousNextHalfDelta;
+								// If the value is outside of the previous-next value range, the tangent is flat.
+								bIsFlatTangent = FMath::Abs(Value - PreviousNextAverage) >= FMath::Abs(PreviousNextHalfDelta);
+							}
+							else
+							{
+								//Start/End tangent with the ClampProgressive flag are flat.
+								bIsFlatTangent = true;
+							}
+						}
+						else if (bTangentGenericClamp && (bPreviousKeyValid || bNextKeyValid))
+						{
+							if (bPreviousKeyValid && PreviousValue == Value)
+							{
+								bIsFlatTangent = true;
+							}
+							if (bNextKeyValid)
+							{
+								bIsFlatTangent |= Value == NextValue;
+							}
+						}
+						else if (bTangentGenericTimeIndependent)
+						{
+							//Spline tangent key, because bTangentGenericClampProgressive include bTangentGenericTimeIndependent, we must treat this case after bTangentGenericClampProgressive
+							if (KeyCount == 1)
+							{
+								bIsFlatTangent = true;
+							}
+							else
+							{
+								//Spline tangent key must be User mode since we want to keep the tangents provide by the fbx key left and right derivatives
+								NewTangentMode = EInterchangeCurveTangentMode::User;
+							}
+						}
+
+						if (bIsFlatTangent)
+						{
+							RightTangent = 0;
+							LeftTangent = 0;
+							//To force flat tangent we need to set the tangent mode to user
+							NewTangentMode = EInterchangeCurveTangentMode::User;
+						}
+
+					}
+					break;
+				}
+
+				//auto with weighted give the wrong result, so when auto is weighted we set user mode and set the Right tangent equal to the left tangent.
+				//Auto has only the left tangent set
+				if (NewTangentMode == EInterchangeCurveTangentMode::Auto && (bLeftWeightActive || bRightWeightActive))
+				{
+
+					NewTangentMode = EInterchangeCurveTangentMode::User;
+					RightTangent = LeftTangent;
+				}
+
+				if (NewTangentMode != EInterchangeCurveTangentMode::Auto)
+				{
+					const bool bEqualTangents = FMath::IsNearlyEqual(LeftTangent, RightTangent);
+					//If tangents are different then broken.
+					if (bEqualTangents)
+					{
+						NewTangentMode = EInterchangeCurveTangentMode::User;
+					}
+					else
+					{
+						NewTangentMode = EInterchangeCurveTangentMode::Break;
+					}
+				}
+
+				//Only cubic interpolation allow weighted tangents
+				if (Key.interpolation == UFBX_INTERPOLATION_CUBIC)
+				{
+					if (bLeftWeightActive && bRightWeightActive)
+					{
+						NewTangentWeightMode = EInterchangeCurveTangentWeightMode::WeightedBoth;
+					}
+					else if (bLeftWeightActive)
+					{
+						NewTangentWeightMode = EInterchangeCurveTangentWeightMode::WeightedArrive;
+						RightTangentWeight = DefaultCurveWeight;
+					}
+					else if (bRightWeightActive)
+					{
+						NewTangentWeightMode = EInterchangeCurveTangentWeightMode::WeightedLeave;
+						LeftTangentWeight = DefaultCurveWeight;
+					}
+					else
+					{
+						NewTangentWeightMode = EInterchangeCurveTangentWeightMode::WeightedNone;
+						LeftTangentWeight = DefaultCurveWeight;
+						RightTangentWeight = DefaultCurveWeight;
+					}
+
+					auto ComputeWeightInternal = [](float TimeA, float TimeB, const float TangentSlope, const float TangentWeight)
+					{
+						const float X = TimeA - TimeB;
+						const float Y = TangentSlope * X;
+						return FMath::Sqrt(X * X + Y * Y) * TangentWeight;
+					};
+
+					if (!FMath::IsNearlyZero(LeftTangentWeight))
+					{
+						if (bPreviousKeyValid)
+						{
+							LeftTangentWeight = ComputeWeightInternal(KeyTimeValue, PreviousKeyTimeValue, LeftTangent, LeftTangentWeight);
+						}
+						else
+						{
+							LeftTangentWeight = 0.0f;
+						}
+					}
+
+					if (!FMath::IsNearlyZero(RightTangentWeight))
+					{
+						if (bNextKeyValid)
+						{
+							RightTangentWeight = ComputeWeightInternal(NextKeyTimeValue, KeyTimeValue, RightTangent, RightTangentWeight);
+						}
+						else
+						{
+							RightTangentWeight = 0.0f;
+						}
+					}
+				}
+
+				const bool bForceDisableTangentRecompute = false; //No need to recompute all the tangents of the curve every time we change de key.
+				InterchangeCurveKey.InterpMode = NewInterpMode;
+				InterchangeCurveKey.TangentMode = NewTangentMode;
+				InterchangeCurveKey.TangentWeightMode = NewTangentWeightMode;
+
+				InterchangeCurveKey.ArriveTangent = LeftTangent;
+				InterchangeCurveKey.LeaveTangent = RightTangent;
+				InterchangeCurveKey.ArriveTangentWeight = LeftTangentWeight;
+				InterchangeCurveKey.LeaveTangentWeight = RightTangentWeight;
+			}
+		}
+
+		// UE::Interchange::Private::CreateTrackNodeUid
+		FString CreateTrackNodeUid(const FString& JointUid, const int32 AnimationIndex)
+		{
+			FString TrackNodeUid = TEXT("\\SkeletalAnimation\\") + JointUid + TEXT("\\") + FString::FromInt(AnimationIndex);
+			return TrackNodeUid;
+		}
+
 	}
 
 	void FUfbxAnimation::FetchSkinnedAnimation(const UE::Interchange::FAnimationPayloadQuery& PayloadQuery, const FSkeletalAnimationPayloadContext& SkeletalAnimationPayload, FAnimationPayloadData& PayloadData)
@@ -43,9 +323,9 @@ namespace UE::Interchange::Private
 		auto MakeTransform = [&BakedNode, BakeFrequency](const int32 TransformIndex)
 		{
 			FTransform Transform;
-			FVector Translation = Convert::ConvertVec3(Impl::GetAnimationKeyValue(BakedNode.translation_keys, BakedNode.constant_translation, TransformIndex, BakeFrequency));
-			FQuat Rotation = Convert::ConvertQuat(Impl::GetAnimationKeyValue(BakedNode.rotation_keys, BakedNode.constant_rotation, TransformIndex, BakeFrequency));
-			FVector Scale = Convert::ConvertVec3(Impl::GetAnimationKeyValue(BakedNode.scale_keys, BakedNode.constant_scale, TransformIndex, BakeFrequency));
+			FVector Translation = Convert::ConvertVec3(AnimUtils::GetAnimationKeyValue(BakedNode.translation_keys, BakedNode.constant_translation, TransformIndex, BakeFrequency));
+			FQuat Rotation = Convert::ConvertQuat(AnimUtils::GetAnimationKeyValue(BakedNode.rotation_keys, BakedNode.constant_rotation, TransformIndex, BakeFrequency));
+			FVector Scale = Convert::ConvertVec3(AnimUtils::GetAnimationKeyValue(BakedNode.scale_keys, BakedNode.constant_scale, TransformIndex, BakeFrequency));
 
 			Transform.SetTranslation(Translation);
 			Transform.SetRotation(Rotation);
@@ -59,4 +339,309 @@ namespace UE::Interchange::Private
 		}
 	}
 
+	bool FUfbxAnimation::FetchMorphTargetAnimation(const FUfbxParser& Parser, const FMorphAnimationPayloadContext& Animation, TArray<FInterchangeCurve>& OutCurves)
+	{
+		if (Animation.Curve)
+		{
+			// #ufbx_todo: UE::Interchange::Private::FAnimationPayloadContext::InternalFetchMorphTargetCurvePayload uses 0.01f scale
+			AnimUtils::ConvertCurve(*Animation.Curve, 0.01f, OutCurves.AddDefaulted_GetRef().Keys);
+			return true;
+		}
+		
+		return false;
+	}
+
+	bool FUfbxAnimation::FetchRigidAnimation(const FUfbxParser& Parser, const FRigidAnimationPayloadContext& Animation,
+		TArray<FInterchangeCurve>& OutCurves)
+	{
+		auto AddCurve = [&OutCurves](const ufbx_anim_prop* Prop, int32 CurveIndex, double Scale=1.0)
+		{
+			const ufbx_anim_curve* Curve =  Prop ? Prop->anim_value->curves[CurveIndex] : nullptr;
+			if (Curve)
+			{
+				AnimUtils::ConvertCurve(*Curve, Scale, OutCurves.AddDefaulted_GetRef().Keys);
+			}
+			else
+			{
+				OutCurves.AddDefaulted();
+			}
+		};
+
+		OutCurves.Reserve(9);
+
+		// #ufbx_todo: FAnimationPayloadContext::InternalFetchCurveNodePayloadToFile does some pivot handling, see ResetPivotsPrePostRotationsAndSetRotationOrder
+
+		AddCurve(Animation.TranslationProp, 0);
+		AddCurve(Animation.TranslationProp, 2, -1.0);
+		AddCurve(Animation.TranslationProp, 1);
+		AddCurve(Animation.RotationProp, 0);
+		AddCurve(Animation.RotationProp, 1);
+		AddCurve(Animation.RotationProp, 2);
+		AddCurve(Animation.ScalingProp, 0);
+		AddCurve(Animation.ScalingProp, 2);
+		AddCurve(Animation.ScalingProp, 1);
+
+		return true;
+	}
+
+	void FUfbxAnimation::AddAnimation(FUfbxParser& Parser, FUfbxScene& Scene, UInterchangeBaseNodeContainer& NodeContainer,
+	                                  const TMap<ufbx_blend_shape*, ufbx_mesh*> MeshPerBlendShape)
+	{
+		double FrameRate = Parser.Scene->settings.frames_per_second;
+
+		TArray<FString> TransformAnimTrackNodeUids;
+
+		int32 AnimStackCount = Parser.Scene->anim_stacks.count;
+		for (int32 AnimStackIndex = 0; AnimStackIndex < AnimStackCount; ++AnimStackIndex)
+		{
+			ufbx_anim_stack* Stack = Parser.Scene->anim_stacks[AnimStackIndex];
+
+			// keep track nodes for current AnimStackIndex only (not all TrackNodes!)
+			TMap<FString, UInterchangeSkeletalAnimationTrackNode*> TrackNodePerSkeletonRoot;
+
+			ufbx_bake_opts BakeOptions = {};
+			BakeOptions.maximum_sample_rate = FrameRate;
+			BakeOptions.minimum_sample_rate = FrameRate;
+			BakeOptions.resample_rate = FrameRate;
+
+			ufbx_error ErrorResult;
+
+			ufbx_shared_ptr<ufbx_baked_anim> BakedAnim(ufbx_bake_anim(Parser.Scene, Stack->anim, &BakeOptions, &ErrorResult));
+			if (!BakedAnim)
+			{
+				UInterchangeResultError_Generic* Message = Parser.AddMessage<UInterchangeResultError_Generic>();
+
+				TArray<char> ErrorBuf;
+				ErrorBuf.SetNum(ErrorResult.info_length+64);
+				ufbx_format_error(ErrorBuf.GetData(), ErrorBuf.Num(), &ErrorResult);
+
+				Message->Text = FText::Format(LOCTEXT("AnimBakeError", "UFBX couldn't bake animation: '{0}'."), FText::FromString(UTF8_TO_TCHAR(ErrorBuf.GetData())));
+				continue;
+			}
+
+			TConstArrayView<ufbx_baked_node> BakedNodes(BakedAnim->nodes.data, BakedAnim->nodes.count);
+			if (BakedNodes.IsEmpty())
+			{
+				continue;
+			}
+
+			TConstArrayView<ufbx_node*> FbxNodes(Parser.Scene->nodes.data, Parser.Scene->nodes.count);
+
+			FString DisplayString = Convert::ToUnrealString(Stack->name);
+
+			auto GetTrackNode = [&](ufbx_node* Node)
+			{
+				FString SkeletonRootNodeUid = Scene.GetSkeletonRoot(Node); 
+
+				UInterchangeSkeletalAnimationTrackNode*& SkeletalAnimationTrackNode = TrackNodePerSkeletonRoot.FindOrAdd(SkeletonRootNodeUid);
+				if (!SkeletalAnimationTrackNode)
+				{
+					FString TrackNodeUid = AnimUtils::CreateTrackNodeUid(SkeletonRootNodeUid, AnimStackIndex);
+
+					SkeletalAnimationTrackNode = NewObject< UInterchangeSkeletalAnimationTrackNode >(&NodeContainer);
+					SkeletalAnimationTrackNode->InitializeNode(TrackNodeUid, DisplayString, EInterchangeNodeContainerType::TranslatedAsset);
+					SkeletalAnimationTrackNode->AddBooleanAttribute(TEXT("RenameLikeLegacyFbx"), true);
+					SkeletalAnimationTrackNode->SetCustomAnimationSampleRate(Parser.Scene->settings.frames_per_second);
+
+					SkeletalAnimationTrackNode->SetCustomSkeletonNodeUid(SkeletonRootNodeUid);
+				
+					SkeletalAnimationTrackNode->SetCustomAnimationStartTime(Stack->anim->time_begin);
+					SkeletalAnimationTrackNode->SetCustomAnimationStopTime(Stack->anim->time_end);
+
+					SkeletalAnimationTrackNode->SetCustomSourceTimelineAnimationStartTime(Stack->time_begin);
+					SkeletalAnimationTrackNode->SetCustomSourceTimelineAnimationStopTime(Stack->time_end);
+					NodeContainer.AddNode(SkeletalAnimationTrackNode);
+				}
+				return SkeletalAnimationTrackNode;
+			};
+
+			for (int32 BakedNodeIndex = 0; BakedNodeIndex < BakedAnim->nodes.count; ++BakedNodeIndex)
+			{
+				const ufbx_baked_node& BakedNode = BakedAnim->nodes[BakedNodeIndex];
+
+				if (!ensureMsgf(FbxNodes.IsValidIndex(BakedNode.typed_id), TEXT("Unexpected type_id encountered in ufbx baked anim nodes")))
+				{
+					continue;
+				}
+
+				ufbx_node* Node = FbxNodes[BakedNode.typed_id];
+				if (Node->bone)
+				{
+					FString NodeUid = Parser.GetNodeUid(*Node);
+
+					UInterchangeSkeletalAnimationTrackNode* SkeletalAnimationTrackNode = GetTrackNode(Node);
+
+					// #ufbx_todo: Need check animated parents - whole subtree has to be animated too?
+					// see UE::Interchange::Private::FFbxAnimation::AddSkeletalTransformAnimation
+					// or maybe there's an option in ufbx baking that already does that?
+
+					FString PayLoadKey = NodeUid + TEXT("_") + FString::FromInt(AnimStackIndex) + TEXT("_SkeletalAnimationPayloadKey");
+
+					Parser.PayloadContexts.Add(PayLoadKey, FSkeletalAnimationPayloadContext{NodeUid, AnimStackIndex, BakedNodeIndex, BakedAnim});
+
+					SkeletalAnimationTrackNode->SetAnimationPayloadKeyForSceneNodeUid(NodeUid, PayLoadKey, EInterchangeAnimationPayLoadType::BAKED);
+				}
+
+			}
+
+			// Collect curves for other animated properties - morph, ..
+			for (int32 LayerIndex = 0; LayerIndex < Stack->layers.count; ++LayerIndex)
+			{
+				ufbx_anim_layer* Layer = Stack->layers[LayerIndex];
+
+				TMap<const ufbx_node*, FRigidAnimationPayloadContext> TransformationPayloadContexts;
+
+				for (int32 AnimPropIndex = 0; AnimPropIndex < Layer->anim_props.count; ++AnimPropIndex)
+				{
+					const ufbx_anim_prop& Prop = Layer->anim_props[AnimPropIndex];
+
+					if (FCStringAnsi::Strcmp(UFBX_DeformPercent, Prop.prop_name.data)==0)
+					{
+						if (ufbx_blend_channel* Channel = ufbx_as_blend_channel(Prop.element))
+						{
+							ufbx_blend_shape* TargetShape = Channel->target_shape;
+
+							ufbx_anim_curve* MorphCurve = Prop.anim_value->curves[0];
+
+							Convert::FMeshNameAndUid MorphTargetNameId(Parser, *TargetShape);
+
+							// from UE::Interchange::Private::FFbxAnimation::AddMorphTargetCurvesAnimation
+							// #ufbx_todo: do we need different PayloadKey for same curve used to animate different meshes?
+							// (MorphTargetAnimationBuildingData.InterchangeMeshNode ? MorphTargetAnimationBuildingData.InterchangeMeshNode->GetUniqueID() : FString()) //Same shape can be animated on different mesh node
+							// If not - maybe we could just create key based on curve's global(within scene) index - ufbx_scene::anim_curves?
+
+							FString PayloadKey = MorphTargetNameId.UniqueID
+								+ TEXT("\\") + FString::FromInt(AnimStackIndex)
+								+ TEXT("\\") + FString::FromInt(LayerIndex)
+								+ TEXT("\\") + FString::FromInt(AnimPropIndex)
+								+ TEXT("_CurveAnimationPayloadKey");
+
+							const ufbx_mesh*const* MeshFound = MeshPerBlendShape.Find(TargetShape);
+
+							const ufbx_mesh* Mesh = *MeshFound;
+							TSet<UInterchangeSkeletalAnimationTrackNode*> Tracks;
+							for (ufbx_skin_deformer* Deformer : Mesh->skin_deformers)
+							{
+								for (ufbx_skin_cluster* Cluster : Deformer->clusters)
+								{
+									Tracks.Add(GetTrackNode(Cluster->bone_node));
+								}
+							}
+
+							for (UInterchangeSkeletalAnimationTrackNode* Track : Tracks)
+							{
+								Track->SetAnimationPayloadKeyForMorphTargetNodeUid(MorphTargetNameId.UniqueID, PayloadKey, EInterchangeAnimationPayLoadType::MORPHTARGETCURVE);
+							}
+
+							FMorphAnimationPayloadContext MorphPayloadContext;
+							MorphPayloadContext.Curve = MorphCurve;
+							Parser.PayloadContexts.Add(PayloadKey, MorphPayloadContext);
+
+						}
+					}
+
+					const ufbx_node* Node = ufbx_as_node(Prop.element);
+					// #ufbx_todo: ignore bones? Btw, why bones are not here?
+					if (Node && !Node->bone)
+					{
+						if (FCStringAnsi::Strcmp(UFBX_Lcl_Translation, Prop.prop_name.data)==0)
+						{
+							TransformationPayloadContexts.FindOrAdd(Node).TranslationProp = &Prop;
+						}
+						if (FCStringAnsi::Strcmp(UFBX_Lcl_Rotation, Prop.prop_name.data)==0)
+						{
+							TransformationPayloadContexts.FindOrAdd(Node).RotationProp = &Prop;
+						}
+						if (FCStringAnsi::Strcmp(UFBX_Lcl_Scaling, Prop.prop_name.data)==0)
+						{
+							TransformationPayloadContexts.FindOrAdd(Node).ScalingProp = &Prop;
+						}
+					}
+				}
+
+				for (const TPair<const ufbx_node*, FRigidAnimationPayloadContext>& Pair  : TransformationPayloadContexts)
+				{
+					const ufbx_node* Node = Pair.Key;
+					const FRigidAnimationPayloadContext& RigidAnimationPayloadContext = Pair.Value;
+
+					// inspired by UE::Interchange::Private::FFbxAnimation::AddRigidTransformAnimation
+
+					constexpr int32 TranslationChannel = 0x0001 | 0x0002 | 0x0004;
+					constexpr int32 RotationChannel = 0x0008 | 0x0010 | 0x0020;
+					constexpr int32 ScaleChannel = 0x0040 | 0x0080 | 0x0100;
+
+					int32 UsedChannels = 0;
+					if (RigidAnimationPayloadContext.TranslationProp && Algo::AnyOf(MakeConstArrayView(RigidAnimationPayloadContext.TranslationProp->anim_value->curves)))
+					{
+						UsedChannels |= TranslationChannel;
+					}
+					if (RigidAnimationPayloadContext.RotationProp && Algo::AnyOf(MakeConstArrayView(RigidAnimationPayloadContext.RotationProp->anim_value->curves)))
+					{
+						UsedChannels |= RotationChannel;
+					}
+					if (RigidAnimationPayloadContext.ScalingProp && Algo::AnyOf(MakeConstArrayView(RigidAnimationPayloadContext.ScalingProp->anim_value->curves)))
+					{
+						UsedChannels |= ScaleChannel;
+					}
+
+					if (UsedChannels != 0)
+					{
+						if (UInterchangeSceneNode** Found = Parser.ElementIdToSceneNode.Find(Node->element_id))
+						{
+							UInterchangeSceneNode* SceneNode = *Found; 
+
+							const FString PayloadKey = Parser.GetNodeUid(*Node) + TEXT("_RigidAnimationPayloadKey");
+
+							Parser.PayloadContexts.Add(PayloadKey, RigidAnimationPayloadContext);
+
+							UInterchangeTransformAnimationTrackNode* TransformAnimTrackNode = NewObject< UInterchangeTransformAnimationTrackNode >(&NodeContainer);
+
+							const FString TransformAnimTrackNodeName = FString::Printf(TEXT("%s"), *SceneNode->GetDisplayLabel());
+							const FString TransformAnimTrackNodeUid = TEXT("\\AnimationTrack\\") + TransformAnimTrackNodeName;
+
+							NodeContainer.SetupNode(TransformAnimTrackNode, TransformAnimTrackNodeUid, TransformAnimTrackNodeName, EInterchangeNodeContainerType::TranslatedAsset);
+							TransformAnimTrackNode->SetCustomActorDependencyUid(*SceneNode->GetUniqueID());
+							TransformAnimTrackNode->SetCustomAnimationPayloadKey(PayloadKey, EInterchangeAnimationPayLoadType::CURVE);
+							TransformAnimTrackNode->SetCustomUsedChannels(UsedChannels);
+
+							// #ufbx_todo: UE::Interchange::Private::ProcessCustomAttributes
+							// ProcessCustomAttributes(Parser, Node, TransformAnimTrackNode);
+
+							TransformAnimTrackNodeUids.Add(TransformAnimTrackNode->GetUniqueID());
+						}
+
+					}
+
+				}
+			}
+		}
+
+		//Only one Track Set Node per fbx file:
+		if (TransformAnimTrackNodeUids.Num() > 0)
+		{
+			FString Name;
+			if (UInterchangeSceneNode** RootSceneNodeFound = Parser.ElementIdToSceneNode.Find(Parser.Scene->root_node->element_id))
+			{
+				UInterchangeSceneNode* RootSceneNode = *RootSceneNodeFound;
+				Name = RootSceneNode->GetName();
+			}
+
+			UInterchangeAnimationTrackSetNode* TrackSetNode = NewObject< UInterchangeAnimationTrackSetNode >(&NodeContainer);
+
+			TrackSetNode->SetCustomFrameRate(FrameRate);
+
+			const FString AnimTrackSetNodeUid = TEXT("\\Animation\\") + Name;
+			const FString AnimTrackSetNodeDisplayLabel = Name + TEXT("_TrackSetNode");
+
+			NodeContainer.SetupNode(TrackSetNode, AnimTrackSetNodeUid, AnimTrackSetNodeDisplayLabel, EInterchangeNodeContainerType::TranslatedAsset);
+
+			for (const FString& TransformAnimTrackNodeUid : TransformAnimTrackNodeUids)
+			{
+				TrackSetNode->AddCustomAnimationTrackUid(TransformAnimTrackNodeUid);
+			}
+		}
+
+	}
 }
+
+#undef LOCTEXT_NAMESPACE
diff --git a/Engine/Plugins/Interchange/Runtime/Source/Parsers/Fbx/Private/Ufbx/UfbxAnimation.h b/Engine/Plugins/Interchange/Runtime/Source/Parsers/Fbx/Private/Ufbx/UfbxAnimation.h
--- a/Engine/Plugins/Interchange/Runtime/Source/Parsers/Fbx/Private/Ufbx/UfbxAnimation.h
+++ b/Engine/Plugins/Interchange/Runtime/Source/Parsers/Fbx/Private/Ufbx/UfbxAnimation.h
@@ -3,6 +3,7 @@
 #pragma once
 
 #include "CoreMinimal.h"
+#include "UfbxParser.h"
 
 namespace UE::Interchange
 {
@@ -10,13 +11,20 @@ namespace UE::Interchange
 	struct FAnimationPayloadData;
 }
 
+struct FInterchangeCurve;
+
 namespace UE::Interchange::Private
 {
 	struct FSkeletalAnimationPayloadContext;
+	class FUfbxScene;
 
 	class FUfbxAnimation
 	{
 	public:
 		static void FetchSkinnedAnimation(const UE::Interchange::FAnimationPayloadQuery& PayloadQuery, const FSkeletalAnimationPayloadContext& SkeletalAnimationPayload, FAnimationPayloadData& PayloadData);
+		static bool FetchMorphTargetAnimation(const FUfbxParser& Parser, const FMorphAnimationPayloadContext& Animation, TArray<FInterchangeCurve>& OutCurves);
+		static bool FetchRigidAnimation(const FUfbxParser& Parser, const FRigidAnimationPayloadContext& Animation, TArray<FInterchangeCurve>& OutCurves);
+
+		static void AddAnimation(FUfbxParser& Parser, FUfbxScene& Scene, UInterchangeBaseNodeContainer& NodeContainer, const TMap<ufbx_blend_shape*, ufbx_mesh*> MeshPerBlendShape);
 	};
 }
diff --git a/Engine/Plugins/Interchange/Runtime/Source/Parsers/Fbx/Private/Ufbx/UfbxMaterial.cpp b/Engine/Plugins/Interchange/Runtime/Source/Parsers/Fbx/Private/Ufbx/UfbxMaterial.cpp
--- a/Engine/Plugins/Interchange/Runtime/Source/Parsers/Fbx/Private/Ufbx/UfbxMaterial.cpp
+++ b/Engine/Plugins/Interchange/Runtime/Source/Parsers/Fbx/Private/Ufbx/UfbxMaterial.cpp
@@ -72,7 +72,7 @@ namespace UE::Interchange::Private
 			FString ShaderUniqueID = ShaderGraphNode.GetUniqueID();
 			const FString SamplerNodeUid = UInterchangeShaderNode::MakeNodeUid(NodeName, ShaderUniqueID);
 
-			// todo: ??
+			// #ufbx_todo: ??
 			// if (const UInterchangeShaderNode* SamplerNode = Cast<UInterchangeShaderNode>(NodeContainer.GetNode(SamplerNodeUid)))
 			// {
 			// 	return SamplerNode;
@@ -153,13 +153,13 @@ namespace UE::Interchange::Private
 
 					if (MaterialMap.value_components == 4)
 					{
-						// todo: converter utils
+						// #ufbx_todo: converter utils
 						FLinearColor PropertyValue(Convert::ConvertVec4(MaterialMap.value_vec4));
 						ShaderGraphNode.AddLinearColorAttribute(InputAttributeKey, PropertyValue);
 					}
 					else if (MaterialMap.value_components == 3)
 					{
-						// todo: converter utils
+						// #ufbx_todo: converter utils
 						FLinearColor PropertyValue(Convert::ConvertVec3(MaterialMap.value_vec3));
 						ShaderGraphNode.AddLinearColorAttribute(InputAttributeKey, PropertyValue);
 					}
@@ -172,7 +172,7 @@ namespace UE::Interchange::Private
 				{
 					if (const UInterchangeShaderNode* TextureSampleShader = CreateTextureSampler(InputName, MaterialMap.texture))
 					{
-						// todo: lerp etc
+						// #ufbx_todo: lerp etc
 						UInterchangeShaderNode* NodeToConnectTo = &ShaderGraphNode;
 						FString InputToConnectTo = InputName.ToString();
 
@@ -191,7 +191,7 @@ namespace UE::Interchange::Private
 	{
 		for (ufbx_material* Material : Parser.Scene->materials)
 		{
-			// todo: Uid in FFbxMaterial parser is made from sanitized name
+			// #ufbx_todo: Uid in FFbxMaterial parser is made from sanitized name
 			// but two different materials  might happen to have same name AFTER sanitation.
 			// SO material name shoul dbe made unique for this case. E.g. by adding number after name string.
 			// This means that we need to memoize mapping of material to this generated name to assign proper names to meshes
diff --git a/Engine/Plugins/Interchange/Runtime/Source/Parsers/Fbx/Private/Ufbx/UfbxMesh.cpp b/Engine/Plugins/Interchange/Runtime/Source/Parsers/Fbx/Private/Ufbx/UfbxMesh.cpp
--- a/Engine/Plugins/Interchange/Runtime/Source/Parsers/Fbx/Private/Ufbx/UfbxMesh.cpp
+++ b/Engine/Plugins/Interchange/Runtime/Source/Parsers/Fbx/Private/Ufbx/UfbxMesh.cpp
@@ -220,7 +220,7 @@ namespace UE::Interchange::Private
 			bool bMeshHasVertexBinormal = Mesh->vertex_bitangent.exists && Mesh->vertex_tangent.exists;
 			bool bMeshHasVertexTangent = bMeshHasVertexBinormal;
 			bool bMeshHasSmoothGroup = true;
-			bool bMeshHasVertexColor = false; // todo: 
+			bool bMeshHasVertexColor = false; // #ufbx_todo: 
 			int32 MeshUVCount = Mesh->uv_sets.count;
 			
 			// Why are these needed?
@@ -266,12 +266,12 @@ namespace UE::Interchange::Private
 				{
 					FString ChannelName = Convert::ToUnrealString(Channel->name);
 					ufbx_blend_shape* Shape = Channel->target_shape;
-					// todo: FBX SDK has multiple targets per channel via FbxBlendShapeChannel::GetTargetShapeCount(see comment on FbxBlendShapeChannel)
+					// #ufbx_todo: FBX SDK has multiple targets per channel via FbxBlendShapeChannel::GetTargetShapeCount(see comment on FbxBlendShapeChannel)
 					// check how uFBX handles this
 					
 					FString ShapeName = Parser.GetElementName(Shape->element);
 
-					// todo: 
+					// #ufbx_todo: 
 					// FString ShapeName;
 					// if (CurrentChannelMorphTargetCount > 1)
 					// {
@@ -292,7 +292,7 @@ namespace UE::Interchange::Private
 
 					if (ShapeName.IsEmpty())
 					{
-						// todo:
+						// #ufbx_todo:
 						ensure(false);
 						// ShapeName = TEXT("MorphTarget_") + FString::FromInt(ChannelIndex) + TEXT("_") + FString::FromInt(ChannelMorphTargetIndex);
 					}
@@ -326,7 +326,12 @@ namespace UE::Interchange::Private
 						}
 						MorphTargetNode->SetPayLoadKey(MorphTargetPayLoadKey, EInterchangeMeshPayLoadType::MORPHTARGET);
 					}
-					// todo: Channel->keyframes
+
+					if (ensure(!SkeletonRootPerBlendShape.Contains(Shape)))
+					{
+						SkeletonRootPerBlendShape.Add(Shape, Mesh);	
+					}
+					
 				}
 			}
 		}
@@ -352,7 +357,7 @@ namespace UE::Interchange::Private
 	{
 		TRACE_CPUPROFILER_EVENT_SCOPE(FUfbxMesh::FetchBlendShape);
 
-		// todo: normals
+		// #ufbx_todo: normals
 
 		FStaticMeshAttributes Attributes(MeshDescription);
 		Attributes.Register();
@@ -414,7 +419,7 @@ namespace UE::Interchange::Private
 
 		// FSkeletalMeshAttributes MeshAttributes(MeshDescription); for skeletal
 		FStaticMeshAttributes Attributes(MeshDescription);
-		// todo: why bKeepExistingAttribute>
+		// #ufbx_todo: why bKeepExistingAttribute>
 		constexpr bool bKeepExistingAttribute = true;
 		Attributes.Register(bKeepExistingAttribute);
 
@@ -425,7 +430,7 @@ namespace UE::Interchange::Private
 
 		Builder.EnablePolyGroups();
 
-		// todo: had to set at least one UV otherwise FStaticMeshOperations::ComputeMikktTangents
+		// #ufbx_todo: had to set at least one UV otherwise FStaticMeshOperations::ComputeMikktTangents
 		//  same was also done on FBX SDK parser 
 		Builder.SetNumUVLayers(FMath::Max(1, UvCount));
 
@@ -436,7 +441,7 @@ namespace UE::Interchange::Private
 		for (int32 VertexIndex = 0; VertexIndex < VertexCount; ++VertexIndex)
 		{
 			ufbx_vec3 Position = Mesh.vertex_position.values[VertexIndex];
-			// todo: FbxPosition = TotalMatrix.MultT(FbxPosition);
+			// #ufbx_todo: FbxPosition = TotalMatrix.MultT(FbxPosition);
 			// also find out what is that Matrix about
 			const FVector VertexPosition = Convert::ConvertVec3(Position);
 			const FVertexID VertexID = Builder.AppendVertex(VertexPosition);
@@ -527,7 +532,7 @@ namespace UE::Interchange::Private
 			const ufbx_skin_vertex& SkinVertex  = Skin.vertices[VertexIndex];
 			FVertexID VertexID = FaceCollector.VertexMapping[VertexIndex];
 
-			// todo: reuse array
+			// #ufbx_todo: reuse array
 			TArray<FBoneWeight> BoneWeights;
 
 			for (uint32 WeightIndex = 0; WeightIndex < SkinVertex.num_weights; ++WeightIndex)
@@ -555,7 +560,7 @@ namespace UE::Interchange::Private
 		// for (int32 ClusterIndex = 0; ClusterIndex < ClusterCount; ClusterIndex++)
 		// {
 		// 	ufbx_skin_cluster& Cluster = *Skin.clusters[ClusterIndex];
-		// 	// todo: ?
+		// 	// #ufbx_todo: ?
 		// 	// When Maya plug-in exports rigid binding, it will generate "CompensationCluster" for each ancestor links.
 		// 	// FBX writes these "CompensationCluster" out. The CompensationCluster also has weight 1 for vertices.
 		// 	// Unreal importer should skip these clusters.
diff --git a/Engine/Plugins/Interchange/Runtime/Source/Parsers/Fbx/Private/Ufbx/UfbxMesh.h b/Engine/Plugins/Interchange/Runtime/Source/Parsers/Fbx/Private/Ufbx/UfbxMesh.h
--- a/Engine/Plugins/Interchange/Runtime/Source/Parsers/Fbx/Private/Ufbx/UfbxMesh.h
+++ b/Engine/Plugins/Interchange/Runtime/Source/Parsers/Fbx/Private/Ufbx/UfbxMesh.h
@@ -31,10 +31,17 @@ namespace UE::Interchange::Private
 		static bool FetchSkinnedMesh(FUfbxParser& Parser, FMeshDescription& MeshDescription, ufbx_element* Element, const FTransform& MeshGlobalTransform, TArray<FString>& OutJointUniqueNames);
 		static bool FetchBlendShape(FUfbxParser& Parser, FMeshDescription& MeshDescription, const ufbx_mesh& Mesh, const ufbx_blend_shape& BlendShape, const FTransform& MeshGlobalTransform);
 
+		const TMap<ufbx_blend_shape*, ufbx_mesh*>& GetBlendShapes()
+		{
+			return SkeletonRootPerBlendShape;
+		}
+
 	private:
 		UInterchangeMeshNode* CreateMeshNode(UInterchangeBaseNodeContainer& NodeContainer, const FString& NodeDisplayLabel, const FString& NodeUniqueID);
 
 		FUfbxParser& Parser;
+
+		TMap<ufbx_blend_shape*, ufbx_mesh*> SkeletonRootPerBlendShape;
 	};
 
 }
\ No newline at end of file
diff --git a/Engine/Plugins/Interchange/Runtime/Source/Parsers/Fbx/Private/Ufbx/UfbxParser.cpp b/Engine/Plugins/Interchange/Runtime/Source/Parsers/Fbx/Private/Ufbx/UfbxParser.cpp
--- a/Engine/Plugins/Interchange/Runtime/Source/Parsers/Fbx/Private/Ufbx/UfbxParser.cpp
+++ b/Engine/Plugins/Interchange/Runtime/Source/Parsers/Fbx/Private/Ufbx/UfbxParser.cpp
@@ -16,16 +16,15 @@
 
 #include "Fbx/InterchangeFbxMessages.h"
 #include "InterchangeHelper.h"
+#include "InterchangeCommonAnimationPayload.h"
 
 #include "MeshDescription.h"
 #include "Misc/Paths.h"
-
-// Include implementation of ufbx
-#include "InterchangeCommonAnimationPayload.h"
-#include "ufbx.c"
 #include "Misc/FileHelper.h"
 #include "Serialization/MemoryWriter.h"
 
+#include "ufbx.c"
+
 
 #define LOCTEXT_NAMESPACE "InterchangeFbxParser"
 
@@ -56,7 +55,7 @@ FString FUfbxParser::GetElementName(const ufbx_element& Element, bool bIsJoint)
 	return ElementName;
 }
 
-// todo: see FFbxHelper::GetMeshName
+// #ufbx_todo: see FFbxHelper::GetMeshName
 FString FUfbxParser::GetMeshName(const ufbx_element& Mesh) const
 {
 	return UTF8_TO_TCHAR(Mesh.name.data);
@@ -64,7 +63,7 @@ FString FUfbxParser::GetMeshName(const ufbx_element& Mesh) const
 
 FString FUfbxParser::GetNodeUid(const ufbx_node& Node) const
 {
-	// todo: check that element_id is enough
+	// #ufbx_todo: check that element_id is enough
 	// reference FBX parser's uses UE::Interchange::Private::FFbxHelper::GetFbxNodeHierarchyName to create NodeUid
 
 	FString NodeUid = GetElementName(Node.element) + TEXT("_") + FString::FromInt(Node.element_id);
@@ -85,7 +84,7 @@ FString FUfbxParser::GetNodeUid(const ufbx_node& Node) const
 
 FName FUfbxParser::GetMaterialName(const ufbx_material& Material) const
 {
-	// todo: ensure uniqueness when GetElementName returns same value for different materials
+	// #ufbx_todo: ensure uniqueness when GetElementName returns same value for different materials
 	return *GetElementName(Material.element);
 }
 
@@ -121,19 +120,19 @@ bool FUfbxParser::LoadFbxFile(const FString& Filename, UInterchangeBaseNodeConta
 	// this option allows to override root transform with provided LoadOpts.root_transform
 	LoadOpts.use_root_transform = false;
 
-	// todo: LoadOpts.allow_nodes_out_of_root= true; // might just to this(to reparent all out-of-root nodes under root_node) -
+	// #ufbx_todo: LoadOpts.allow_nodes_out_of_root= true; // might just to this(to reparent all out-of-root nodes under root_node) -
 	// LoadOpts.allow_empty_faces = true;
 
-	// todo: progress reporting. In case someday we want to report detailed progress from Worker/Parser
+	// #ufbx_todo: progress reporting. In case someday we want to report detailed progress from Worker/Parser
 	// LoadOpts.progress_cb
 
-	// todo: handle geometry transform - UInterchangeSceneNode::SetCustomGeometricTransform
+	// #ufbx_todo: handle geometry transform - UInterchangeSceneNode::SetCustomGeometricTransform
 	// LoadOpts.geometry_transform_handling
 
-	// todo:
+	// #ufbx_todo:
 	// LoadOpts.inherit_mode_handling
 
-	// todo:
+	// #ufbx_todo:
 	// LoadOpts.pivot_handling
 
 	LoadOpts.space_conversion = UFBX_SPACE_CONVERSION_ADJUST_TRANSFORMS;
@@ -156,17 +155,17 @@ bool FUfbxParser::LoadFbxFile(const FString& Filename, UInterchangeBaseNodeConta
 		Message->Text = FText::Format(LOCTEXT("CannotOpenFBXFile", "Cannot open FBX file '{Filename}'."), FilenameText);
 		return false;
 	}
-	// todo: read scene converter to UE's
+	// #ufbx_todo: read scene converter to UE's
 
 	UInterchangeSourceNode* SourceNode = UInterchangeSourceNode::FindOrCreateUniqueInstance(&NodeContainer);
-	// todo:
+	// #ufbx_todo:
 	// SourceNode->SetCustomSourceFrameRateNumerator(FrameRate);
 	// SourceNode->SetCustomSourceFrameRateDenominator(Denominator);
 
-	// todo: check
+	// #ufbx_todo: check
 	// SourceNode->SetCustomUseLegacySkeletalMeshBakeTransform(true);
 
-	// todo: FileDetails
+	// #ufbx_todo: FileDetails
 
 	return true;
 }
@@ -178,15 +177,68 @@ void FUfbxParser::FillContainerWithFbxScene(UInterchangeBaseNodeContainer& NodeC
 	FUfbxScene SceneConverter(*this);
 	SceneConverter.InitHierarchy(NodeContainer);
 
-	// todo: CleanupFbxData();
+	// #ufbx_todo: CleanupFbxData();
 	FUfbxMaterial(*this).AddMaterials(NodeContainer);
 
+	FUfbxMesh Meshes(*this);
 	// Get meshes first, nodes reference them later in ProcessNodes
-	FUfbxMesh(*this).AddAllMeshes(NodeContainer);
+	Meshes.AddAllMeshes(NodeContainer);
 
 	SceneConverter.ProcessNodes(NodeContainer);
+	FUfbxAnimation::AddAnimation(*this, SceneConverter, NodeContainer, Meshes.GetBlendShapes());
+}
+
+bool FUfbxParser::FetchPayloadData(const FString& PayloadKey, const FString& PayloadFilepath)
+{
+	// #ufbx_todo: defence - see UE::Interchange::Private::FFbxParser::FetchPayloadData
+
+	const FPayloadContext* PayloadContextFound = PayloadContexts.Find(PayloadKey);
+	if (!PayloadContextFound)
+	{
+		return false;
+	}
+
+	const FPayloadContext& PayloadContext = *PayloadContextFound;
+	switch (PayloadContext.Kind)
+	{
+	case FPayloadContext::AnimationMorph:
+		{
+			FMorphAnimationPayloadContext& Animation = const_cast<FMorphAnimationPayloadContext&>(PayloadContexts.GetMorphAnimation(PayloadContext));
 
-	SceneConverter.AddAnimation(NodeContainer);
+			TArray<FInterchangeCurve> InterchangeCurves;
+			if (FUfbxAnimation::FetchMorphTargetAnimation(*this, Animation, InterchangeCurves))
+			{
+				TArray64<uint8> Buffer;
+				FMemoryWriter64 Ar(Buffer);
+				Ar << InterchangeCurves;
+				Ar << Animation.InbetweenCurveNames;
+				Ar << Animation.InbetweenFullWeights;
+				FFileHelper::SaveArrayToFile(Buffer, *PayloadFilepath);
+				return true;
+			}
+		}
+		break;
+	case FPayloadContext::AnimationRigid:
+		{
+			FRigidAnimationPayloadContext& Animation = const_cast<FRigidAnimationPayloadContext&>(PayloadContexts.GetRigidAnimation(PayloadContext));
+
+			TArray<FInterchangeCurve> InterchangeCurves;
+
+			if (FUfbxAnimation::FetchRigidAnimation(*this, Animation, InterchangeCurves))
+			{
+				TArray64<uint8> Buffer;
+				FMemoryWriter64 Ar(Buffer);
+				Ar << InterchangeCurves;
+				FFileHelper::SaveArrayToFile(Buffer, *PayloadFilepath);
+				return true;
+			}
+
+			return true;
+		}
+		break;
+	}
+	
+	return false;
 }
 
 bool FUfbxParser::FetchMeshPayloadData(const FString& PayloadKey, const FTransform& MeshGlobalTransform, const FString& PayloadFilepath)
@@ -224,7 +276,7 @@ bool FUfbxParser::FetchMeshPayloadData(const FString& PayloadKey, const FTransfo
 				}
 				else
 				{
-					ensure(false); // todo: implement
+					ensure(false); // #ufbx_todo: implement
 				}
 			}
 			break;
@@ -302,7 +354,7 @@ bool FUfbxParser::FetchAnimationBakeTransformPayload(const TArray<UE::Interchang
                                                      const FString& ResultFolder, FCriticalSection* ResultPayloadsCriticalSection, TAtomic<int64>& UniqueIdCounter,
                                                      TMap<FString, FString>& ResultPayloads)
 {
-	// todo: resolve concurrent access to ResultPayloads; or, maybe replace ResultPayloads with TQueue? Not clear why it's a map and why it's tested that keys already exist there...
+	// #ufbx_todo: resolve concurrent access to ResultPayloads; or, maybe replace ResultPayloads with TQueue? Not clear why it's a map and why it's tested that keys already exist there...
 
 	for (const FAnimationPayloadQuery& PayloadQuery : PayloadQueries)
 	{
@@ -315,10 +367,10 @@ bool FUfbxParser::FetchAnimationBakeTransformPayload(const TArray<UE::Interchang
 
 			FString QueryHashString = PayloadQuery.GetHashString();
 
-			// todo: check, FBX SDK parser somehow checks that ResultPayloads already has QueryHashString(i.e. same query already processed) Why is this?
+			// #ufbx_todo: check, FBX SDK parser somehow checks that ResultPayloads already has QueryHashString(i.e. same query already processed) Why is this?
 			FString PayloadFilepathCopy;
 			{
-				// todo: guard concurrent access to ResultPayloads, see UE::Interchange::Private::FFbxAnimation::FetchAnimationBakeTransformPayload
+				// #ufbx_todo: guard concurrent access to ResultPayloads, see UE::Interchange::Private::FFbxAnimation::FetchAnimationBakeTransformPayload
 
 				FString& PayloadFilepath = ResultPayloads.FindOrAdd(QueryHashString);
 				//To avoid file path with too many character, we hash the payloadKey so we have a deterministic length for the file path.
diff --git a/Engine/Plugins/Interchange/Runtime/Source/Parsers/Fbx/Private/Ufbx/UfbxParser.h b/Engine/Plugins/Interchange/Runtime/Source/Parsers/Fbx/Private/Ufbx/UfbxParser.h
--- a/Engine/Plugins/Interchange/Runtime/Source/Parsers/Fbx/Private/Ufbx/UfbxParser.h
+++ b/Engine/Plugins/Interchange/Runtime/Source/Parsers/Fbx/Private/Ufbx/UfbxParser.h
@@ -10,7 +10,7 @@
 // also might do #define UFBX_MINIMAL and enable what's needed
 #include <ufbx.h>
 
-// todo: move IFbxParser to separate file and include IFbxParser.h instead
+// #ufbx_todo: move IFbxParser to separate file and include IFbxParser.h instead
 #include "FbxAPI.h"
 
 #include "MeshDescription.h"
@@ -25,6 +25,8 @@ namespace UE::Interchange::Private
 			SkinnedMesh,
 			Morph,
 			AnimationSkinned,
+			AnimationMorph,
+			AnimationRigid
 		};
 
 		FPayloadContext() = delete;
@@ -52,6 +54,25 @@ namespace UE::Interchange::Private
 		ufbx_shared_ptr<ufbx_baked_anim> BakedAnim;
 	};
 
+	struct FMorphAnimationPayloadContext
+	{
+		const ufbx_anim_curve* Curve;
+
+		// #ufbx_todo:
+		//In between blend shape animation support
+		//The count of the curve names should match the in between full weights
+		TArray<FString> InbetweenCurveNames;
+		TArray<float> InbetweenFullWeights;
+	};
+
+	struct FRigidAnimationPayloadContext
+	{
+		// Transform animation curves array in the order Interchange accepts it - Tx, Ty, Tz, Rx, Ry, Rz, Sx, Sy, Sz
+		const ufbx_anim_prop* TranslationProp = nullptr;
+		const ufbx_anim_prop* RotationProp = nullptr;
+		const ufbx_anim_prop* ScalingProp = nullptr;
+	};
+
 	class FUfbxParser: public IFbxParser, public FNoncopyable
 	{
 	public:
@@ -65,7 +86,7 @@ namespace UE::Interchange::Private
 		virtual void SetResultContainer(UInterchangeResultsContainer* Result) override;
 		virtual void SetConvertSettings(const bool InbConvertScene, const bool InbForceFrontXAxis, const bool InbConvertSceneUnit, const bool InbKeepFbxNamespace)  override
 		{
-			// todo: 
+			// #ufbx_todo: 
 		}
 
 		virtual bool LoadFbxFile(const FString& Filename, UInterchangeBaseNodeContainer& NodeContainer) override;
@@ -74,7 +95,7 @@ namespace UE::Interchange::Private
 		virtual void FillContainerWithFbxScene(UInterchangeBaseNodeContainer& NodeContainer) override;
 
 		/* Extract the fbx data from the sdk into our node container */
-		virtual bool FetchPayloadData(const FString& PayloadKey, const FString& PayloadFilepath) override {	return false; }
+		virtual bool FetchPayloadData(const FString& PayloadKey, const FString& PayloadFilepath) override;
 
 		/* Extract the fbx mesh data from the sdk into our node container */
 		virtual bool FetchMeshPayloadData(const FString& PayloadKey, const FTransform& MeshGlobalTransform, const FString& PayloadFilepath) override;
@@ -93,7 +114,7 @@ namespace UE::Interchange::Private
 		FString SourceFilename;
 		TWeakObjectPtr<UInterchangeResultsContainer> ResultsContainer;
 
-		// todo: Messages functionality is a copy of FFbxParser's, maybe share
+		// #ufbx_todo: Messages functionality is a copy of FFbxParser's, maybe share
 		/**
 		 * This function is used to add the given message object directly into the results for this operation.
 		 */
@@ -141,6 +162,24 @@ namespace UE::Interchange::Private
 				PayloadContexts.Add(MorphTargetPayLoadKey, FPayloadContext(FPayloadContext::Morph, MorphId));
 			}
 
+			void Add(const FString& PayloadKey, const FSkeletalAnimationPayloadContext& Payload)
+			{
+				int32 PayloadId = SkeletalAnimations.Add(Payload);
+				PayloadContexts.Add(PayloadKey, FPayloadContext(FPayloadContext::AnimationSkinned, PayloadId));
+			}
+
+			void Add(const FString& PayloadKey, const FMorphAnimationPayloadContext& Payload)
+			{
+				int32 PayloadId = MorphAnimations.Add(Payload);
+				PayloadContexts.Add(PayloadKey, FPayloadContext(FPayloadContext::AnimationMorph, PayloadId));
+			}
+
+			void Add(const FString& PayloadKey, const FRigidAnimationPayloadContext& Payload)
+			{
+				int32 PayloadId = RigidAnimations.Add(Payload);
+				PayloadContexts.Add(PayloadKey, FPayloadContext(FPayloadContext::AnimationRigid, PayloadId));
+			}
+
 			bool Contains(const FString& PayloadKey)
 			{
 				return PayloadContexts.Contains(PayloadKey);
@@ -161,15 +200,19 @@ namespace UE::Interchange::Private
 				return &MorphPayloads[PayloadPtr->Index];
 			}
 
-			void Add(const FString& PayloadKey, const FSkeletalAnimationPayloadContext& Payload)
+			const FSkeletalAnimationPayloadContext& GetAnimation(const FPayloadContext& Payload)
 			{
-				int32 PayloadId = SkeletalAnimations.Add(Payload);
-				PayloadContexts.Add(PayloadKey, FPayloadContext(FPayloadContext::AnimationSkinned, PayloadId));
+				return SkeletalAnimations[Payload.Index];
 			}
 
-			const FSkeletalAnimationPayloadContext& GetAnimation(const FPayloadContext& Payload)
+			const FMorphAnimationPayloadContext& GetMorphAnimation(const FPayloadContext& Payload)
 			{
-				return SkeletalAnimations[Payload.Index];
+				return MorphAnimations[Payload.Index];
+			}
+
+			const FRigidAnimationPayloadContext& GetRigidAnimation(const FPayloadContext& Payload)
+			{
+				return RigidAnimations[Payload.Index];
 			}
 
 			TMap<FString, FPayloadContext> PayloadContexts;
@@ -177,6 +220,11 @@ namespace UE::Interchange::Private
 			// Payloads that a simple ufbx elements
 			TArray<FMorph> MorphPayloads;
 			TArray<FSkeletalAnimationPayloadContext> SkeletalAnimations;
+
+			TArray<FMorphAnimationPayloadContext> MorphAnimations;
+			TArray<FRigidAnimationPayloadContext> RigidAnimations;
+			
+			
 			
 
 		};
diff --git a/Engine/Plugins/Interchange/Runtime/Source/Parsers/Fbx/Private/Ufbx/UfbxScene.cpp b/Engine/Plugins/Interchange/Runtime/Source/Parsers/Fbx/Private/Ufbx/UfbxScene.cpp
--- a/Engine/Plugins/Interchange/Runtime/Source/Parsers/Fbx/Private/Ufbx/UfbxScene.cpp
+++ b/Engine/Plugins/Interchange/Runtime/Source/Parsers/Fbx/Private/Ufbx/UfbxScene.cpp
@@ -90,8 +90,8 @@ namespace UE::Interchange::Private
 
 		UInterchangeSceneNode* SceneNode = *SceneNodeFound;
 
-		// todo: Node->geometry_transform
-		// todo: possible to set SetCustomGeometricTransform, SetCustomPivotNodeTransform?
+		// #ufbx_todo: Node->geometry_transform
+		// #ufbx_todo: possible to set SetCustomGeometricTransform, SetCustomPivotNodeTransform?
 
 		if (InNode.mesh)
 		{
@@ -139,6 +139,11 @@ namespace UE::Interchange::Private
 				UnrealNode->SetCustomBindPoseLocalTransform(&NodeContainer, LocalBindTransform, bResetCache);
 			}
 		}
+
+		if (InNode.attrib_type == UFBX_ELEMENT_LOD_GROUP)
+		{
+			SceneNode->AddSpecializedType(FSceneNodeStaticData::GetLodGroupSpecializeTypeString());
+		}
 	}
 
 	void FUfbxScene::ProcessNodes(UInterchangeBaseNodeContainer& NodeContainer)
@@ -178,94 +183,11 @@ namespace UE::Interchange::Private
 		}
 	}
 
-	void FUfbxScene::AddAnimation(UInterchangeBaseNodeContainer& NodeContainer)
+	FString FUfbxScene::GetSkeletonRoot(ufbx_node* Node)
 	{
-		double FrameRate = Parser.Scene->settings.frames_per_second;
-
-		int32 AnimStackCount = Parser.Scene->anim_stacks.count;
-		for (int32 AnimStackIndex = 0; AnimStackIndex < AnimStackCount; ++AnimStackIndex)
-		{
-			ufbx_anim_stack* Stack = Parser.Scene->anim_stacks[AnimStackIndex];
-
-			ufbx_bake_opts BakeOptions = {};
-			BakeOptions.maximum_sample_rate = FrameRate;
-			BakeOptions.minimum_sample_rate = FrameRate;
-			BakeOptions.resample_rate = FrameRate;
-
-			ufbx_error ErrorResult;
-
-			ufbx_shared_ptr<ufbx_baked_anim> BakedAnim(ufbx_bake_anim(Parser.Scene, Stack->anim, &BakeOptions, &ErrorResult));
-			if (!BakedAnim)
-			{
-				UInterchangeResultError_Generic* Message = Parser.AddMessage<UInterchangeResultError_Generic>();
-
-				TArray<char> ErrorBuf;
-				ErrorBuf.SetNum(ErrorResult.info_length+64);
-				ufbx_format_error(ErrorBuf.GetData(), ErrorBuf.Num(), &ErrorResult);
-
-				Message->Text = FText::Format(LOCTEXT("AnimBakeError", "UFBX couldn't bake animation: '{0}'."), FText::FromString(UTF8_TO_TCHAR(ErrorBuf.GetData())));
-				continue;
-			}
-
-			TConstArrayView<ufbx_baked_node> BakedNodes(BakedAnim->nodes.data, BakedAnim->nodes.count);
-			if (BakedNodes.IsEmpty())
-			{
-				continue;
-			}
-
-			TConstArrayView<ufbx_node*> FbxNodes(Parser.Scene->nodes.data, Parser.Scene->nodes.count);
-
-			uint32_t NodeTypeId = BakedNodes[0].typed_id;
-			if (!ensure(FbxNodes.IsValidIndex(NodeTypeId)))
-			{
-				continue;
-			}
-
-			FString DisplayString = Convert::ToUnrealString(Stack->name);
-			FString SkeletonRootNodeUid = Parser.GetNodeUid(*SceneUtils::FindSkeletonRoot(FbxNodes[NodeTypeId]));
-			FString TrackNodeUid = SceneUtils::CreateTrackNodeUid(SkeletonRootNodeUid, AnimStackIndex);
-
-			UInterchangeSkeletalAnimationTrackNode* SkeletalAnimationTrackNode = NewObject< UInterchangeSkeletalAnimationTrackNode >(&NodeContainer);
-			SkeletalAnimationTrackNode->InitializeNode(TrackNodeUid, DisplayString, EInterchangeNodeContainerType::TranslatedAsset);
-			SkeletalAnimationTrackNode->AddBooleanAttribute(TEXT("RenameLikeLegacyFbx"), true);
-			SkeletalAnimationTrackNode->SetCustomAnimationSampleRate(Parser.Scene->settings.frames_per_second);
-
-			SkeletalAnimationTrackNode->SetCustomSkeletonNodeUid(SkeletonRootNodeUid);
-			
-			SkeletalAnimationTrackNode->SetCustomAnimationStartTime(Stack->anim->time_begin);
-			SkeletalAnimationTrackNode->SetCustomAnimationStopTime(Stack->anim->time_end);
-
-			SkeletalAnimationTrackNode->SetCustomSourceTimelineAnimationStartTime(Stack->time_begin);
-			SkeletalAnimationTrackNode->SetCustomSourceTimelineAnimationStopTime(Stack->time_end);
-
-			for (int32 BakedNodeIndex = 0; BakedNodeIndex < BakedAnim->nodes.count; ++BakedNodeIndex)
-			{
-				const ufbx_baked_node& BakedNode = BakedAnim->nodes[BakedNodeIndex];
-
-				if (!ensureMsgf(FbxNodes.IsValidIndex(BakedNode.typed_id), TEXT("Unexpected type_id encountered in ufbx baked anim nodes")))
-				{
-					continue;
-				}
-
-				ufbx_node* Node = FbxNodes[BakedNode.typed_id];
-				FString NodeUid = Parser.GetNodeUid(*Node);
-
-				// todo: Need check animated parents - whole subtree has to be animated too?
-				// see UE::Interchange::Private::FFbxAnimation::AddSkeletalTransformAnimation
-				// or maybe there's an option in ufbx baking that already does that?
-
-				FString PayLoadKey = NodeUid + TEXT("_") + FString::FromInt(AnimStackIndex) + TEXT("_SkeletalAnimationPayloadKey");
-
-				Parser.PayloadContexts.Add(PayLoadKey, FSkeletalAnimationPayloadContext{NodeUid, AnimStackIndex, BakedNodeIndex, BakedAnim});
-
-				SkeletalAnimationTrackNode->SetAnimationPayloadKeyForSceneNodeUid(NodeUid, PayLoadKey, EInterchangeAnimationPayLoadType::BAKED);
-
-				NodeContainer.AddNode(SkeletalAnimationTrackNode);
-			}
-		}
-
+		FString* SkeletonRootFound = SkeletonRootPerBone.Find(Node);
+		return SkeletonRootFound ? *SkeletonRootFound : Parser.GetNodeUid(*SceneUtils::FindSkeletonRoot(Node));
 	}
-
 }
 
 #undef LOCTEXT_NAMESPACE
diff --git a/Engine/Plugins/Interchange/Runtime/Source/Parsers/Fbx/Private/Ufbx/UfbxScene.h b/Engine/Plugins/Interchange/Runtime/Source/Parsers/Fbx/Private/Ufbx/UfbxScene.h
--- a/Engine/Plugins/Interchange/Runtime/Source/Parsers/Fbx/Private/Ufbx/UfbxScene.h
+++ b/Engine/Plugins/Interchange/Runtime/Source/Parsers/Fbx/Private/Ufbx/UfbxScene.h
@@ -8,6 +8,8 @@ class UInterchangeBaseNodeContainer;
 class UInterchangeSceneNode;
 
 struct ufbx_node;typedef ufbx_node ufbx_node;
+struct ufbx_mesh;typedef ufbx_mesh ufbx_mesh;
+struct ufbx_blend_shape;typedef ufbx_blend_shape ufbx_blend_shape;
 
 namespace UE::Interchange::Private
 {
@@ -22,7 +24,8 @@ namespace UE::Interchange::Private
 
 		void InitHierarchy(UInterchangeBaseNodeContainer& NodeContainer);
 		void ProcessNodes(UInterchangeBaseNodeContainer& NodeContainer);
-		void AddAnimation(UInterchangeBaseNodeContainer& NodeContainer);
+
+		FString GetSkeletonRoot(ufbx_node* Node);;
 
 	private:
 		void ProcessNode(UInterchangeBaseNodeContainer& NodeContainer, const ufbx_node& Node);
@@ -30,6 +33,7 @@ namespace UE::Interchange::Private
 		FUfbxParser& Parser;
 
 		TSet<uint32> CommonJointRootNodes;
+		TMap<ufbx_node*, FString> SkeletonRootPerBone;
 	};
 
 }
