Removed FNavigationOctree::AddNode Warning during Level Import.

As discussed with Jean-Luc, this CL is a valid but a temporary solution for the bigger underlying issue of having actors being assigned "not fully ready asset" - (UE-309334)

The issue was that, during a frame, freshly created SMActors where picked up by the navigation, and because this was happening before the SM was built the bounds were not computed yet causing the warning.

To solve the issue, we simply need to call `StaticMesh->CalculateExtendedBounds` which should fallback to an already calculated CachedMeshDescriptionBounds, thus causing minimal performance overhead.

Note:
- Similar issue could in theory happen for SKM Actors too. But because SKM Actors have `CanAffectNavigation` set to False by default the warning does not happen.
[RN] Removed FNavigationOctree::AddNode Warning during Level Import.

diff --git a/Engine/Plugins/Interchange/Runtime/Source/Import/Private/Scene/InterchangeStaticMeshActorFactory.cpp b/Engine/Plugins/Interchange/Runtime/Source/Import/Private/Scene/InterchangeStaticMeshActorFactory.cpp
--- a/Engine/Plugins/Interchange/Runtime/Source/Import/Private/Scene/InterchangeStaticMeshActorFactory.cpp
+++ b/Engine/Plugins/Interchange/Runtime/Source/Import/Private/Scene/InterchangeStaticMeshActorFactory.cpp
@@ -9,6 +9,7 @@
 #include "Components/StaticMeshComponent.h"
 #include "Engine/StaticMesh.h"
 #include "Engine/StaticMeshActor.h"
+#include "InterchangeImportLog.h"
 
 #include UE_INLINE_GENERATED_CPP_BY_NAME(InterchangeStaticMeshActorFactory)
 
@@ -41,9 +42,10 @@ UObject* UInterchangeStaticMeshActorFactory::ProcessActor(AActor& SpawnedActor,
 			}
 		}
 
-
 		if (bHasGeometricTransform)
 		{
+			UE_LOG(LogInterchangeImport, Warning, TEXT("BOUNDS bHasGeometricTransform"));
+
 			UStaticMeshComponent* GeometricTransformMeshComponent = nullptr;
 			if (StaticMeshComponent->GetNumChildrenComponents() > 0)
 			{
@@ -71,19 +73,8 @@ UObject* UInterchangeStaticMeshActorFactory::ProcessActor(AActor& SpawnedActor,
 
 				if (const UInterchangeFactoryBaseNode* MeshNode = ActorHelper::FindAssetInstanceFactoryNode(&NodeContainer, &FactoryNode))
 				{
-					FSoftObjectPath ReferenceObject;
-					MeshNode->GetCustomReferenceObject(ReferenceObject);
-					if (UStaticMesh* StaticMesh = Cast<UStaticMesh>(ReferenceObject.TryLoad()))
-					{
-						if (StaticMesh != GeometricTransformMeshComponent->GetStaticMesh())
-						{
-							GeometricTransformMeshComponent->SetStaticMesh(StaticMesh);
-						}
-						if (MeshActorFactoryNode)
-						{
-							UE::Interchange::ActorHelper::ApplySlotMaterialDependencies(NodeContainer, *MeshActorFactoryNode, *GeometricTransformMeshComponent);
-						}
-					}
+					UE_LOG(LogInterchangeImport, Warning, TEXT("BOUNDS Process Actor"));
+					ApplyStaticMeshToComponent(MeshNode, GeometricTransformMeshComponent, NodeContainer, MeshActorFactoryNode);
 				}
 
 				GeometricTransformMeshComponent->SetupAttachment(StaticMeshComponent);
@@ -103,19 +94,8 @@ UObject* UInterchangeStaticMeshActorFactory::ProcessActor(AActor& SpawnedActor,
 		{
 			StaticMeshComponent->UnregisterComponent();
 
-			FSoftObjectPath ReferenceObject;
-			MeshNode->GetCustomReferenceObject(ReferenceObject);
-			if (UStaticMesh* StaticMesh = Cast<UStaticMesh>(ReferenceObject.TryLoad()))
-			{
-				if (StaticMesh != StaticMeshComponent->GetStaticMesh())
-				{
-					StaticMeshComponent->SetStaticMesh(StaticMesh);
-				}
-				if (MeshActorFactoryNode)
-				{
-					UE::Interchange::ActorHelper::ApplySlotMaterialDependencies(NodeContainer, *MeshActorFactoryNode, *StaticMeshComponent);
-				}
-			}
+			UE_LOG(LogInterchangeImport, Warning, TEXT("BOUNDS Process Actor"));
+			ApplyStaticMeshToComponent(MeshNode, StaticMeshComponent, NodeContainer, MeshActorFactoryNode);
 
 			StaticMeshComponent->RegisterComponent();
 		}
@@ -141,17 +121,8 @@ void UInterchangeStaticMeshActorFactory::SetupObject_GameThread(const FSetupObje
 			const UInterchangeStaticMeshFactoryNode* StaticMeshFactoryNode = TargetNodeUids.IsEmpty() ? nullptr : Cast<UInterchangeStaticMeshFactoryNode>(Arguments.NodeContainer->GetFactoryNode(TargetNodeUids[0]));
 			if (StaticMeshFactoryNode)
 			{
-				FSoftObjectPath ReferenceObject;
-				StaticMeshFactoryNode->GetCustomReferenceObject(ReferenceObject);
-				if (UStaticMesh* StaticMesh = Cast<UStaticMesh>(ReferenceObject.TryLoad()))
-				{
-					StaticMeshComponent->SetStaticMesh(StaticMesh);
-
-					if (const UInterchangeMeshActorFactoryNode* MeshActorFactoryNode = Cast<UInterchangeMeshActorFactoryNode>(Arguments.FactoryNode))
-					{
-						UE::Interchange::ActorHelper::ApplySlotMaterialDependencies(*Arguments.NodeContainer, *MeshActorFactoryNode, *StaticMeshComponent);
-					}
-				}
+				UE_LOG(LogInterchangeImport, Warning, TEXT("BOUNDS SetupObject_GameThread"));
+				ApplyStaticMeshToComponent(StaticMeshFactoryNode, StaticMeshComponent, *Arguments.NodeContainer, Cast<UInterchangeMeshActorFactoryNode>(Arguments.FactoryNode));
 			}
 		}
 	}
@@ -195,3 +166,28 @@ void UInterchangeStaticMeshActorFactory::ExecuteResetObjectProperties(const UInt
 	}
 }
 
+void UInterchangeStaticMeshActorFactory::ApplyStaticMeshToComponent(const UInterchangeFactoryBaseNode* MeshNode, UStaticMeshComponent* StaticMeshComponent, const UInterchangeBaseNodeContainer& NodeContainer, const UInterchangeMeshActorFactoryNode* MeshFactoryNode)
+{
+	FSoftObjectPath ReferenceObject;
+	MeshNode->GetCustomReferenceObject(ReferenceObject);
+	
+	if (UStaticMesh* StaticMesh = Cast<UStaticMesh>(ReferenceObject.TryLoad()))
+	{
+		if (StaticMesh != StaticMeshComponent->GetStaticMesh())
+		{
+			StaticMeshComponent->SetStaticMesh(StaticMesh);
+		}
+
+		UE_LOG(LogInterchangeImport, Warning, TEXT("BOUNDS static mesh '%s' : '%0.3f, %0.3f, %0.3f'."), *StaticMesh->GetName(), StaticMesh->GetBounds().BoxExtent.X, StaticMesh->GetBounds().BoxExtent.Y, StaticMesh->GetBounds().BoxExtent.Z);
+
+		// UE-306557: At this point, bounds might not be calculated yet but the navigation system will try to add it to the Octree, causing an Empty Bounds warning.
+		// To avoid this warning, we manually calculate the bounds here.
+		StaticMesh->CalculateExtendedBounds();
+
+		// Apply the materials from the mesh node to the static mesh component
+		if (MeshFactoryNode)
+		{
+			UE::Interchange::ActorHelper::ApplySlotMaterialDependencies(NodeContainer, *MeshFactoryNode, *StaticMeshComponent);
+		}
+	}
+}
diff --git a/Engine/Plugins/Interchange/Runtime/Source/Import/Public/Scene/InterchangeStaticMeshActorFactory.h b/Engine/Plugins/Interchange/Runtime/Source/Import/Public/Scene/InterchangeStaticMeshActorFactory.h
--- a/Engine/Plugins/Interchange/Runtime/Source/Import/Public/Scene/InterchangeStaticMeshActorFactory.h
+++ b/Engine/Plugins/Interchange/Runtime/Source/Import/Public/Scene/InterchangeStaticMeshActorFactory.h
@@ -13,7 +13,9 @@
 
 class AActor;
 class AStaticMeshActor;
+class UStaticMeshComponent;
 class UInterchangeActorFactoryNode;
+class UInterchangeMeshActorFactoryNode;
 
 UCLASS(MinimalAPI, BlueprintType)
 class UInterchangeStaticMeshActorFactory : public UInterchangeActorFactory
@@ -41,6 +43,11 @@ protected:
 	UE_API virtual void ExecuteResetObjectProperties(const UInterchangeBaseNodeContainer* BaseNodeContainer, UInterchangeFactoryBaseNode* FactoryNode, UObject* ImportedObject) override;
 	// Interchange actor factory interface end
 	//////////////////////////////////////////////////////////////////////////
+
+private:
+
+	// Helper function to assign a static mesh to its component while checking for other settings like navigation bounds and material dependencies.
+	void ApplyStaticMeshToComponent(const UInterchangeFactoryBaseNode* MeshNode, UStaticMeshComponent* StaticMeshComponent, const UInterchangeBaseNodeContainer& NodeContainer, const UInterchangeMeshActorFactoryNode* MeshFactoryNode);
 };
 
 
