[Interchange] Adding SceneComponent and Instanced StaticMeshComponent support.

#rnx

diff --git a/Engine/Plugins/Interchange/Runtime/Source/FactoryNodes/Private/InterchangeActorFactoryNode.cpp b/Engine/Plugins/Interchange/Runtime/Source/FactoryNodes/Private/InterchangeActorFactoryNode.cpp
--- a/Engine/Plugins/Interchange/Runtime/Source/FactoryNodes/Private/InterchangeActorFactoryNode.cpp
+++ b/Engine/Plugins/Interchange/Runtime/Source/FactoryNodes/Private/InterchangeActorFactoryNode.cpp
@@ -6,6 +6,8 @@
 
 #if WITH_ENGINE
 #include "Components/SceneComponent.h"
+#include "Components/InstancedStaticMeshComponent.h"
+#include "Components/HierarchicalInstancedStaticMeshComponent.h"
 #include "GameFramework/Actor.h"
 #endif
 
@@ -240,6 +242,7 @@ UInterchangeActorFactoryNode::UInterchangeActorFactoryNode()
 {
 	LayerNames.Initialize(Attributes.ToSharedRef(), TEXT("__LayerNames__"));
 	Tags.Initialize(Attributes.ToSharedRef(), TEXT("__Tags__"));
+	ComponentUids.Initialize(Attributes.ToSharedRef(), TEXT("__ComponentUids__"));
 }
 
 void UInterchangeActorFactoryNode::GetLayerNames(TArray<FString>& OutLayerNames) const
@@ -301,3 +304,13 @@ void UInterchangeActorFactoryNode::OnRestoreAllCustomAttributeDelegates()
 	REFILL_CUSTOM_ATTRIBUTE_APPLY_DELEGATE(ComponentVisibility, bool)
 	REFILL_CUSTOM_ATTRIBUTE_APPLY_DELEGATE(ActorVisibility, bool)
 }
+
+bool UInterchangeActorFactoryNode::AddComponentUid(const FString& ComponentUid)
+{
+	return ComponentUids.AddItem(ComponentUid);
+}
+
+void UInterchangeActorFactoryNode::GetComponentUids(TArray<FString>& OutComponentUids) const
+{
+	ComponentUids.GetItems(OutComponentUids);
+}
diff --git a/Engine/Plugins/Interchange/Runtime/Source/FactoryNodes/Private/InterchangeSceneComponentFactoryNodes.cpp b/Engine/Plugins/Interchange/Runtime/Source/FactoryNodes/Private/InterchangeSceneComponentFactoryNodes.cpp
new file mode 100644
--- /dev/null
+++ b/Engine/Plugins/Interchange/Runtime/Source/FactoryNodes/Private/InterchangeSceneComponentFactoryNodes.cpp
@@ -0,0 +1,133 @@
+// Copyright Epic Games, Inc. All Rights Reserved.
+
+#include "InterchangeSceneComponentFactoryNodes.h"
+
+#include UE_INLINE_GENERATED_CPP_BY_NAME(InterchangeSceneComponentFactoryNodes)
+
+#if WITH_ENGINE
+#include "Components/SceneComponent.h"
+#include "Components/InstancedStaticMeshComponent.h"
+#include "Components/HierarchicalInstancedStaticMeshComponent.h"
+#include "GameFramework/Actor.h"
+#endif
+
+
+UInterchangeSceneComponentFactoryNode::UInterchangeSceneComponentFactoryNode()
+{
+	ComponentUids.Initialize(Attributes, TEXT("__ComponentUids__Key"));
+}
+
+FString UInterchangeSceneComponentFactoryNode::GetTypeName() const
+{
+	return TEXT("UInterchangeSceneComponentFactoryNode");
+}
+
+UClass* UInterchangeSceneComponentFactoryNode::GetObjectClass() const
+{
+#if WITH_ENGINE
+	return USceneComponent::StaticClass();
+#else
+	return nullptr;
+#endif
+}
+
+
+bool UInterchangeSceneComponentFactoryNode::AddComponentUid(const FString& ComponentUid)
+{
+	return ComponentUids.AddItem(ComponentUid);
+}
+
+void UInterchangeSceneComponentFactoryNode::GetComponentUids(TArray<FString>& OutComponentUids) const
+{
+	ComponentUids.GetItems(OutComponentUids);
+}
+
+
+bool UInterchangeSceneComponentFactoryNode::GetCustomLocalTransform(FTransform& AttributeValue) const
+{
+	IMPLEMENT_NODE_ATTRIBUTE_GETTER(LocalTransform, FTransform);
+}
+
+bool UInterchangeSceneComponentFactoryNode::SetCustomLocalTransform(const FTransform& AttributeValue)
+{
+	IMPLEMENT_NODE_ATTRIBUTE_SETTER_NODELEGATE(LocalTransform, FTransform);
+}
+
+
+bool UInterchangeSceneComponentFactoryNode::GetCustomComponentVisibility(bool& AttributeValue) const
+{
+	IMPLEMENT_NODE_ATTRIBUTE_GETTER(ComponentVisibility, bool);
+}
+
+bool UInterchangeSceneComponentFactoryNode::SetCustomComponentVisibility(const bool& AttributeValue)
+{
+	IMPLEMENT_NODE_ATTRIBUTE_SETTER_NODELEGATE(ComponentVisibility, bool);
+}
+
+
+
+UInterchangeInstancedStaticMeshComponentFactoryNode ::UInterchangeInstancedStaticMeshComponentFactoryNode ()
+{
+	InstanceTransforms.Initialize(Attributes, TEXT("__InstanceTransforms__"));
+}
+
+FString UInterchangeInstancedStaticMeshComponentFactoryNode ::GetTypeName() const
+{
+	return TEXT("UInterchangeInstancedStaticMeshComponentFactoryNode ");
+}
+
+UClass* UInterchangeInstancedStaticMeshComponentFactoryNode ::GetObjectClass() const
+{
+#if WITH_ENGINE
+	bool bHierarchicalISM = false;
+	GetCustomHierarchicalISM(bHierarchicalISM);
+	if (bHierarchicalISM)
+	{
+		return UHierarchicalInstancedStaticMeshComponent::StaticClass();
+	}
+	else
+	{
+		return UInstancedStaticMeshComponent::StaticClass();
+	}
+#else
+	return nullptr;
+#endif
+}
+
+void UInterchangeInstancedStaticMeshComponentFactoryNode ::AddInstanceTransform(const FTransform& InstanceTransform)
+{
+	InstanceTransforms.AddItem(InstanceTransform);
+}
+
+void UInterchangeInstancedStaticMeshComponentFactoryNode ::AddInstanceTransforms(const TArray<FTransform>& InInstanceTransforms)
+{
+	for (const FTransform& InstanceTransform : InInstanceTransforms)
+	{
+		InstanceTransforms.AddItem(InstanceTransform);
+	}
+}
+
+void UInterchangeInstancedStaticMeshComponentFactoryNode ::GetInstanceTransforms(TArray<FTransform>& OutInstanceTransforms) const
+{
+	InstanceTransforms.GetItems(OutInstanceTransforms);
+}
+
+bool UInterchangeInstancedStaticMeshComponentFactoryNode ::GetCustomInstancedAssetUid(FString& AttributeValue) const
+{
+	IMPLEMENT_NODE_ATTRIBUTE_GETTER(AssetInstanceUid, FString);
+}
+
+bool UInterchangeInstancedStaticMeshComponentFactoryNode ::SetCustomInstancedAssetUid(const FString& AttributeValue)
+{
+	IMPLEMENT_NODE_ATTRIBUTE_SETTER_NODELEGATE(AssetInstanceUid, FString);
+}
+
+bool UInterchangeInstancedStaticMeshComponentFactoryNode ::GetCustomHierarchicalISM(bool& AttributeValue) const
+{
+	IMPLEMENT_NODE_ATTRIBUTE_GETTER(HierarchicalISM, bool);
+}
+
+bool UInterchangeInstancedStaticMeshComponentFactoryNode ::SetCustomHierarchicalISM(const bool& AttributeValue)
+{
+	IMPLEMENT_NODE_ATTRIBUTE_SETTER_NODELEGATE(HierarchicalISM, bool);
+}
\ No newline at end of file
diff --git a/Engine/Plugins/Interchange/Runtime/Source/FactoryNodes/Public/InterchangeActorFactoryNode.h b/Engine/Plugins/Interchange/Runtime/Source/FactoryNodes/Public/InterchangeActorFactoryNode.h
--- a/Engine/Plugins/Interchange/Runtime/Source/FactoryNodes/Public/InterchangeActorFactoryNode.h
+++ b/Engine/Plugins/Interchange/Runtime/Source/FactoryNodes/Public/InterchangeActorFactoryNode.h
@@ -84,6 +84,14 @@ public:
 
 
 	UE_API virtual void OnRestoreAllCustomAttributeDelegates() override;
+
+
+	UFUNCTION(BlueprintCallable, Category = "Interchange | Scene | Component")
+	UE_API bool AddComponentUid(const FString& ComponentUid);
+
+	UFUNCTION(BlueprintCallable, Category = "Interchange | Scene | Component")
+	UE_API void GetComponentUids(TArray<FString>& OutComponentUids) const;
+
 private:
 
 	bool ApplyCustomGlobalTransformToAsset(UObject* Asset) const;
@@ -108,6 +116,9 @@ private:
 	//A scene node can be part of multiple Layers.
 	UE::Interchange::TArrayAttributeHelper<FString> LayerNames;
 	UE::Interchange::TArrayAttributeHelper<FString> Tags;
+
+	//A scene node can represent many special types
+	UE::Interchange::TArrayAttributeHelper<FString> ComponentUids;
 };
 
 #undef UE_API
diff --git a/Engine/Plugins/Interchange/Runtime/Source/FactoryNodes/Public/InterchangeSceneComponentFactoryNodes.h b/Engine/Plugins/Interchange/Runtime/Source/FactoryNodes/Public/InterchangeSceneComponentFactoryNodes.h
new file mode 100644
--- /dev/null
+++ b/Engine/Plugins/Interchange/Runtime/Source/FactoryNodes/Public/InterchangeSceneComponentFactoryNodes.h
@@ -0,0 +1,101 @@
+// Copyright Epic Games, Inc. All Rights Reserved.
+
+#pragma once
+
+#include "Nodes/InterchangeFactoryBaseNode.h"
+
+#include "InterchangeSceneComponentFactoryNodes.generated.h"
+
+#define UE_API INTERCHANGEFACTORYNODES_API
+
+//TODO: move to its own file.
+UCLASS(MinimalAPI, BlueprintType)
+class UInterchangeSceneComponentFactoryNode : public UInterchangeFactoryBaseNode
+{
+	GENERATED_BODY()
+
+public:
+	UE_API UInterchangeSceneComponentFactoryNode();
+
+	UE_API virtual FString GetTypeName() const override;
+
+	UE_API virtual UClass* GetObjectClass() const override;
+
+
+	UFUNCTION(BlueprintCallable, Category = "Interchange | Scene | Component")
+	UE_API bool AddComponentUid(const FString& ComponentUid);
+
+	UFUNCTION(BlueprintCallable, Category = "Interchange | Scene | Component")
+	UE_API void GetComponentUids(TArray<FString>& OutComponentUids) const;
+
+
+	UFUNCTION(BlueprintCallable, Category = "Interchange | Node | Component")
+	UE_API bool GetCustomLocalTransform(FTransform& AttributeValue) const;
+
+	UFUNCTION(BlueprintCallable, Category = "Interchange | Node | Component")
+	UE_API bool SetCustomLocalTransform(const FTransform& AttributeValue);
+
+
+	UFUNCTION(BlueprintCallable, Category = "Interchange | Node | Component")
+	UE_API bool GetCustomComponentVisibility(bool& AttributeValue) const;
+
+	UFUNCTION(BlueprintCallable, Category = "Interchange | Node | Component")
+	UE_API bool SetCustomComponentVisibility(const bool& AttributeValue);
+
+private:
+	//Component's Local Transform.
+	const UE::Interchange::FAttributeKey Macro_CustomLocalTransformKey = UE::Interchange::FAttributeKey(TEXT("LocalTransform"));
+
+	//Component's Visibility
+	const UE::Interchange::FAttributeKey Macro_CustomComponentVisibilityKey = UE::Interchange::FAttributeKey(TEXT("ComponentVisibility"));
+
+	//Childnre Component Uids.
+	UE::Interchange::TArrayAttributeHelper<FString> ComponentUids;
+};
+
+UCLASS(MinimalAPI, BlueprintType)
+class UInterchangeInstancedStaticMeshComponentFactoryNode  : public UInterchangeSceneComponentFactoryNode
+{
+	GENERATED_BODY()
+public:
+	UE_API UInterchangeInstancedStaticMeshComponentFactoryNode ();
+
+	UE_API virtual FString GetTypeName() const override;
+
+	UE_API virtual UClass* GetObjectClass() const override;
+
+	UFUNCTION(BlueprintCallable, Category = "Interchange | Scene | ComponentFactory | Instanced")
+	UE_API void AddInstanceTransform(const FTransform& InstanceTransform);
+	UFUNCTION(BlueprintCallable, Category = "Interchange | Scene | ComponentFactory | Instanced")
+	UE_API void AddInstanceTransforms(const TArray<FTransform>& InInstanceTransforms);
+
+	UFUNCTION(BlueprintCallable, Category = "Interchange | Scene | ComponentFactory | Instanced")
+	UE_API void GetInstanceTransforms(TArray<FTransform>& OutInstanceTransforms) const;
+
+	/** Get which asset, if any, a scene node is instantiating. Return false if the Attribute was not set previously.*/
+	UFUNCTION(BlueprintCallable, Category = "Interchange | Scene | ComponentFactory | Instanced")
+	UE_API bool GetCustomInstancedAssetUid(FString& AttributeValue) const;
+
+	/** Add an asset for this scene node to instantiate. */
+	UFUNCTION(BlueprintCallable, Category = "Interchange | Scene | ComponentFactory | Instanced")
+	UE_API bool SetCustomInstancedAssetUid(const FString& AttributeValue);
+
+	
+	UFUNCTION(BlueprintCallable, Category = "Interchange | Scene | ComponentFactory | Instanced")
+	UE_API bool GetCustomHierarchicalISM(bool& AttributeValue) const;
+
+	UFUNCTION(BlueprintCallable, Category = "Interchange | Scene | ComponentFactory | Instanced")
+	UE_API bool SetCustomHierarchicalISM(const bool& AttributeValue);
+
+private:
+	//Instanced StaticMesh Component's Instances.
+	UE::Interchange::TArrayAttributeHelper<FTransform> InstanceTransforms;
+
+	//To track/connect StaticMesh.
+	const UE::Interchange::FAttributeKey Macro_CustomAssetInstanceUidKey = UE::Interchange::FAttributeKey(TEXT("AssetInstanceUid"));
+
+	//To track if it should be Hierarchical ISM
+	const UE::Interchange::FAttributeKey Macro_CustomHierarchicalISMKey = UE::Interchange::FAttributeKey(TEXT("Hierarchical"));
+};
+
+#undef UE_API
diff --git a/Engine/Plugins/Interchange/Runtime/Source/Import/Private/Mesh/InterchangeSkeletalMeshFactory.cpp b/Engine/Plugins/Interchange/Runtime/Source/Import/Private/Mesh/InterchangeSkeletalMeshFactory.cpp
--- a/Engine/Plugins/Interchange/Runtime/Source/Import/Private/Mesh/InterchangeSkeletalMeshFactory.cpp
+++ b/Engine/Plugins/Interchange/Runtime/Source/Import/Private/Mesh/InterchangeSkeletalMeshFactory.cpp
@@ -24,6 +24,7 @@
 #include "InterchangeMeshDefinitions.h"
 #include "InterchangeMeshNode.h"
 #include "InterchangeMeshUtilities.h"
+#include "InterchangeSceneComponentNodes.h"
 #include "InterchangeSceneNode.h"
 #include "InterchangeSkeletalMeshLodDataNode.h"
 #include "InterchangeSkeletalMeshFactoryNode.h"
@@ -530,31 +531,70 @@ namespace UE
 						}
 					}
 					
-					//The Mesh node parent bake transform can be pass to the payload request or not, it depend on the count of instance and the scale of the transform.
-					const FInternalInstanceData & InstanceData = MeshInstancesDatas.FindChecked(MeshNodeContext.TranslatorPayloadKey.UniqueId);
-					AppendSettings.MeshTransform = InstanceData.ShouldFetchWithTransform() ? FTransform::Identity : MeshNodeContext.SceneGlobalTransform.Get(FTransform::Identity);
-					if (bKeepSectionsSeparate)
+					if (MeshNodeContext.ISMComponentNode != nullptr)
 					{
-						AppendSettings.PolygonGroupsDelegate = FAppendPolygonGroupsDelegate::CreateLambda([](const FMeshDescription& SourceMesh, FMeshDescription& TargetMesh, PolygonGroupMap& RemapPolygonGroup)
+						TArray<FTransform> InstanceTransforms;
+						MeshNodeContext.ISMComponentNode->GetInstanceTransforms(InstanceTransforms);
+
+						for (const FTransform& InstanceTransform : InstanceTransforms)
+						{
+							//The Mesh node parent bake transform can be pass to the payload request or not, it depend on the count of instance and the scale of the transform.
+							const FInternalInstanceData& InstanceData = MeshInstancesDatas.FindChecked(MeshNodeContext.TranslatorPayloadKey.UniqueId);
+							AppendSettings.MeshTransform = InstanceData.ShouldFetchWithTransform() ? InstanceTransform : InstanceTransform * MeshNodeContext.SceneGlobalTransform.Get(FTransform::Identity);
+							if (bKeepSectionsSeparate)
 							{
-								UE::Interchange::Private::MeshHelper::RemapPolygonGroups(SourceMesh, TargetMesh, RemapPolygonGroup);
-							});
-					}
+								AppendSettings.PolygonGroupsDelegate = FAppendPolygonGroupsDelegate::CreateLambda([](const FMeshDescription& SourceMesh, FMeshDescription& TargetMesh, PolygonGroupMap& RemapPolygonGroup)
+									{
+										UE::Interchange::Private::MeshHelper::RemapPolygonGroups(SourceMesh, TargetMesh, RemapPolygonGroup);
+									});
+							}
 
-					FStaticMeshOperations::AppendMeshDescription(LodMeshPayload->MeshDescription, DestinationMeshDescription, AppendSettings);
-					if (MeshNodeContext.MeshNode->IsSkinnedMesh() || bIsRigidMesh)
-					{
-						FSkeletalMeshOperations::AppendSkinWeight(LodMeshPayload->MeshDescription, DestinationMeshDescription, SkeletalMeshAppendSettings);
+							FStaticMeshOperations::AppendMeshDescription(LodMeshPayload->MeshDescription, DestinationMeshDescription, AppendSettings);
+							if (MeshNodeContext.MeshNode->IsSkinnedMesh() || bIsRigidMesh)
+							{
+								FSkeletalMeshOperations::AppendSkinWeight(LodMeshPayload->MeshDescription, DestinationMeshDescription, SkeletalMeshAppendSettings);
+							}
+							if (bImportMorphTarget)
+							{
+								FillMorphTargetMeshDescriptionsPerMorphTargetName(MeshNodeContext
+									, MorphTargetMeshDescriptionsPerMorphTargetName
+									, LodPayloads
+									, VertexOffset
+									, Arguments.NodeContainer
+									, Arguments.AssetName);
+							}
+						}
 					}
-					if (bImportMorphTarget)
+					else
 					{
-						FillMorphTargetMeshDescriptionsPerMorphTargetName(MeshNodeContext
-																		, MorphTargetMeshDescriptionsPerMorphTargetName
-																		, LodPayloads
-																		, VertexOffset
-																		, Arguments.NodeContainer
-																		, Arguments.AssetName);
+						//The Mesh node parent bake transform can be pass to the payload request or not, it depend on the count of instance and the scale of the transform.
+						const FInternalInstanceData& InstanceData = MeshInstancesDatas.FindChecked(MeshNodeContext.TranslatorPayloadKey.UniqueId);
+						AppendSettings.MeshTransform = InstanceData.ShouldFetchWithTransform() ? FTransform::Identity : MeshNodeContext.SceneGlobalTransform.Get(FTransform::Identity);
+						if (bKeepSectionsSeparate)
+						{
+							AppendSettings.PolygonGroupsDelegate = FAppendPolygonGroupsDelegate::CreateLambda([](const FMeshDescription& SourceMesh, FMeshDescription& TargetMesh, PolygonGroupMap& RemapPolygonGroup)
+								{
+									UE::Interchange::Private::MeshHelper::RemapPolygonGroups(SourceMesh, TargetMesh, RemapPolygonGroup);
+								});
+						}
+
+						FStaticMeshOperations::AppendMeshDescription(LodMeshPayload->MeshDescription, DestinationMeshDescription, AppendSettings);
+						if (MeshNodeContext.MeshNode->IsSkinnedMesh() || bIsRigidMesh)
+						{
+							FSkeletalMeshOperations::AppendSkinWeight(LodMeshPayload->MeshDescription, DestinationMeshDescription, SkeletalMeshAppendSettings);
+						}
+						if (bImportMorphTarget)
+						{
+							FillMorphTargetMeshDescriptionsPerMorphTargetName(MeshNodeContext
+								, MorphTargetMeshDescriptionsPerMorphTargetName
+								, LodPayloads
+								, VertexOffset
+								, Arguments.NodeContainer
+								, Arguments.AssetName);
+						}
 					}
+
+					
 				}
 
 				//The color data is linearized twice by the translator, we need to convert to sRGB to have proper linear
@@ -627,8 +667,9 @@ namespace UE
 				{
 					return false;
 				}
-				const UInterchangeSceneNode* RootJointNode = Cast<UInterchangeSceneNode>(Arguments.NodeContainer->GetNode(RootJointNodeId));
-				if (!RootJointNode)
+
+				const UInterchangeBaseNode* BaseNode = Arguments.NodeContainer->GetNode(RootJointNodeId);
+				if (!BaseNode)
 				{
 					return false;
 				}
@@ -664,12 +705,20 @@ namespace UE
 					, BoneNotBindNames
 					, bImportSockets);
 
-				FTransform RootJointNodeGlobalTransform;
-				ensure(RootJointNode->GetCustomGlobalTransform(Arguments.NodeContainer, GlobalOffsetTransform, RootJointNodeGlobalTransform));
-				FTransform RootJointNodeLocalTransform;
-				ensure(RootJointNode->GetCustomLocalTransform(RootJointNodeLocalTransform));
-				FTransform BakeToRootJointTransfromModifier = RootJointNodeGlobalTransform.Inverse() * RootJointNodeLocalTransform; //It is used for !bBakeMeshes, and so the GlobalTransform will be inversed out when multiplied into the CustomBindPoseGlobalTransform
-				FTransform BakeFromRootJointTransfromModifier = RootJointNodeLocalTransform.Inverse() * RootJointNodeGlobalTransform * GlobalOffsetTransform.Inverse(); //GlobalOffseTransform will be added by the BindPoseGlobalTransform when bBakeMeshes && !bRootAncestorOfMeshDependency is used
+				
+				FTransform BakeToRootJointTransfromModifier;
+				FTransform BakeFromRootJointTransfromModifier;
+
+				if (const UInterchangeSceneNode* RootJointNode = Cast<UInterchangeSceneNode>(BaseNode))
+				{
+					FTransform RootJointNodeGlobalTransform;
+					FTransform RootJointNodeLocalTransform;
+					RootJointNode->GetCustomLocalTransform(RootJointNodeLocalTransform);
+					RootJointNode->GetCustomGlobalTransform(Arguments.NodeContainer, GlobalOffsetTransform, RootJointNodeGlobalTransform);
+
+					BakeToRootJointTransfromModifier = RootJointNodeGlobalTransform.Inverse() * RootJointNodeLocalTransform; //It is used for !bBakeMeshes, and so the GlobalTransform will be inversed out when multiplied into the CustomBindPoseGlobalTransform
+					BakeFromRootJointTransfromModifier = RootJointNodeLocalTransform.Inverse() * RootJointNodeGlobalTransform * GlobalOffsetTransform.Inverse(); //GlobalOffseTransform will be added by the BindPoseGlobalTransform when bBakeMeshes && !bRootAncestorOfMeshDependency is used
+				}
 
 				{
 					// fbx legacy has a special way to bake the skeletal mesh that do not fit the interchange standard.
@@ -693,19 +742,40 @@ namespace UE
 					for (const FString& MeshUid : MeshUids)
 					{
 						FMeshNodeContext MeshReference;
-						MeshReference.MeshNode = Cast<UInterchangeMeshNode>(Arguments.NodeContainer->GetNode(MeshUid));
+						const UInterchangeBaseNode* MeshBaseNode = Arguments.NodeContainer->GetNode(MeshUid);
+						MeshReference.MeshNode = Cast<UInterchangeMeshNode>(MeshBaseNode);
 						if (!MeshReference.MeshNode)
 						{
 							//The reference is a scene node and we need to bake the geometry
-							MeshReference.SceneNode = Cast<UInterchangeSceneNode>(Arguments.NodeContainer->GetNode(MeshUid));
+							MeshReference.SceneNode = Cast<UInterchangeSceneNode>(MeshBaseNode);
+
+							FTransform SceneNodeComponentTransform;
+							if (MeshReference.SceneNode == nullptr)
+							{
+								MeshReference.ISMComponentNode = Cast<UInterchangeInstancedStaticMeshComponentNode>(MeshBaseNode);
+								if (MeshReference.ISMComponentNode != nullptr)
+								{
+									MeshReference.SceneNode = MeshReference.ISMComponentNode->GetParentSceneNodeAndTransform(Arguments.NodeContainer, SceneNodeComponentTransform);
+									
+									FString AssetInstanceUid;
+									MeshReference.ISMComponentNode->GetCustomInstancedAssetUid(AssetInstanceUid);
+									MeshReference.MeshNode = Cast<UInterchangeMeshNode>(Arguments.NodeContainer->GetNode(AssetInstanceUid));
+								}
+							}
+
+							if (MeshReference.ISMComponentNode == nullptr)
+							{
+								FString MeshDependencyUid;
+								MeshReference.SceneNode->GetCustomAssetInstanceUid(MeshDependencyUid);
+								MeshReference.MeshNode = Cast<UInterchangeMeshNode>(Arguments.NodeContainer->GetNode(MeshDependencyUid));
+							}
+
 							if (!ensure(MeshReference.SceneNode != nullptr))
 							{
 								continue;
 							}
-							FString MeshDependencyUid;
-							MeshReference.SceneNode->GetCustomAssetInstanceUid(MeshDependencyUid);
-							MeshReference.MeshNode = Cast<UInterchangeMeshNode>(Arguments.NodeContainer->GetNode(MeshDependencyUid));
-							bool bRootAncestorOfSceneNode = Arguments.NodeContainer->GetIsAncestor(MeshReference.SceneNode->GetUniqueID(), RootJointNode->GetParentUid());
+
+							bool bRootAncestorOfSceneNode = Arguments.NodeContainer->GetIsAncestor(MeshReference.SceneNode->GetUniqueID(), BaseNode->GetParentUid());
 							//Cache the scene node global matrix, we will use this matrix to bake the vertices, add the node geometric mesh offset to this matrix to bake it properly
 							FTransform SceneNodeTransform;
 							if (!ImportAssetObjectLODDataRef.bUseTimeZeroAsBindPose || !MeshReference.SceneNode->GetCustomTimeZeroGlobalTransform(Arguments.NodeContainer, GlobalOffsetTransform, SceneNodeTransform))
@@ -736,6 +806,12 @@ namespace UE
 							{
 								SceneNodeTransform = SceneNodeGeometricTransform * SceneNodeTransform;
 							}
+
+							if (MeshReference.ISMComponentNode != nullptr)
+							{
+								SceneNodeTransform = SceneNodeComponentTransform * SceneNodeTransform;
+							}
+
 							MeshReference.SceneGlobalTransform = SceneNodeTransform;
 						}
 						else
@@ -847,10 +923,11 @@ void UInterchangeSkeletalMeshFactory::CreatePayloadTasks(const FImportAssetObjec
 			continue;
 		}
 
-		const UInterchangeSceneNode* RootJointNode = Cast<UInterchangeSceneNode>(Arguments.NodeContainer->GetNode(RootJointNodeId));
-		if (!RootJointNode)
+		FTransform RootJointGlobalTransform;
+		const UInterchangeBaseNode* RootJointNode = Arguments.NodeContainer->GetNode(RootJointNodeId);
+		if (const UInterchangeSceneNode* RootJointSceneNode = Cast<const UInterchangeSceneNode>(RootJointNode))
 		{
-			continue;
+			RootJointSceneNode->GetCustomGlobalTransform(Arguments.NodeContainer, GlobalOffsetTransform, RootJointGlobalTransform);
 		}
 
 		FLodPayloads& LodPayloads = PayloadsPerLodIndex.FindOrAdd(LodIndex);
@@ -860,9 +937,6 @@ void UInterchangeSkeletalMeshFactory::CreatePayloadTasks(const FImportAssetObjec
 
 		SkeletonNode->GetCustomUseTimeZeroForBindPose(ImportAssetObjectLODData.bUseTimeZeroAsBindPose);
 
-		FTransform RootJointGlobalTransform;
-		RootJointNode->GetCustomGlobalTransform(Arguments.NodeContainer, GlobalOffsetTransform, RootJointGlobalTransform);
-
 		Private::BuildMeshReferences(
 			Arguments
 			, RootJointNodeId
@@ -1151,7 +1225,7 @@ UInterchangeFactoryBase::FImportAssetResult UInterchangeSkeletalMeshFactory::Beg
 			continue;
 		}
 
-		const UInterchangeSceneNode* RootJointNode = Cast<UInterchangeSceneNode>(Arguments.NodeContainer->GetNode(RootJointNodeId));
+		const UInterchangeBaseNode* RootJointNode = Arguments.NodeContainer->GetNode(RootJointNodeId);
 		if (!RootJointNode)
 		{
 			UInterchangeResultWarning_Generic* Message = AddMessage<UInterchangeResultWarning_Generic>();
@@ -1356,7 +1430,7 @@ UInterchangeFactoryBase::FImportAssetResult UInterchangeSkeletalMeshFactory::Imp
 			continue;
 		}
 		
-		const UInterchangeSceneNode* RootJointNode = Cast<UInterchangeSceneNode>(Arguments.NodeContainer->GetNode(RootJointNodeId));
+		const UInterchangeBaseNode* RootJointNode = Arguments.NodeContainer->GetNode(RootJointNodeId);
 		if (!RootJointNode)
 		{
 			UInterchangeResultWarning_Generic* Message = AddMessage<UInterchangeResultWarning_Generic>();
diff --git a/Engine/Plugins/Interchange/Runtime/Source/Import/Private/Mesh/InterchangeSkeletonHelper.cpp b/Engine/Plugins/Interchange/Runtime/Source/Import/Private/Mesh/InterchangeSkeletonHelper.cpp
--- a/Engine/Plugins/Interchange/Runtime/Source/Import/Private/Mesh/InterchangeSkeletonHelper.cpp
+++ b/Engine/Plugins/Interchange/Runtime/Source/Import/Private/Mesh/InterchangeSkeletonHelper.cpp
@@ -249,10 +249,31 @@ namespace UE::Interchange::Private
 		, TArray<FString>& OutBoneNotBindNames
 		, bool bImportSockets)
 	{
-		const UInterchangeSceneNode* JointNode = Cast<UInterchangeSceneNode>(NodeContainer->GetNode(JointNodeId));
+		const UInterchangeBaseNode* BaseNode = NodeContainer->GetNode(JointNodeId);
+		const UInterchangeSceneNode* JointNode = Cast<UInterchangeSceneNode>(BaseNode);
 		if (!JointNode)
 		{
-			UE_LOG(LogInterchangeImport, Warning, TEXT("Invalid Skeleton Joint"));
+			if (JointInfos.Num() == 0 && BaseNode)
+			{
+				FJointInfo& Info = JointInfos.AddZeroed_GetRef();
+				Info.Name = BaseNode->GetDisplayLabel();
+				Info.ParentIndex = INDEX_NONE;
+				Info.LocalTransform = FTransform::Identity;
+
+				SkeletalMeshImportData::FBone& Bone = RefBonesBinary.AddZeroed_GetRef();
+				Bone.Name = Info.Name;
+				Bone.ParentIndex = Info.ParentIndex;
+				Bone.BonePos.Transform = FTransform3f::Identity;
+				Bone.BonePos.Length = 0.0f;
+				Bone.BonePos.XSize = 1.0f;
+				Bone.BonePos.YSize = 1.0f;
+				Bone.BonePos.ZSize = 1.0f;
+			}
+			else
+			{
+				UE_LOG(LogInterchangeImport, Warning, TEXT("Invalid Skeleton Joint"));
+			}
+			
 			return;
 		}
 
@@ -266,7 +287,7 @@ namespace UE::Interchange::Private
 		FTransform BindPoseLocalTransform;
 		bool bHasBindPoseTransform = false;
 
-		ensure(JointNode->GetCustomLocalTransform(LocalTransform));
+		JointNode->GetCustomLocalTransform(LocalTransform);
 		bHasTimeZeroTransform = JointNode->GetCustomTimeZeroLocalTransform(TimeZeroLocalTransform);
 		bHasBindPoseTransform = JointNode->GetCustomBindPoseLocalTransform(BindPoseLocalTransform);
 
diff --git a/Engine/Plugins/Interchange/Runtime/Source/Import/Private/Mesh/InterchangeStaticMeshFactory.cpp b/Engine/Plugins/Interchange/Runtime/Source/Import/Private/Mesh/InterchangeStaticMeshFactory.cpp
--- a/Engine/Plugins/Interchange/Runtime/Source/Import/Private/Mesh/InterchangeStaticMeshFactory.cpp
+++ b/Engine/Plugins/Interchange/Runtime/Source/Import/Private/Mesh/InterchangeStaticMeshFactory.cpp
@@ -13,6 +13,7 @@
 #include "InterchangeImportLog.h"
 #include "InterchangeManager.h"
 #include "InterchangeMeshNode.h"
+#include "InterchangeSceneComponentNodes.h"
 #include "InterchangeSceneNode.h"
 #include "InterchangeSourceData.h"
 #include "InterchangeStaticMeshFactoryNode.h"
@@ -115,6 +116,15 @@ void UInterchangeStaticMeshFactory::CreatePayloadTasks(const FImportAssetObjectP
 			OutMeshNode = Cast<const UInterchangeMeshNode>(Node);
 			if (OutMeshNode == nullptr)
 			{
+				if (const UInterchangeInstancedStaticMeshComponentNode* ISMComponentNode = Cast<const UInterchangeInstancedStaticMeshComponentNode>(Node))
+				{
+					FString MeshDependencyUid;
+					ISMComponentNode->GetCustomInstancedAssetUid(MeshDependencyUid);
+					OutMeshNode = Cast<UInterchangeMeshNode>(Arguments.NodeContainer->GetNode(MeshDependencyUid));
+
+					return;
+				}
+
 				// MeshUid must refer to a scene node
 				const UInterchangeSceneNode* SceneNode = Cast<const UInterchangeSceneNode>(Node);
 				if (!ensure(SceneNode))
diff --git a/Engine/Plugins/Interchange/Runtime/Source/Import/Private/Scene/InterchangeActorFactory.cpp b/Engine/Plugins/Interchange/Runtime/Source/Import/Private/Scene/InterchangeActorFactory.cpp
--- a/Engine/Plugins/Interchange/Runtime/Source/Import/Private/Scene/InterchangeActorFactory.cpp
+++ b/Engine/Plugins/Interchange/Runtime/Source/Import/Private/Scene/InterchangeActorFactory.cpp
@@ -3,6 +3,7 @@
 
 #include "InterchangeActorFactoryNode.h"
 #include "InterchangeCameraFactoryNode.h"
+#include "InterchangeSceneComponentFactoryNodes.h"
 #include "Nodes/InterchangeBaseNodeContainer.h"
 #include "Nodes/InterchangeFactoryBaseNode.h"
 #include "Scene/InterchangeActorHelper.h"
@@ -11,6 +12,9 @@
 #include "Components/SceneComponent.h"
 #include "Components/ActorComponent.h"
 
+#include "Components/InstancedStaticMeshComponent.h"
+#include "InterchangeStaticMeshFactoryNode.h"
+
 #if WITH_EDITORONLY_DATA
 #include "Engine/World.h"
 #include "Layers/LayersSubsystem.h"
@@ -82,6 +86,85 @@ void UInterchangeActorFactory::ExecuteResetObjectProperties(const UInterchangeBa
 	Super::ExecuteResetObjectProperties(BaseNodeContainer, FactoryNode, ImportedObject);	
 }
 
+
+namespace UE::Interchange::Private
+{
+	USceneComponent* CreateSceneComponent(AActor* Owner, UClass* ComponentClass, USceneComponent* ParentComponent, const FString& DisplayName)
+	{
+		USceneComponent* SceneComponent = NewObject<USceneComponent>(Owner, ComponentClass, *DisplayName);
+		SceneComponent->SetupAttachment(ParentComponent);
+		SceneComponent->RegisterComponent();
+		Owner->AddInstanceComponent(SceneComponent);
+		return SceneComponent;
+	};
+
+	void ProcessComponents(AActor* Owner, const UInterchangeBaseNodeContainer* NodeContainer, const TArray<FString>& ComponentUids, USceneComponent* ParentComponent, EComponentMobility::Type MobilityToSet)
+	{
+		for (const FString& ComponentUid : ComponentUids)
+		{
+			if (const UInterchangeSceneComponentFactoryNode* SceneComponentFactoryNode = Cast<const UInterchangeSceneComponentFactoryNode>(NodeContainer->GetNode(ComponentUid)))
+			{
+				UClass* ComponentClass = SceneComponentFactoryNode->GetObjectClass();
+				FString Name = SceneComponentFactoryNode->GetDisplayLabel();
+				
+				USceneComponent* CreatedSceneComponent = CreateSceneComponent(Owner, ComponentClass, ParentComponent, Name);
+
+				if (const UInterchangeInstancedStaticMeshComponentFactoryNode * ISMComponentNode = Cast<const UInterchangeInstancedStaticMeshComponentFactoryNode >(NodeContainer->GetNode(ComponentUid)))
+				{
+					if (UInstancedStaticMeshComponent* ISMComponent = Cast<UInstancedStaticMeshComponent>(CreatedSceneComponent))
+					{
+						TArray<FTransform> InstanceTransforms;
+						ISMComponentNode->GetInstanceTransforms(InstanceTransforms);
+						for (const FTransform& InstanceTransform : InstanceTransforms)
+						{
+							ISMComponent->AddInstance(InstanceTransform);
+						}
+
+						FString AssetInstanceUid;
+						ISMComponentNode->GetCustomInstancedAssetUid(AssetInstanceUid);
+
+						if (const UInterchangeStaticMeshFactoryNode* StaticMeshFactoryNode = Cast<const UInterchangeStaticMeshFactoryNode>(NodeContainer->GetNode(AssetInstanceUid)))
+						{
+							FSoftObjectPath ReferenceObject;
+							StaticMeshFactoryNode->GetCustomReferenceObject(ReferenceObject);
+							if (UStaticMesh* StaticMesh = Cast<UStaticMesh>(ReferenceObject.TryLoad()))
+							{
+								ISMComponent->SetStaticMesh(StaticMesh);
+							}
+						}
+					}
+				}
+
+				//General setups:
+				{
+					FTransform LocalTransform;
+					SceneComponentFactoryNode->GetCustomLocalTransform(LocalTransform);
+
+					CreatedSceneComponent->SetMobility(EComponentMobility::Type::Movable); //so that RelativeTransform can be set
+					CreatedSceneComponent->SetRelativeTransform(LocalTransform);
+					CreatedSceneComponent->SetMobility(MobilityToSet);
+
+					bool bComponentVisibility = true;
+					if (SceneComponentFactoryNode->GetCustomComponentVisibility(bComponentVisibility))
+					{
+						CreatedSceneComponent->SetVisibility(bComponentVisibility);
+					}
+
+					TArray<FString> ChildrenComponentUids;
+					SceneComponentFactoryNode->GetComponentUids(ChildrenComponentUids);
+
+					ProcessComponents(Owner, NodeContainer, ChildrenComponentUids, CreatedSceneComponent, MobilityToSet);
+				}
+			}
+			else
+			{
+				ensureMsgf(false, TEXT("Unexpected Component class type in ComponentUids."));
+			}
+		}
+	}
+}
+
+
 UObject* UInterchangeActorFactory::ImportSceneObject_GameThread(const UInterchangeFactoryBase::FImportSceneObjectsParams& CreateSceneObjectsParams)
 {
 	using namespace UE::Interchange;
@@ -106,6 +189,11 @@ UObject* UInterchangeActorFactory::ImportSceneObject_GameThread(const UInterchan
 
 				ApplyAllCustomAttributesToObject(CreateSceneObjectsParams, *SpawnedActor, ObjectToUpdate);
 
+				TArray<FString> ComponentUids;
+				FactoryNode->GetComponentUids(ComponentUids);
+
+				UE::Interchange::Private::ProcessComponents(SpawnedActor, CreateSceneObjectsParams.NodeContainer, ComponentUids, RootComponent, CachedMobility);
+				
 				// Restore mobility value
 				if (CachedMobility != EComponentMobility::Type::Movable)
 				{
diff --git a/Engine/Plugins/Interchange/Runtime/Source/Import/Public/Mesh/InterchangeSkeletalMeshFactory.h b/Engine/Plugins/Interchange/Runtime/Source/Import/Public/Mesh/InterchangeSkeletalMeshFactory.h
--- a/Engine/Plugins/Interchange/Runtime/Source/Import/Public/Mesh/InterchangeSkeletalMeshFactory.h
+++ b/Engine/Plugins/Interchange/Runtime/Source/Import/Public/Mesh/InterchangeSkeletalMeshFactory.h
@@ -17,6 +17,7 @@
 
 #define UE_API INTERCHANGEIMPORT_API
 
+class UInterchangeInstancedStaticMeshComponentNode;
 class UInterchangeSceneNode;
 class UInterchangeSkeletalMeshFactoryNode;
 class USkeletalMesh;
@@ -29,6 +30,7 @@ namespace UE::Interchange
 	{
 		const UInterchangeMeshNode* MeshNode = nullptr;
 		const UInterchangeSceneNode* SceneNode = nullptr;
+		const UInterchangeInstancedStaticMeshComponentNode* ISMComponentNode = nullptr;
 		TOptional<FTransform> SceneGlobalTransform;
 		FInterchangeMeshPayLoadKey TranslatorPayloadKey;
 
diff --git a/Engine/Plugins/Interchange/Runtime/Source/Nodes/Private/InterchangeMeshLODContainerNode.cpp b/Engine/Plugins/Interchange/Runtime/Source/Nodes/Private/InterchangeMeshLODContainerNode.cpp
new file mode 100644
--- /dev/null
+++ b/Engine/Plugins/Interchange/Runtime/Source/Nodes/Private/InterchangeMeshLODContainerNode.cpp
@@ -0,0 +1,36 @@
+// Copyright Epic Games, Inc. All Rights Reserved.
+
+#include "InterchangeMeshLODContainerNode.h"
+
+#include UE_INLINE_GENERATED_CPP_BY_NAME(InterchangeMeshLODContainerNode)
+
+UInterchangeMeshLODContainerNode::UInterchangeMeshLODContainerNode()
+{
+	LODMeshUids.Initialize(Attributes.ToSharedRef(), TEXT("__LODMeshUids__Key"));
+}
+
+FString UInterchangeMeshLODContainerNode::GetTypeName() const
+{
+	const FString TypeName = TEXT("LODContainerComponentNode");
+	return TypeName;
+}
+
+bool UInterchangeMeshLODContainerNode::AddMeshLODNodeUid(const FString& MeshLODNodeUid)
+{
+	return LODMeshUids.AddItem(MeshLODNodeUid);
+}
+
+void UInterchangeMeshLODContainerNode::GetMeshLODNodeUids(TArray<FString>& OutMeshLODNodeUid) const
+{
+	LODMeshUids.GetItems(OutMeshLODNodeUid);
+}
+
+bool UInterchangeMeshLODContainerNode::RemoveMeshLODNodeUid(const FString& MeshLODNodeUid)
+{
+	return LODMeshUids.RemoveItem(MeshLODNodeUid);
+}
+
+bool UInterchangeMeshLODContainerNode::ResetMeshLODNodeUids()
+{
+	return LODMeshUids.RemoveAllItems();
+}
\ No newline at end of file
diff --git a/Engine/Plugins/Interchange/Runtime/Source/Nodes/Private/InterchangeSceneComponentNodes.cpp b/Engine/Plugins/Interchange/Runtime/Source/Nodes/Private/InterchangeSceneComponentNodes.cpp
new file mode 100644
--- /dev/null
+++ b/Engine/Plugins/Interchange/Runtime/Source/Nodes/Private/InterchangeSceneComponentNodes.cpp
@@ -0,0 +1,156 @@
+// Copyright Epic Games, Inc. All Rights Reserved.
+
+#include "InterchangeSceneComponentNodes.h"
+#include "InterchangeSceneNode.h"
+#include "Nodes/InterchangeBaseNodeContainer.h"
+
+#include UE_INLINE_GENERATED_CPP_BY_NAME(InterchangeSceneComponentNodes)
+
+UInterchangeSceneComponentNode::UInterchangeSceneComponentNode()
+{
+	ComponentUids.Initialize(Attributes, TEXT("__ComponentUids__Key"));
+}
+
+FString UInterchangeSceneComponentNode::GetTypeName() const
+{
+	static const FString TypeName = TEXT("SceneComponentNode");
+	return TypeName;
+}
+
+bool UInterchangeSceneComponentNode::AddComponentUid(const FString& ComponentUid)
+{
+	return ComponentUids.AddItem(ComponentUid);
+}
+
+void UInterchangeSceneComponentNode::GetComponentUids(TArray<FString>& OutComponentUids) const
+{
+	ComponentUids.GetItems(OutComponentUids);
+}
+
+bool UInterchangeSceneComponentNode::GetCustomLocalTransform(FTransform& AttributeValue) const
+{
+	IMPLEMENT_NODE_ATTRIBUTE_GETTER(LocalTransform, FTransform);
+}
+
+bool UInterchangeSceneComponentNode::SetCustomLocalTransform(const FTransform& AttributeValue)
+{
+	IMPLEMENT_NODE_ATTRIBUTE_SETTER_NODELEGATE(LocalTransform, FTransform);
+}
+
+bool UInterchangeSceneComponentNode::GetCustomComponentVisibility(bool& AttributeValue) const
+{
+	IMPLEMENT_NODE_ATTRIBUTE_GETTER(ComponentVisibility, bool);
+}
+
+bool UInterchangeSceneComponentNode::SetCustomComponentVisibility(const bool& AttributeValue)
+{
+	IMPLEMENT_NODE_ATTRIBUTE_SETTER_NODELEGATE(ComponentVisibility, bool);
+}
+
+const UInterchangeSceneNode* UInterchangeSceneComponentNode::GetParentSceneNodeAndTransform(const UInterchangeBaseNodeContainer* BaseNodeContainer
+	, FTransform& SceneNodeTransform
+	, bool bForceRecache) const
+{
+	if (bForceRecache)
+	{
+		CacheSceneNodeTransform.Reset();
+		CacheParentSceneNode.Reset();
+	}
+	if (!CacheSceneNodeTransform.IsSet())
+	{
+		FTransform LocalTransform;
+		if (GetCustomLocalTransform(LocalTransform))
+		{
+			//Compute the Global
+			FString ParentUid = GetParentUid();
+			if (!ParentUid.IsEmpty())
+			{
+				const UInterchangeBaseNode* ParentBaseNode = BaseNodeContainer->GetNode(ParentUid);
+				FTransform SceneNodeGlobalParent;
+				if (const UInterchangeSceneComponentNode* ParentSceneComponentNode = Cast<UInterchangeSceneComponentNode>(ParentBaseNode))
+				{
+					const UInterchangeSceneNode* ParentSceneNode = ParentSceneComponentNode->GetParentSceneNodeAndTransform(BaseNodeContainer, SceneNodeGlobalParent, bForceRecache);
+					CacheParentSceneNode = ParentSceneNode;
+				}
+				if (const UInterchangeSceneNode* ParentSceneNode = Cast<UInterchangeSceneNode>(ParentBaseNode))
+				{
+					CacheParentSceneNode = ParentSceneNode;
+					CacheSceneNodeTransform = LocalTransform;
+				}
+				CacheSceneNodeTransform = LocalTransform * SceneNodeGlobalParent;
+			}
+			else
+			{
+				//Scene Node without parent will need the global offset to be apply
+				CacheSceneNodeTransform = LocalTransform;
+
+				return nullptr;
+			}
+		}
+		else
+		{
+			CacheSceneNodeTransform = FTransform::Identity;
+		}
+	}
+	//The cache is always valid here
+	if (ensure(CacheSceneNodeTransform.IsSet()))
+	{
+		SceneNodeTransform = CacheSceneNodeTransform.GetValue();
+	}
+	else
+	{
+		SceneNodeTransform = FTransform::Identity;
+	}
+
+	if (CacheParentSceneNode.IsSet())
+	{
+		return CacheParentSceneNode.GetValue();
+	}
+	else
+	{
+		return nullptr;
+	}
+}
+
+
+
+
+FString UInterchangeInstancedStaticMeshComponentNode::GetTypeName() const
+{
+	const FString TypeName = TEXT("InstancedStaticMeshComponentNode");
+	return TypeName;
+}
+
+UInterchangeInstancedStaticMeshComponentNode::UInterchangeInstancedStaticMeshComponentNode()
+{
+	InstanceTransforms.Initialize(Attributes, TEXT("__InstanceTransforms__"));
+}
+void UInterchangeInstancedStaticMeshComponentNode::AddInstanceTransform(const FTransform& InstanceTransform)
+{
+	InstanceTransforms.AddItem(InstanceTransform);
+}
+void UInterchangeInstancedStaticMeshComponentNode::AddInstanceTransforms(const TArray<FTransform>& InInstanceTransforms)
+{
+	for (const FTransform& InstanceTransform : InInstanceTransforms)
+	{
+		InstanceTransforms.AddItem(InstanceTransform);
+	}
+}
+void UInterchangeInstancedStaticMeshComponentNode::GetInstanceTransforms(TArray<FTransform>& OutInstanceTransforms) const
+{
+	InstanceTransforms.GetItems(OutInstanceTransforms);
+}
+
+bool UInterchangeInstancedStaticMeshComponentNode::GetCustomInstancedAssetUid(FString& AttributeValue) const
+{
+	IMPLEMENT_NODE_ATTRIBUTE_GETTER(AssetInstanceUid, FString);
+}
+
+bool UInterchangeInstancedStaticMeshComponentNode::SetCustomInstancedAssetUid(const FString& AttributeValue)
+{
+	IMPLEMENT_NODE_ATTRIBUTE_SETTER_NODELEGATE(AssetInstanceUid, FString);
+}
+
+
+
+
diff --git a/Engine/Plugins/Interchange/Runtime/Source/Nodes/Private/InterchangeSceneNode.cpp b/Engine/Plugins/Interchange/Runtime/Source/Nodes/Private/InterchangeSceneNode.cpp
--- a/Engine/Plugins/Interchange/Runtime/Source/Nodes/Private/InterchangeSceneNode.cpp
+++ b/Engine/Plugins/Interchange/Runtime/Source/Nodes/Private/InterchangeSceneNode.cpp
@@ -76,6 +76,12 @@ namespace UE
 			static FString CurveAnimationTypesKey(TEXT("__CurveAnimationTypes__Key"));
 			return CurveAnimationTypesKey;
 		}
+
+		const FString& FSceneNodeStaticData::GetComponentUidsKey()
+		{
+			static FString ComponentsUidsKey(TEXT("__ComponentUids__Key"));
+			return ComponentsUidsKey;
+		}
 	}//ns Interchange
 }//ns UE
 
@@ -88,6 +94,7 @@ UInterchangeSceneNode::UInterchangeSceneNode()
 	LayerNames.Initialize(Attributes, UE::Interchange::FSceneNodeStaticData::GetLayerNamesKey());
 	Tags.Initialize(Attributes, UE::Interchange::FSceneNodeStaticData::GetTagsKey());
 	CurveAnimationTypes.Initialize(Attributes.ToSharedRef(), UE::Interchange::FSceneNodeStaticData::GetCurveAnimationTypesKey());
+	ComponentUids.Initialize(Attributes.ToSharedRef(), UE::Interchange::FSceneNodeStaticData::GetComponentUidsKey());
 }
 
 /**
@@ -365,6 +372,8 @@ bool UInterchangeSceneNode::GetGlobalTransformInternal(const UE::Interchange::FA
 	, FTransform& AttributeValue
 	, bool bForceRecache) const
 {
+	bool bAllowsNoTransforms = IsSpecializedTypeContains(UE::Interchange::FSceneNodeStaticData::GetLodGroupSpecializeTypeString());
+
 	UE::Interchange::FAttributeKey TransformKey = LocalTransformKey;
 	if (!Attributes->ContainAttribute(TransformKey))
 	{
@@ -373,8 +382,9 @@ bool UInterchangeSceneNode::GetGlobalTransformInternal(const UE::Interchange::FA
 		{
 			TransformKey = Macro_CustomLocalTransformKey;
 		}
-		else
+		else if (!bAllowsNoTransforms)
 		{
+			//LOD Group nodes do not necessarily have Transforms.
 			return false;
 		}
 	}
@@ -384,15 +394,17 @@ bool UInterchangeSceneNode::GetGlobalTransformInternal(const UE::Interchange::FA
 	}
 	if (!CacheTransform.IsSet())
 	{
-		FTransform LocalTransform;
+		FTransform LocalTransform = FTransform::Identity;
 		UE::Interchange::FAttributeStorage::TAttributeHandle<FTransform> AttributeHandle = GetAttributeHandle<FTransform>(TransformKey);
-		if (AttributeHandle.IsValid() && AttributeHandle.Get(LocalTransform) == UE::Interchange::EAttributeStorageResult::Operation_Success)
+		if ((AttributeHandle.IsValid() && AttributeHandle.Get(LocalTransform) == UE::Interchange::EAttributeStorageResult::Operation_Success)
+			|| bAllowsNoTransforms)
 		{
 			//Compute the Global
-			if (Attributes->ContainAttribute(UE::Interchange::FBaseNodeStaticData::ParentIDKey()))
+			FString ParentUid = GetParentUid();
+			if (!ParentUid.IsEmpty())
 			{
 				FTransform GlobalParent;
-				if (const UInterchangeSceneNode* ParentSceneNode = Cast<UInterchangeSceneNode>(BaseNodeContainer->GetNode(GetParentUid())))
+				if (const UInterchangeSceneNode* ParentSceneNode = Cast<UInterchangeSceneNode>(BaseNodeContainer->GetNode(ParentUid)))
 				{
 					if (LocalTransformKey == Macro_CustomLocalTransformKey)
 					{
@@ -527,4 +539,14 @@ bool UInterchangeSceneNode::SetAnimationCurveTypeForCurveName(const FString& Cur
 bool UInterchangeSceneNode::GetAnimationCurveTypeForCurveName(const FString& CurveName, EInterchangeAnimationPayLoadType& OutCurveAnimationType) const
 {
 	return CurveAnimationTypes.GetValue(CurveName, OutCurveAnimationType);
-}
\ No newline at end of file
+}
+
+bool UInterchangeSceneNode::AddComponentUid(const FString& ComponentUid)
+{
+	return ComponentUids.AddItem(ComponentUid);
+}
+
+void UInterchangeSceneNode::GetComponentUids(TArray<FString>& OutComponentUids) const
+{
+	ComponentUids.GetItems(OutComponentUids);
+}
diff --git a/Engine/Plugins/Interchange/Runtime/Source/Nodes/Public/InterchangeMeshLODContainerNode.h b/Engine/Plugins/Interchange/Runtime/Source/Nodes/Public/InterchangeMeshLODContainerNode.h
new file mode 100644
--- /dev/null
+++ b/Engine/Plugins/Interchange/Runtime/Source/Nodes/Public/InterchangeMeshLODContainerNode.h
@@ -0,0 +1,38 @@
+// Copyright Epic Games, Inc. All Rights Reserved.
+
+#pragma once
+
+#include "Nodes/InterchangeBaseNode.h"
+
+#include "InterchangeMeshLODContainerNode.generated.h"
+
+#define UE_API INTERCHANGENODES_API
+
+UCLASS(MinimalAPI, BlueprintType)
+class UInterchangeMeshLODContainerNode : public UInterchangeBaseNode
+{
+	GENERATED_BODY()
+
+public:
+
+	UE_API UInterchangeMeshLODContainerNode();
+
+	UE_API virtual FString GetTypeName() const override;
+
+	UFUNCTION(BlueprintCallable, Category = "Interchange | Scene | Component | LODContainer")
+	UE_API bool AddMeshLODNodeUid(const FString& MeshLODNodeUid);
+
+	UFUNCTION(BlueprintCallable, Category = "Interchange | Scene | Component | LODContainer")
+	UE_API void GetMeshLODNodeUids(TArray<FString>& OutMeshLODNodeUid) const;
+
+	UFUNCTION(BlueprintCallable, Category = "Interchange | Scene | Component | LODContainer")
+	UE_API bool RemoveMeshLODNodeUid(const FString& MeshLODNodeUid);
+
+	UFUNCTION(BlueprintCallable, Category = "Interchange | Scene | Component | LODContainer")
+	UE_API bool ResetMeshLODNodeUids();
+
+private:
+	UE::Interchange::TArrayAttributeHelper<FString> LODMeshUids;
+};
+
+#undef UE_API
diff --git a/Engine/Plugins/Interchange/Runtime/Source/Nodes/Public/InterchangeSceneComponentNodes.h b/Engine/Plugins/Interchange/Runtime/Source/Nodes/Public/InterchangeSceneComponentNodes.h
new file mode 100644
--- /dev/null
+++ b/Engine/Plugins/Interchange/Runtime/Source/Nodes/Public/InterchangeSceneComponentNodes.h
@@ -0,0 +1,106 @@
+// Copyright Epic Games, Inc. All Rights Reserved.
+
+#pragma once
+
+#include "CoreMinimal.h"
+#include "Nodes/InterchangeBaseNode.h"
+
+#include "InterchangeSceneComponentNodes.generated.h"
+
+#define UE_API INTERCHANGENODES_API
+
+class UInterchangeSceneNode;
+
+UCLASS(MinimalAPI, BlueprintType)
+class UInterchangeSceneComponentNode : public UInterchangeBaseNode
+{
+	GENERATED_BODY()
+
+public:
+
+	UE_API UInterchangeSceneComponentNode();
+
+	/**
+	 * Return the node type name of the class. This is used when reporting errors.
+	 */
+	UE_API virtual FString GetTypeName() const override;
+
+
+	UFUNCTION(BlueprintCallable, Category = "Interchange | Scene | Component")
+	UE_API bool AddComponentUid(const FString& ComponentUid);
+
+	UFUNCTION(BlueprintCallable, Category = "Interchange | Scene | Component")
+	UE_API void GetComponentUids(TArray<FString>& OutComponentUids) const;
+
+
+	UFUNCTION(BlueprintCallable, Category = "Interchange | Node | Component")
+	UE_API bool GetCustomLocalTransform(FTransform& AttributeValue) const;
+
+	UFUNCTION(BlueprintCallable, Category = "Interchange | Node | Component")
+	UE_API bool SetCustomLocalTransform(const FTransform& AttributeValue);
+
+
+	UFUNCTION(BlueprintCallable, Category = "Interchange | Node | Component")
+	UE_API bool GetCustomComponentVisibility(bool& AttributeValue) const;
+
+	UFUNCTION(BlueprintCallable, Category = "Interchange | Node | Component")
+	UE_API bool SetCustomComponentVisibility(const bool& AttributeValue);
+
+
+	/** Get's the SceneNode that the SceneComponentNode belongs to, also calculates the GlobalTransform within the SceneNode space.*/
+	UFUNCTION(BlueprintCallable, Category = "Interchange | Scene | Component | Instanced")
+	UE_API const UInterchangeSceneNode* GetParentSceneNodeAndTransform(const UInterchangeBaseNodeContainer* BaseNodeContainer, FTransform& SceneNodeTransform, bool bForceRecache = false) const;
+
+private:
+	//Component's Local Transform.
+	const UE::Interchange::FAttributeKey Macro_CustomLocalTransformKey = UE::Interchange::FAttributeKey(TEXT("LocalTransform"));
+
+	//Component's Visibility
+	const UE::Interchange::FAttributeKey Macro_CustomComponentVisibilityKey = UE::Interchange::FAttributeKey(TEXT("ComponentVisibility"));
+
+	//Children Component Uids.
+	UE::Interchange::TArrayAttributeHelper<FString> ComponentUids;
+
+	//mutable caches for global transforms
+	mutable TOptional<FTransform> CacheSceneNodeTransform;
+	mutable TOptional<const UInterchangeSceneNode*> CacheParentSceneNode;
+};
+
+UCLASS(MinimalAPI, BlueprintType)
+class UInterchangeInstancedStaticMeshComponentNode : public UInterchangeSceneComponentNode
+{
+	GENERATED_BODY()
+
+public:
+	UE_API UInterchangeInstancedStaticMeshComponentNode();
+
+	/**
+	 * Return the node type name of the class. This is used when reporting errors.
+	 */
+	UE_API virtual FString GetTypeName() const override;
+
+	UFUNCTION(BlueprintCallable, Category = "Interchange | Scene | Component | Instanced")
+	UE_API void AddInstanceTransform(const FTransform& InstanceTransform);
+	UFUNCTION(BlueprintCallable, Category = "Interchange | Scene | Component | Instanced")
+	UE_API void AddInstanceTransforms(const TArray<FTransform>& InInstanceTransforms);
+
+	UFUNCTION(BlueprintCallable, Category = "Interchange | Scene | Component | Instanced")
+	UE_API void GetInstanceTransforms(TArray<FTransform>& OutInstanceTransforms) const;
+
+	/** Get which asset, if any, a scene node is instantiating. Return false if the Attribute was not set previously.*/
+	UFUNCTION(BlueprintCallable, Category = "Interchange | Scene | Component | Instanced")
+	UE_API bool GetCustomInstancedAssetUid(FString& AttributeValue) const;
+
+	/** Add an asset for this scene node to instantiate. */
+	UFUNCTION(BlueprintCallable, Category = "Interchange | Scene | Component | Instanced")
+	UE_API bool SetCustomInstancedAssetUid(const FString& AttributeValue);
+
+private:
+	//Instanced StaticMesh Component's Instances.
+	UE::Interchange::TArrayAttributeHelper<FTransform> InstanceTransforms;
+
+	//To track/connect StaticMesh.
+	const UE::Interchange::FAttributeKey Macro_CustomAssetInstanceUidKey = UE::Interchange::FAttributeKey(TEXT("AssetInstanceUid"));
+};
+
+#undef UE_API
diff --git a/Engine/Plugins/Interchange/Runtime/Source/Nodes/Public/InterchangeSceneNode.h b/Engine/Plugins/Interchange/Runtime/Source/Nodes/Public/InterchangeSceneNode.h
--- a/Engine/Plugins/Interchange/Runtime/Source/Nodes/Public/InterchangeSceneNode.h
+++ b/Engine/Plugins/Interchange/Runtime/Source/Nodes/Public/InterchangeSceneNode.h
@@ -31,6 +31,7 @@ namespace UE
 			static UE_API const FString& GetLayerNamesKey();
 			static UE_API const FString& GetTagsKey();
 			static UE_API const FString& GetCurveAnimationTypesKey();
+			static UE_API const FString& GetComponentUidsKey();
 		};
 
 	}//ns Interchange
@@ -288,6 +289,13 @@ public:
 	UFUNCTION(BlueprintCallable, Category = "Interchange | Node | Joint")
 	UE_API bool GetAnimationCurveTypeForCurveName(const FString& CurveName, EInterchangeAnimationPayLoadType& OutCurveAnimationType) const;
 
+
+	UFUNCTION(BlueprintCallable, Category = "Interchange | Scene | Component")
+	UE_API bool AddComponentUid(const FString& ComponentUid);
+
+	UFUNCTION(BlueprintCallable, Category = "Interchange | Scene | Component")
+	UE_API void GetComponentUids(TArray<FString>& OutComponentUids) const;
+
 private:
 
 	UE_API bool GetGlobalTransformInternal(const UE::Interchange::FAttributeKey LocalTransformKey
@@ -343,6 +351,9 @@ private:
 	//A scene node can have Attributes defining Curves. (Mostly used for tracking Custom Attributes' Animation Types).
 	UE::Interchange::TMapAttributeHelper<FString, EInterchangeAnimationPayLoadType> CurveAnimationTypes;
 
+	//A scene node can represent many special types
+	UE::Interchange::TArrayAttributeHelper<FString> ComponentUids;
+
 	//A scene node can reference an animation asset on top of base asset:
 	const UE::Interchange::FAttributeKey Macro_CustomAnimationAssetUidToPlayKey = UE::Interchange::FAttributeKey(TEXT("AnimationAssetUidToPlay"));
 
diff --git a/Engine/Plugins/Interchange/Runtime/Source/Pipelines/Private/InterchangeGenericGeometryCachePipeline.cpp b/Engine/Plugins/Interchange/Runtime/Source/Pipelines/Private/InterchangeGenericGeometryCachePipeline.cpp
--- a/Engine/Plugins/Interchange/Runtime/Source/Pipelines/Private/InterchangeGenericGeometryCachePipeline.cpp
+++ b/Engine/Plugins/Interchange/Runtime/Source/Pipelines/Private/InterchangeGenericGeometryCachePipeline.cpp
@@ -37,27 +37,30 @@ void UInterchangeGenericMeshPipeline::ExecutePreImportPipelineGeometryCache()
 		const FInterchangeLodSceneNodeContainer& SceneNodeContainer = MeshInstance.SceneNodePerLodIndex.Contains(LodIndex) ?
 			MeshInstance.SceneNodePerLodIndex.FindChecked(LodIndex) : EmptySceneNodeContainer;
 
-		for (const UInterchangeSceneNode* SceneNode : SceneNodeContainer.SceneNodes)
+		for (const UInterchangeBaseNode* BaseNode : SceneNodeContainer.BaseNodes)
 		{
-			if (CommonMeshesProperties->bBakeMeshes)
+			if (const UInterchangeSceneNode* SceneNode = Cast<const UInterchangeSceneNode>(BaseNode))
 			{
-				// If we're baking, always add all scene nodes: We may be trying to bake multiple instances of the same mesh node
-				// into a single big geometry cache asset
-				MeshUids.Add(SceneNode->GetUniqueID());
-			}
-			else
-			{
-				// If not baking, only add a single scene node per each mesh asset: This prevents us from ending up with multiple copies
-				// of the same mesh on top of itself (as we won't bake transforms), and it also allows us to instance geometry caches
-				// by just making several scene nodes that reference the same animated mesh node: Only one of the "instances" will end
-				// up on the geometry cache asset, and we'll reuse the asset for all the scene nodes
-				FString AssetUid;
-				SceneNode->GetCustomAssetInstanceUid(AssetUid);
-				if (!VisitedMeshAssetUids.Contains(AssetUid))
+				if (CommonMeshesProperties->bBakeMeshes)
 				{
-					VisitedMeshAssetUids.Add(AssetUid);
+					// If we're baking, always add all scene nodes: We may be trying to bake multiple instances of the same mesh node
+					// into a single big geometry cache asset
 					MeshUids.Add(SceneNode->GetUniqueID());
 				}
+				else
+				{
+					// If not baking, only add a single scene node per each mesh asset: This prevents us from ending up with multiple copies
+					// of the same mesh on top of itself (as we won't bake transforms), and it also allows us to instance geometry caches
+					// by just making several scene nodes that reference the same animated mesh node: Only one of the "instances" will end
+					// up on the geometry cache asset, and we'll reuse the asset for all the scene nodes
+					FString AssetUid;
+					SceneNode->GetCustomAssetInstanceUid(AssetUid);
+					if (!VisitedMeshAssetUids.Contains(AssetUid))
+					{
+						VisitedMeshAssetUids.Add(AssetUid);
+						MeshUids.Add(SceneNode->GetUniqueID());
+					}
+				}
 			}
 		}
 	}
diff --git a/Engine/Plugins/Interchange/Runtime/Source/Pipelines/Private/InterchangeGenericScenesPipeline.cpp b/Engine/Plugins/Interchange/Runtime/Source/Pipelines/Private/InterchangeGenericScenesPipeline.cpp
--- a/Engine/Plugins/Interchange/Runtime/Source/Pipelines/Private/InterchangeGenericScenesPipeline.cpp
+++ b/Engine/Plugins/Interchange/Runtime/Source/Pipelines/Private/InterchangeGenericScenesPipeline.cpp
@@ -23,6 +23,8 @@
 #include "InterchangePipelineHelper.h"
 #include "InterchangePipelineLog.h"
 #include "InterchangePipelineMeshesUtilities.h"
+#include "InterchangeSceneComponentFactoryNodes.h"
+#include "InterchangeSceneComponentNodes.h"
 #include "InterchangeSceneImportAsset.h"
 #include "InterchangeSceneImportAssetFactoryNode.h"
 #include "InterchangeSceneNode.h"
@@ -264,11 +266,21 @@ void UInterchangeGenericLevelPipeline::ExecutePipeline(UInterchangeBaseNodeConta
 	}
 
 	TArray<UInterchangeSceneNode*> SceneNodes;
+	TArray<UInterchangeSceneComponentNode*> SceneComponentNodes;
 	TMap<FString, const UInterchangeStaticMeshFactoryNode*> LodGroupFactoryNodeLookUp;
 	TSet<FString> SkinnedMeshNodes;
 
 	// Find all translated node we need for this pipeline
-	PopulateSceneNodesPreImport(SceneNodes, LodGroupFactoryNodeLookUp, SkinnedMeshNodes);
+	PopulateSceneNodesPreImport(SceneNodes, SceneComponentNodes, LodGroupFactoryNodeLookUp, SkinnedMeshNodes);
+
+	//If bCollapseISMSceneComponents is true then the UInterchangeInstancedStaticMeshComponentFactoryNode  will be created from ExecuteSceneNodePreImport.
+	if (!bCollapseISMSceneComponents)
+	{
+		for (UInterchangeSceneComponentNode* SceneComponentNode : SceneComponentNodes)
+		{
+			CreateComponentFactoryNode(SceneComponentNode);
+		}
+	}
 
 #if WITH_EDITORONLY_DATA
 	const FString FilePath = FPaths::ConvertRelativePathToFull(InSourceDatas[0]->GetFilename());
@@ -461,10 +473,11 @@ void UInterchangeGenericLevelPipeline::ExecutePipeline(UInterchangeBaseNodeConta
 }
 
 void UInterchangeGenericLevelPipeline::PopulateSceneNodesPreImport(TArray<UInterchangeSceneNode*>& SceneNodes,
-TMap<FString, const UInterchangeStaticMeshFactoryNode*>& LodGroupFactoryNodeLookUp,
-TSet<FString>& SkinnedMeshNodes) const
+	TArray<UInterchangeSceneComponentNode*>& SceneComponentNodes,
+	TMap<FString, const UInterchangeStaticMeshFactoryNode*>& LodGroupFactoryNodeLookUp,
+	TSet<FString>& SkinnedMeshNodes) const
 {
-	BaseNodeContainer->IterateNodes([this, &SceneNodes, &LodGroupFactoryNodeLookUp, &SkinnedMeshNodes](const FString& NodeUid, UInterchangeBaseNode* Node)
+	BaseNodeContainer->IterateNodes([this, &SceneNodes, &SceneComponentNodes, &LodGroupFactoryNodeLookUp, &SkinnedMeshNodes](const FString& NodeUid, UInterchangeBaseNode* Node)
 	{
 		switch (Node->GetNodeContainerType())
 		{
@@ -474,6 +487,10 @@ TSet<FString>& SkinnedMeshNodes) const
 			{
 				SceneNodes.Add(SceneNode);
 			}
+			else if (UInterchangeSceneComponentNode* SceneComponentNode = Cast<UInterchangeSceneComponentNode>(Node))
+			{
+				SceneComponentNodes.Add(SceneComponentNode);
+			}
 		}
 		case EInterchangeNodeContainerType::FactoryData:
 		{
@@ -541,6 +558,114 @@ TSet<FString>& SkinnedMeshNodes) const
 	});
 }
 
+void UInterchangeGenericLevelPipeline::CreateAndSetUpCollapsedISMComponents(UInterchangeActorFactoryNode* ActorFactoryNode, const UInterchangeSceneNode* SceneNode) const
+{
+	TArray<FString> ComponentUids;
+	SceneNode->GetComponentUids(ComponentUids);
+
+	TMap<FString, TArray<FTransform>> InstanceTransformsPerMeshes;
+	CollectCollapsedInstancedStaticMeshTransforms(ComponentUids, InstanceTransformsPerMeshes, FTransform::Identity);
+
+	for (const TPair<FString, TArray<FTransform>>& Entry : InstanceTransformsPerMeshes)
+	{
+		const UInterchangeBaseNode* MeshBaseNode = BaseNodeContainer->GetNode(Entry.Key);
+		
+		if (!MeshBaseNode)
+		{
+			continue;
+		}
+
+		FString DisplayLabel = MeshBaseNode->GetDisplayLabel();
+		FString UniqueId = ActorFactoryNode->GetUniqueID() + TEXT("\\CollapsedISM\\") + Entry.Key;
+		
+		UInterchangeInstancedStaticMeshComponentFactoryNode * ISMComponentFactoryNode = NewObject<UInterchangeInstancedStaticMeshComponentFactoryNode >(BaseNodeContainer, NAME_None);
+		BaseNodeContainer->SetupNode(ISMComponentFactoryNode, UniqueId, DisplayLabel, EInterchangeNodeContainerType::FactoryData, ActorFactoryNode->GetUniqueID());
+		
+		FString MeshFactoryNodeUid = UInterchangeFactoryBaseNode::BuildFactoryNodeUid(Entry.Key);
+		
+		ISMComponentFactoryNode->SetCustomInstancedAssetUid(MeshFactoryNodeUid);
+		ISMComponentFactoryNode->AddInstanceTransforms(Entry.Value);
+		ISMComponentFactoryNode->SetCustomHierarchicalISM(bUseHierarchicalISMComponents);
+
+		ActorFactoryNode->AddFactoryDependencyUid(MeshFactoryNodeUid);
+		ActorFactoryNode->AddComponentUid(UniqueId);
+	}
+}
+
+void UInterchangeGenericLevelPipeline::CollectCollapsedInstancedStaticMeshTransforms(const TArray<FString>& SceneComponentUids, TMap<FString, TArray<FTransform>>& InstanceTransformsPerMeshes, FTransform ParentGlobalTransform) const
+{
+	for (const FString& SceneComponentUid : SceneComponentUids)
+	{
+		if (const UInterchangeSceneComponentNode* SceneComponentNode = Cast<const UInterchangeSceneComponentNode>(BaseNodeContainer->GetNode(SceneComponentUid)))
+		{
+			bool bVisible = true;
+			SceneComponentNode->GetCustomComponentVisibility(bVisible);
+			if (bVisible)
+			{
+				if (const UInterchangeInstancedStaticMeshComponentNode* ISMComponentNode = Cast<const UInterchangeInstancedStaticMeshComponentNode>(SceneComponentNode))
+				{
+					TArray<FTransform> InstanceTransforms;
+					ISMComponentNode->GetInstanceTransforms(InstanceTransforms);
+					for (FTransform& Transform : InstanceTransforms)
+					{
+						Transform = Transform * ParentGlobalTransform;
+					}
+
+					FString MeshUid;
+					ISMComponentNode->GetCustomInstancedAssetUid(MeshUid);
+					TArray<FTransform>& InstanceTransformsPerMesh = InstanceTransformsPerMeshes.FindOrAdd(MeshUid);
+					InstanceTransformsPerMesh.Append(InstanceTransforms);
+				}
+			}
+			
+
+			TArray<FString> SubSceneComponentUids;
+			SceneComponentNode->GetComponentUids(SubSceneComponentUids);
+			if (SubSceneComponentUids.Num())
+			{
+				FTransform LocalTransform;
+				SceneComponentNode->GetCustomLocalTransform(LocalTransform);
+				FTransform GlobalTransform = LocalTransform * ParentGlobalTransform;
+
+				CollectCollapsedInstancedStaticMeshTransforms(SubSceneComponentUids, InstanceTransformsPerMeshes, GlobalTransform);
+			}
+		}
+	}
+}
+
+void UInterchangeGenericLevelPipeline::SetUpComponents(UInterchangeActorFactoryNode* ActorFactoryNode, const UInterchangeSceneNode* SceneNode) const
+{
+	TArray<FString> SceneComponentUids;
+	SceneNode->GetComponentUids(SceneComponentUids);
+	SetUpComponentDependencies(ActorFactoryNode, SceneComponentUids);
+	for (const FString& SceneComponentUid : SceneComponentUids)
+	{
+		FString SceneComponentFactoryUid = UInterchangeFactoryBaseNode::BuildFactoryNodeUid(SceneComponentUid);
+		ActorFactoryNode->AddComponentUid(SceneComponentFactoryUid);
+	}
+}
+
+void UInterchangeGenericLevelPipeline::SetUpComponentDependencies(UInterchangeActorFactoryNode* ActorFactoryNode, const TArray<FString>& SceneComponentUids) const
+{
+	for (const FString& SceneComponentUid : SceneComponentUids)
+	{
+		if (const UInterchangeSceneComponentNode* SceneComponentNode = Cast<const UInterchangeSceneComponentNode>(BaseNodeContainer->GetNode(SceneComponentUid)))
+		{
+			if (const UInterchangeInstancedStaticMeshComponentNode* ISMComponentNode = Cast<const UInterchangeInstancedStaticMeshComponentNode>(SceneComponentNode))
+			{
+				FString MeshNodeUid;
+				if (ISMComponentNode->GetCustomInstancedAssetUid(MeshNodeUid))
+				{
+					ActorFactoryNode->AddFactoryDependencyUid(UInterchangeFactoryBaseNode::BuildFactoryNodeUid(MeshNodeUid));
+				}
+			}
+			TArray<FString> SubSceneComponentUids;
+			SceneComponentNode->GetComponentUids(SubSceneComponentUids);
+			SetUpComponentDependencies(ActorFactoryNode, SubSceneComponentUids);
+		}
+	}
+}
+
 void UInterchangeGenericLevelPipeline::ExecuteSceneNodePreImport(const FTransform& GlobalOffsetTransform, const FSceneNodePreImportData& SceneNodeData)
 {
 	using namespace UE::Interchange;
@@ -784,6 +909,16 @@ void UInterchangeGenericLevelPipeline::ExecuteSceneNodePreImport(const FTransfor
 		SetUpFactoryNode(ActorFactoryNode, SceneNode, TranslatedAssetNode);
 	}
 
+	if (bCollapseISMSceneComponents)
+	{
+		CreateAndSetUpCollapsedISMComponents(ActorFactoryNode, SceneNode);
+	}
+	else
+	{
+		SetUpComponents(ActorFactoryNode, SceneNode);
+	}
+	
+
 	//Make sure all actor factory nodes and dependencies have the specified strategy
 	Private::UpdateReimportStrategyFlags(*BaseNodeContainer, *ActorFactoryNode, ReimportPropertyStrategy);
 	if (bForceReimportDeletedActors)
@@ -864,6 +999,60 @@ UInterchangeActorFactoryNode* UInterchangeGenericLevelPipeline::CreateActorFacto
 	return NewObject<UInterchangeActorFactoryNode>(BaseNodeContainer, NAME_None);
 }
 
+void UInterchangeGenericLevelPipeline::CreateComponentFactoryNode(const UInterchangeSceneComponentNode* SceneComponentNode) const
+{
+	auto GetParentFactoryNodeUid = [&SceneComponentNode]() -> FString
+		{
+			return SceneComponentNode->GetParentUid().IsEmpty() ? TEXT("") : UInterchangeFactoryBaseNode::BuildFactoryNodeUid(SceneComponentNode->GetParentUid());
+		};
+	UInterchangeSceneComponentFactoryNode* SceneComponentFactoryNode = nullptr;
+	const FString ComponentFactoryNodeUid = UInterchangeFactoryBaseNode::BuildFactoryNodeUid(SceneComponentNode->GetUniqueID());
+
+	if (const UInterchangeInstancedStaticMeshComponentNode* ISMComponentNode = Cast<const UInterchangeInstancedStaticMeshComponentNode>(SceneComponentNode))
+	{
+		UInterchangeInstancedStaticMeshComponentFactoryNode * ISMComponentFactoryNode = NewObject<UInterchangeInstancedStaticMeshComponentFactoryNode >(BaseNodeContainer, NAME_None);
+		BaseNodeContainer->SetupNode(ISMComponentFactoryNode, ComponentFactoryNodeUid, SceneComponentNode->GetDisplayLabel(), EInterchangeNodeContainerType::FactoryData, GetParentFactoryNodeUid());
+		UInterchangeUserDefinedAttributesAPI::DuplicateAllUserDefinedAttribute(ISMComponentNode, ISMComponentFactoryNode, false);
+		if (FString AssetInstanceUid; ISMComponentNode->GetCustomInstancedAssetUid(AssetInstanceUid))
+		{
+			ISMComponentFactoryNode->SetCustomInstancedAssetUid(UInterchangeFactoryBaseNode::BuildFactoryNodeUid(AssetInstanceUid));
+		}
+		TArray<FTransform> InstanceTransforms;
+		ISMComponentNode->GetInstanceTransforms(InstanceTransforms);
+		ISMComponentFactoryNode->AddInstanceTransforms(InstanceTransforms);
+		ISMComponentFactoryNode->SetCustomHierarchicalISM(bUseHierarchicalISMComponents);
+
+		SceneComponentFactoryNode = ISMComponentFactoryNode;
+	}
+	else
+	{
+		SceneComponentFactoryNode = NewObject<UInterchangeSceneComponentFactoryNode>(BaseNodeContainer, NAME_None);
+		BaseNodeContainer->SetupNode(SceneComponentFactoryNode, ComponentFactoryNodeUid, SceneComponentNode->GetDisplayLabel(), EInterchangeNodeContainerType::FactoryData, GetParentFactoryNodeUid());
+		UInterchangeUserDefinedAttributesAPI::DuplicateAllUserDefinedAttribute(SceneComponentNode, SceneComponentFactoryNode, false);
+	}
+	//General setups:
+	{
+		FTransform LocalTransform;
+		if (SceneComponentNode->GetCustomLocalTransform(LocalTransform))
+		{
+			SceneComponentFactoryNode->SetCustomLocalTransform(LocalTransform);
+		}
+		TArray<FString> ComponentUids;
+		SceneComponentNode->GetComponentUids(ComponentUids);
+		for (const FString& ComponentUid : ComponentUids)
+		{
+			FString ComponentFactoryUid = UInterchangeFactoryBaseNode::BuildFactoryNodeUid(ComponentUid);
+			SceneComponentFactoryNode->AddComponentUid(ComponentFactoryUid);
+		}
+		bool bComponentVisibility = true;
+		if (SceneComponentNode->GetCustomComponentVisibility(bComponentVisibility))
+		{
+			SceneComponentFactoryNode->SetCustomComponentVisibility(bComponentVisibility);
+		}
+	}
+}
+
+
 void UInterchangeGenericLevelPipeline::SetUpFactoryNode(UInterchangeActorFactoryNode* ActorFactoryNode, const UInterchangeSceneNode* SceneNode, const UInterchangeBaseNode* TranslatedAssetNode) const
 {
 	if (!ensure(BaseNodeContainer && ActorFactoryNode && SceneNode && TranslatedAssetNode))
diff --git a/Engine/Plugins/Interchange/Runtime/Source/Pipelines/Private/InterchangeGenericSkeletalMeshPipeline.cpp b/Engine/Plugins/Interchange/Runtime/Source/Pipelines/Private/InterchangeGenericSkeletalMeshPipeline.cpp
--- a/Engine/Plugins/Interchange/Runtime/Source/Pipelines/Private/InterchangeGenericSkeletalMeshPipeline.cpp
+++ b/Engine/Plugins/Interchange/Runtime/Source/Pipelines/Private/InterchangeGenericSkeletalMeshPipeline.cpp
@@ -12,6 +12,7 @@
 #include "InterchangePhysicsAssetFactoryNode.h"
 #include "InterchangePipelineLog.h"
 #include "InterchangePipelineMeshesUtilities.h"
+#include "InterchangeSceneComponentNodes.h"
 #include "InterchangeSceneNode.h"
 #include "InterchangeSkeletalMeshFactoryNode.h"
 #include "InterchangeSkeletalMeshLodDataNode.h"
@@ -138,16 +139,16 @@ void UInterchangeGenericMeshPipeline::ExecutePreImportPipelineSkeletalMesh()
 					if (MeshInstance.LodGroupNode == nullptr
 						&& MeshInstance.SceneNodePerLodIndex.Num() == 1
 						&& MeshInstance.SceneNodePerLodIndex.Contains(LodIndexZero)
-						&& MeshInstance.SceneNodePerLodIndex[LodIndexZero].SceneNodes.Num() == 1)
+						&& MeshInstance.SceneNodePerLodIndex[LodIndexZero].BaseNodes.Num() == 1)
 					{
-						const UInterchangeSceneNode* SceneNode = MeshInstance.SceneNodePerLodIndex[LodIndexZero].SceneNodes[0];
+						const UInterchangeBaseNode* BaseNode = MeshInstance.SceneNodePerLodIndex[LodIndexZero].BaseNodes[0];
 						//This mesh must be added to all existing LODs over 0
 						for (TPair<int32, TArray<FString>>& MeshUidsPerLodIndexPair : MeshUidsPerLodIndex)
 						{
 							if (MeshUidsPerLodIndexPair.Key > LodIndexZero)
 							{
 								TArray<FString>& TranslatedNodes = MeshUidsPerLodIndexPair.Value;
-								TranslatedNodes.AddUnique(SceneNode->GetUniqueID());
+								TranslatedNodes.AddUnique(BaseNode->GetUniqueID());
 							}
 						}
 					}
@@ -207,9 +208,9 @@ void UInterchangeGenericMeshPipeline::ExecutePreImportPipelineSkeletalMesh()
 								const int32 LodIndex = LodIndexAndSceneNodeContainer.Key;
 								const FInterchangeLodSceneNodeContainer& SceneNodeContainer = LodIndexAndSceneNodeContainer.Value;
 								TArray<FString>& TranslatedNodes = MeshUidsPerLodIndex.FindOrAdd(LodIndex);
-								for (const UInterchangeSceneNode* SceneNode : SceneNodeContainer.SceneNodes)
+								for (const UInterchangeBaseNode* BaseNode : SceneNodeContainer.BaseNodes)
 								{
-									TranslatedNodes.Add(SceneNode->GetUniqueID());
+									TranslatedNodes.Add(BaseNode->GetUniqueID());
 								}
 							}
 						}
@@ -285,9 +286,9 @@ void UInterchangeGenericMeshPipeline::ExecutePreImportPipelineSkeletalMesh()
 							const int32 LodIndex = LodIndexAndSceneNodeContainer.Key;
 							const FInterchangeLodSceneNodeContainer& SceneNodeContainer = LodIndexAndSceneNodeContainer.Value;
 							TArray<FString>& TranslatedNodes = MeshUidsPerLodIndex.FindOrAdd(LodIndex);
-							for (const UInterchangeSceneNode* SceneNode : SceneNodeContainer.SceneNodes)
+							for (const UInterchangeBaseNode* BaseNode : SceneNodeContainer.BaseNodes)
 							{
-								TranslatedNodes.Add(SceneNode->GetUniqueID());
+								TranslatedNodes.Add(BaseNode->GetUniqueID());
 							}
 						}
 					}
@@ -369,6 +370,17 @@ UInterchangeSkeletalMeshFactoryNode* UInterchangeGenericMeshPipeline::CreateSkel
 						return SceneNode;
 					}
 				}
+				if (const UInterchangeInstancedStaticMeshComponentNode* ISMComponentNode = Cast<const UInterchangeInstancedStaticMeshComponentNode>(BaseNodeContainer->GetNode(FirstNodeUid)))
+				{
+					FString AssetInstanceUid;
+					if (ISMComponentNode->GetCustomInstancedAssetUid(AssetInstanceUid))
+					{
+						OutSceneNodeCount = MeshUids->Num();
+						OutFirstMeshNodeUid = AssetInstanceUid;
+
+						return ISMComponentNode;
+					}
+				}
 			}
 		}
 
@@ -582,8 +594,10 @@ void UInterchangeGenericMeshPipeline::AddLodDataToSkeletalMesh(const UInterchang
 		constexpr bool bAddSourceNodeName = true;
 		for (const FString& NodeUid : NodeUids)
 		{
+			const UInterchangeBaseNode* BaseNode = BaseNodeContainer->GetNode(NodeUid);
+
 			TMap<FString, FString> SlotMaterialDependencies;
-			if (const UInterchangeSceneNode* SceneNode = Cast<UInterchangeSceneNode>(BaseNodeContainer->GetNode(NodeUid)))
+			if (const UInterchangeSceneNode* SceneNode = Cast<UInterchangeSceneNode>(BaseNode))
 			{
 				FString MeshDependency;
 				SceneNode->GetCustomAssetInstanceUid(MeshDependency);
@@ -593,7 +607,7 @@ void UInterchangeGenericMeshPipeline::AddLodDataToSkeletalMesh(const UInterchang
 					UInterchangeUserDefinedAttributesAPI::DuplicateAllUserDefinedAttribute(MeshDependencyNode, SkeletalMeshFactoryNode, bAddSourceNodeName);
 					SkeletalMeshFactoryNode->AddTargetNodeUid(MeshDependency);
 					MeshDependencyNode->AddTargetNodeUid(SkeletalMeshFactoryNode->GetUniqueID());
-					
+
 					MeshDependencyNode->GetSlotMaterialDependencies(SlotMaterialDependencies);
 				}
 				else
@@ -602,7 +616,7 @@ void UInterchangeGenericMeshPipeline::AddLodDataToSkeletalMesh(const UInterchang
 				}
 				UInterchangeUserDefinedAttributesAPI::DuplicateAllUserDefinedAttribute(SceneNode, SkeletalMeshFactoryNode, bAddSourceNodeName);
 			}
-			else if (const UInterchangeMeshNode* MeshNode = Cast<UInterchangeMeshNode>(BaseNodeContainer->GetNode(NodeUid)))
+			else if (const UInterchangeMeshNode* MeshNode = Cast<UInterchangeMeshNode>(BaseNode))
 			{
 				UInterchangeUserDefinedAttributesAPI::DuplicateAllUserDefinedAttribute(MeshNode, SkeletalMeshFactoryNode, bAddSourceNodeName);
 				SkeletalMeshFactoryNode->AddTargetNodeUid(NodeUid);
@@ -610,6 +624,19 @@ void UInterchangeGenericMeshPipeline::AddLodDataToSkeletalMesh(const UInterchang
 
 				MeshNode->GetSlotMaterialDependencies(SlotMaterialDependencies);
 			}
+			else if (const UInterchangeInstancedStaticMeshComponentNode* ISMComponentNode = Cast<UInterchangeInstancedStaticMeshComponentNode>(BaseNode))
+			{
+				FString MeshDependency;
+				ISMComponentNode->GetCustomInstancedAssetUid(MeshDependency);
+				if (const UInterchangeMeshNode* MeshDependencyNode = Cast<UInterchangeMeshNode>(BaseNodeContainer->GetNode(MeshDependency)))
+				{
+					UInterchangeUserDefinedAttributesAPI::DuplicateAllUserDefinedAttribute(MeshDependencyNode, SkeletalMeshFactoryNode, bAddSourceNodeName);
+					SkeletalMeshFactoryNode->AddTargetNodeUid(MeshDependency);
+					MeshDependencyNode->AddTargetNodeUid(SkeletalMeshFactoryNode->GetUniqueID());
+
+					MeshDependencyNode->GetSlotMaterialDependencies(SlotMaterialDependencies);
+				}
+			}
 
 			UE::Interchange::MeshesUtilities::ApplySlotMaterialDependencies(*SkeletalMeshFactoryNode, SlotMaterialDependencies, *BaseNodeContainer, &ExistingLodSlotMaterialDependencies);
 			LodDataNode->AddMeshUid(NodeUid);
diff --git a/Engine/Plugins/Interchange/Runtime/Source/Pipelines/Private/InterchangeGenericStaticMeshPipeline.cpp b/Engine/Plugins/Interchange/Runtime/Source/Pipelines/Private/InterchangeGenericStaticMeshPipeline.cpp
--- a/Engine/Plugins/Interchange/Runtime/Source/Pipelines/Private/InterchangeGenericStaticMeshPipeline.cpp
+++ b/Engine/Plugins/Interchange/Runtime/Source/Pipelines/Private/InterchangeGenericStaticMeshPipeline.cpp
@@ -4,10 +4,12 @@
 #include "AssetRegistry/AssetRegistryModule.h"
 #include "Async/Async.h"
 #include "CoreMinimal.h"
+#include "InterchangeMeshLODContainerNode.h"
 #include "InterchangeMaterialFactoryNode.h"
 #include "InterchangeMeshNode.h"
 #include "InterchangePipelineLog.h"
 #include "InterchangePipelineMeshesUtilities.h"
+#include "InterchangeSceneComponentNodes.h"
 #include "InterchangeSceneNode.h"
 #include "InterchangeStaticMeshFactoryNode.h"
 #include "InterchangeStaticMeshLodDataNode.h"
@@ -284,9 +286,9 @@ void UInterchangeGenericMeshPipeline::ExecutePreImportPipelineStaticMesh()
 						const FInterchangeLodSceneNodeContainer& SceneNodeContainer = LodIndexAndSceneNodeContainer.Value;
 
 						TArray<FString>& TranslatedNodes = MeshUidsPerLodIndex.FindOrAdd(LodIndex);
-						for (const UInterchangeSceneNode* SceneNode : SceneNodeContainer.SceneNodes)
+						for (const UInterchangeBaseNode* BaseNode : SceneNodeContainer.BaseNodes)
 						{
-							TranslatedNodes.Add(SceneNode->GetUniqueID());
+							TranslatedNodes.Add(BaseNode->GetUniqueID());
 						}
 					}
 				}
@@ -301,16 +303,16 @@ void UInterchangeGenericMeshPipeline::ExecutePreImportPipelineStaticMesh()
 						if (MeshInstance.LodGroupNode == nullptr
 							&& MeshInstance.SceneNodePerLodIndex.Num() == 1
 							&& MeshInstance.SceneNodePerLodIndex.Contains(LodIndexZero)
-							&& MeshInstance.SceneNodePerLodIndex[LodIndexZero].SceneNodes.Num() == 1)
+							&& MeshInstance.SceneNodePerLodIndex[LodIndexZero].BaseNodes.Num() == 1)
 						{
-							const UInterchangeSceneNode* SceneNode = MeshInstance.SceneNodePerLodIndex[LodIndexZero].SceneNodes[0];
+							const UInterchangeBaseNode* BaseNode = MeshInstance.SceneNodePerLodIndex[LodIndexZero].BaseNodes[0];
 							//This mesh must be added to all existing LODs over 0
 							for (TPair<int32, TArray<FString>>& MeshUidsPerLodIndexPair : MeshUidsPerLodIndex)
 							{
 								if (MeshUidsPerLodIndexPair.Key > LodIndexZero)
 								{
 									TArray<FString>& TranslatedNodes = MeshUidsPerLodIndexPair.Value;
-									TranslatedNodes.AddUnique(SceneNode->GetUniqueID());
+									TranslatedNodes.AddUnique(BaseNode->GetUniqueID());
 								}
 							}
 						}
@@ -384,9 +386,9 @@ void UInterchangeGenericMeshPipeline::ExecutePreImportPipelineStaticMesh()
 								const FInterchangeLodSceneNodeContainer& SceneNodeContainer = LodIndexAndSceneNodeContainer.Value;
 
 								TArray<FString>& TranslatedNodes = MeshUidsPerLodIndex.FindOrAdd(LodIndex);
-								for (const UInterchangeSceneNode* SceneNode : SceneNodeContainer.SceneNodes)
+								for (const UInterchangeBaseNode* BaseNode : SceneNodeContainer.BaseNodes)
 								{
-									TranslatedNodes.Add(SceneNode->GetUniqueID());
+									TranslatedNodes.Add(BaseNode->GetUniqueID());
 								}
 							}
 						}
@@ -414,9 +416,11 @@ void UInterchangeGenericMeshPipeline::ExecutePreImportPipelineStaticMesh()
 								}
 							}
 
-							UInterchangeStaticMeshFactoryNode* StaticMeshFactoryNode = CreateStaticMeshFactoryNode(MeshUidsPerLodIndex);
-							StaticMeshFactoryNodes.Add(StaticMeshFactoryNode);
-							bFoundMeshes = true;
+							if (UInterchangeStaticMeshFactoryNode* StaticMeshFactoryNode = CreateStaticMeshFactoryNode(MeshUidsPerLodIndex, MeshUid))
+							{
+								StaticMeshFactoryNodes.Add(StaticMeshFactoryNode);
+								bFoundMeshes = true;
+							}
 						}
 					}
 					return bFoundMeshes;
@@ -434,7 +438,7 @@ void UInterchangeGenericMeshPipeline::ExecutePreImportPipelineStaticMesh()
 	}
 }
 
-bool UInterchangeGenericMeshPipeline::MakeMeshFactoryNodeUidAndDisplayLabel(const TMap<int32, TArray<FString>>& MeshUidsPerLodIndex, int32 LodIndex, FString& NewNodeUid, FString& DisplayLabel)
+bool UInterchangeGenericMeshPipeline::MakeMeshFactoryNodeUidAndDisplayLabel(const TMap<int32, TArray<FString>>& MeshUidsPerLodIndex, int32 LodIndex, FString& NewNodeUid, FString& DisplayLabel, const FString& BaseMeshUid)
 {
 	int32 SceneNodeCount = 0;
 
@@ -443,6 +447,17 @@ bool UInterchangeGenericMeshPipeline::MakeMeshFactoryNodeUidAndDisplayLabel(cons
 		return false;
 	}
 
+	if (!BaseMeshUid.IsEmpty())
+	{
+		if (const UInterchangeMeshLODContainerNode* LODContainerNode = Cast<const UInterchangeMeshLODContainerNode>(BaseNodeContainer->GetNode(BaseMeshUid)))
+		{
+			NewNodeUid = LODContainerNode->GetUniqueID();
+			DisplayLabel = LODContainerNode->GetDisplayLabel();
+
+			return true;
+		}
+	}
+
 	for (const TPair<int32, TArray<FString>>& LodIndexAndMeshUids : MeshUidsPerLodIndex)
 	{
 		if (LodIndex == LodIndexAndMeshUids.Key)
@@ -480,6 +495,21 @@ bool UInterchangeGenericMeshPipeline::MakeMeshFactoryNodeUidAndDisplayLabel(cons
 						}
 					}
 				}
+
+				if (const UInterchangeInstancedStaticMeshComponentNode* ISMComponentNode = Cast<const UInterchangeInstancedStaticMeshComponentNode>(Node))
+				{
+					FString RefMeshUid;
+					if (ISMComponentNode->GetCustomInstancedAssetUid(RefMeshUid))
+					{
+						const UInterchangeBaseNode* MeshNode = BaseNodeContainer->GetNode(RefMeshUid);
+						if (MeshNode)
+						{
+							DisplayLabel = MeshNode->GetDisplayLabel();
+							NewNodeUid = RefMeshUid;
+							return true;
+						}
+					}
+				}
 			}
 
 			// We found the lod but there is no valid Mesh node to return the Uid
@@ -490,7 +520,7 @@ bool UInterchangeGenericMeshPipeline::MakeMeshFactoryNodeUidAndDisplayLabel(cons
 	return false;
 }
 
-UInterchangeStaticMeshFactoryNode* UInterchangeGenericMeshPipeline::CreateStaticMeshFactoryNode(const TMap<int32, TArray<FString>>& MeshUidsPerLodIndex)
+UInterchangeStaticMeshFactoryNode* UInterchangeGenericMeshPipeline::CreateStaticMeshFactoryNode(const TMap<int32, TArray<FString>>& MeshUidsPerLodIndex, const FString& BaseMeshUid)
 {
 	check(CommonMeshesProperties.IsValid());
 	if (MeshUidsPerLodIndex.Num() == 0)
@@ -502,7 +532,7 @@ UInterchangeStaticMeshFactoryNode* UInterchangeGenericMeshPipeline::CreateStatic
 	FString StaticMeshUid_MeshNamePart;
 	FString DisplayLabel;
 	const int32 BaseLodIndex = 0;
-	if (!MakeMeshFactoryNodeUidAndDisplayLabel(MeshUidsPerLodIndex, BaseLodIndex, StaticMeshUid_MeshNamePart, DisplayLabel))
+	if (!MakeMeshFactoryNodeUidAndDisplayLabel(MeshUidsPerLodIndex, BaseLodIndex, StaticMeshUid_MeshNamePart, DisplayLabel, BaseMeshUid))
 	{
 		// Log an error
 		return nullptr;
@@ -671,6 +701,23 @@ void UInterchangeGenericMeshPipeline::AddLodDataToStaticMesh(UInterchangeStaticM
 
 				MeshNode->GetSlotMaterialDependencies(SlotMaterialDependencies);
 			}
+			else if (const UInterchangeInstancedStaticMeshComponentNode* ISMComponentNode = Cast<UInterchangeInstancedStaticMeshComponentNode>(BaseNodeContainer->GetNode(NodeUid)))
+			{
+				FString MeshDependency;
+				ISMComponentNode->GetCustomInstancedAssetUid(MeshDependency);
+				if (BaseNodeContainer->IsNodeUidValid(MeshDependency))
+				{
+					const UInterchangeMeshNode* MeshDependencyNode = Cast<UInterchangeMeshNode>(BaseNodeContainer->GetNode(MeshDependency));
+					UInterchangeUserDefinedAttributesAPI::DuplicateAllUserDefinedAttribute(MeshDependencyNode, StaticMeshFactoryNode, bAddSourceNodeName);
+					StaticMeshFactoryNode->AddTargetNodeUid(MeshDependency);
+					StaticMeshFactoryNode->AddSocketUids(PipelineMeshesUtilities->GetMeshGeometryByUid(MeshDependency).AttachedSocketUids);
+					MeshDependencyNode->AddTargetNodeUid(StaticMeshFactoryNode->GetUniqueID());
+
+					MeshDependencyNode->GetSlotMaterialDependencies(SlotMaterialDependencies);
+				}
+
+				UInterchangeUserDefinedAttributesAPI::DuplicateAllUserDefinedAttribute(ISMComponentNode, StaticMeshFactoryNode, bAddSourceNodeName);
+			}
 
 			UE::Interchange::MeshesUtilities::ApplySlotMaterialDependencies(*StaticMeshFactoryNode, SlotMaterialDependencies, *BaseNodeContainer, &ExistingLodSlotMaterialDependencies);
 
diff --git a/Engine/Plugins/Interchange/Runtime/Source/Pipelines/Private/InterchangePipelineMeshesUtilities.cpp b/Engine/Plugins/Interchange/Runtime/Source/Pipelines/Private/InterchangePipelineMeshesUtilities.cpp
--- a/Engine/Plugins/Interchange/Runtime/Source/Pipelines/Private/InterchangePipelineMeshesUtilities.cpp
+++ b/Engine/Plugins/Interchange/Runtime/Source/Pipelines/Private/InterchangePipelineMeshesUtilities.cpp
@@ -5,7 +5,9 @@
 #include "CoreMinimal.h"
 #include "Nodes/InterchangeBaseNode.h"
 #include "InterchangeMeshFactoryNode.h"
+#include "InterchangeMeshLODContainerNode.h"
 #include "InterchangeMeshNode.h"
+#include "InterchangeSceneComponentNodes.h"
 #include "InterchangeSceneNode.h"
 #include "Nodes/InterchangeBaseNodeContainer.h"
 #include "UObject/Object.h"
@@ -215,12 +217,15 @@ bool FInterchangePipelineMeshesUtilitiesContext::IsSkeletalMeshInstance(const FI
 		if (MeshInstance.SceneNodePerLodIndex.Contains(0))
 		{
 			const FInterchangeLodSceneNodeContainer& LodSceneNodeContainer = MeshInstance.SceneNodePerLodIndex.FindChecked(0);
-			for (const TObjectPtr<const UInterchangeSceneNode>& SceneNode : LodSceneNodeContainer.SceneNodes)
+			for (const TObjectPtr<const UInterchangeBaseNode>& BaseNode : LodSceneNodeContainer.BaseNodes)
 			{
-				if (UE::Private::InterchangeMeshPipeline::IsSceneNodeNestedInSkeleton(BaseNodeContainer, SceneNode))
+				if (const UInterchangeSceneNode* SceneNode = Cast<const UInterchangeSceneNode>(BaseNode))
 				{
-					bOutIsStaticMeshNestedInSkeleton = true;
-					return true;
+					if (UE::Private::InterchangeMeshPipeline::IsSceneNodeNestedInSkeleton(BaseNodeContainer, SceneNode))
+					{
+						bOutIsStaticMeshNestedInSkeleton = true;
+						return true;
+					}
 				}
 			}
 		}
@@ -325,129 +330,196 @@ UInterchangePipelineMeshesUtilities* UInterchangePipelineMeshesUtilities::Create
 	TMap<FString, TArray<FString>> SceneMeshNodeUidsPerLodParentUidMap;
 	UE::Private::InterchangeMeshPipeline::FindNamedLodGroup(BaseNodeContainer, SceneMeshNodeUidsPerLodParentUidMap);
 
-	//Find all translated scene node we need for this pipeline
-	bool bHasSockets = false;
-	BaseNodeContainer->IterateNodes(
-		[&PipelineMeshesUtilities, &BaseNodeContainer, &SceneMeshNodeUidsPerLodParentUidMap, &SkeletonRootNodeUids, &bHasSockets](const FString& NodeUid, const UInterchangeBaseNode* Node)
+	//returns false if NodeUid is not a Mesh.
+	auto ProcessMeshNode = [&PipelineMeshesUtilities, &BaseNodeContainer, &SceneMeshNodeUidsPerLodParentUidMap](const FString& NodeUid, const UInterchangeBaseNode* BaseNode, const FString& MeshUid) -> bool
 		{
-			if (Node->GetNodeContainerType() == EInterchangeNodeContainerType::TranslatedScene)
+			const UInterchangeBaseNode* MeshNode = BaseNodeContainer->GetNode(MeshUid);
+
+			if (MeshNode && MeshNode->IsA<UInterchangeMeshNode>())
 			{
-				if (const UInterchangeSceneNode* SceneNode = Cast<UInterchangeSceneNode>(Node))
+				FInterchangeMeshGeometry& MeshGeometry = PipelineMeshesUtilities->MeshGeometriesPerMeshUid.FindChecked(MeshUid);
+
+				const UInterchangeBaseNode* BaseParentNode = BaseNodeContainer->GetNode(BaseNode->GetParentUid());
+				const UInterchangeBaseNode* LodGroupNode = nullptr;
+				int32 LodIndex = 0;
+				if (BaseParentNode)
 				{
-					if (SceneNode->IsSpecializedTypeContains(UE::Interchange::FSceneNodeStaticData::GetJointSpecializeTypeString()))
+					if (SceneMeshNodeUidsPerLodParentUidMap.Contains(BaseParentNode->GetUniqueID()))
 					{
-						const UInterchangeSceneNode* ParentJointNode = Cast<UInterchangeSceneNode>(BaseNodeContainer->GetNode(SceneNode->GetParentUid()));
-						if (!ParentJointNode || !ParentJointNode->IsSpecializedTypeContains(UE::Interchange::FSceneNodeStaticData::GetJointSpecializeTypeString()))
+						LodGroupNode = BaseParentNode;
+						const TArray<FString>& LodChildUids = SceneMeshNodeUidsPerLodParentUidMap.FindChecked(BaseParentNode->GetUniqueID());
+						for (int32 ChildLodIndex = 0; ChildLodIndex < LodChildUids.Num(); ++ChildLodIndex)
 						{
-							SkeletonRootNodeUids.Add(SceneNode->GetUniqueID());
+							const FString& ChildrenUid = LodChildUids[ChildLodIndex];
+							if (ChildrenUid.Equals(BaseNode->GetUniqueID()))
+							{
+								LodIndex = ChildLodIndex;
+								break;
+							}
 						}
 					}
-
-					if (IsSceneNodeASocket(SceneNode))
-					{
-						bHasSockets = true;
-					}
-
-					FString MeshUid;
-					if (SceneNode->GetCustomAssetInstanceUid(MeshUid))
+					else
 					{
-						const UInterchangeBaseNode* MeshNode = BaseNodeContainer->GetNode(MeshUid);
-
-						if (MeshNode && MeshNode->IsA<UInterchangeMeshNode>())
+						if (const UInterchangeSceneNode* ParentMeshSceneNode = Cast<UInterchangeSceneNode>(BaseParentNode))
 						{
-							FInterchangeMeshGeometry& MeshGeometry = PipelineMeshesUtilities->MeshGeometriesPerMeshUid.FindChecked(MeshUid);
-
-							const UInterchangeSceneNode* ParentMeshSceneNode = Cast<UInterchangeSceneNode>(BaseNodeContainer->GetNode(SceneNode->GetParentUid()));
-							const UInterchangeSceneNode* LodGroupNode = nullptr;
-							int32 LodIndex = 0;
-							if (ParentMeshSceneNode)
+							FString LastChildUid = BaseNode->GetUniqueID();
+							do
 							{
-								if (SceneMeshNodeUidsPerLodParentUidMap.Contains(ParentMeshSceneNode->GetUniqueID()))
+								if (ParentMeshSceneNode->IsSpecializedTypeContains(UE::Interchange::FSceneNodeStaticData::GetLodGroupSpecializeTypeString()))
 								{
 									LodGroupNode = ParentMeshSceneNode;
-									const TArray<FString>& LodChildUids = SceneMeshNodeUidsPerLodParentUidMap.FindChecked(ParentMeshSceneNode->GetUniqueID());
-									for (int32 ChildLodIndex = 0; ChildLodIndex < LodChildUids.Num(); ++ChildLodIndex)
+									TArray<FString> LodGroupChildrens = BaseNodeContainer->GetNodeChildrenUids(ParentMeshSceneNode->GetUniqueID());
+									for (int32 ChildLodIndex = 0; ChildLodIndex < LodGroupChildrens.Num(); ++ChildLodIndex)
 									{
-										const FString& ChildrenUid = LodChildUids[ChildLodIndex];
-										if (ChildrenUid.Equals(SceneNode->GetUniqueID()))
+										const FString& ChildrenUid = LodGroupChildrens[ChildLodIndex];
+										if (ChildrenUid.Equals(LastChildUid))
 										{
 											LodIndex = ChildLodIndex;
 											break;
 										}
 									}
+									break;
 								}
-								else
+								LastChildUid = ParentMeshSceneNode->GetUniqueID();
+								ParentMeshSceneNode = Cast<const UInterchangeSceneNode>(BaseNodeContainer->GetNode(ParentMeshSceneNode->GetParentUid()));
+							} while (ParentMeshSceneNode);
+						}
+					}
+				}
+
+				if (MeshGeometry.MeshNode->IsA<UInterchangeGeometryCacheNode>())
+				{
+					FInterchangeMeshInstance& MeshInstance = PipelineMeshesUtilities->MeshInstancesPerMeshInstanceUid.FindOrAdd(NodeUid);
+					MeshInstance.LodGroupNode = nullptr;
+					MeshInstance.MeshInstanceUid = NodeUid;
+					FInterchangeLodSceneNodeContainer& InstancedSceneNodes = MeshInstance.SceneNodePerLodIndex.FindOrAdd(LodIndex);
+					InstancedSceneNodes.BaseNodes.AddUnique(BaseNode);
+					MeshGeometry.ReferencingMeshInstanceUids.Add(MeshInstance.MeshInstanceUid);
+					MeshInstance.ReferencingMeshGeometryUids.Add(MeshUid);
+					MeshInstance.bReferenceSkinnedMesh = false;
+					MeshInstance.bReferenceMorphTarget = false;
+					MeshInstance.bHasMorphTargets = false;
+					MeshInstance.bIsAnimated = true;
+				}
+				else if (LodGroupNode)
+				{
+					//We have a LOD
+					FInterchangeMeshInstance& MeshInstance = PipelineMeshesUtilities->MeshInstancesPerMeshInstanceUid.FindOrAdd(LodGroupNode->GetUniqueID());
+					if (MeshInstance.LodGroupNode != nullptr)
+					{
+						//This LodGroup was already created, verify everything is ok
+						checkSlow(MeshInstance.LodGroupNode == LodGroupNode);
+						checkSlow(MeshInstance.MeshInstanceUid.Equals(LodGroupNode->GetUniqueID()));
+					}
+					else
+					{
+						MeshInstance.LodGroupNode = LodGroupNode;
+						MeshInstance.MeshInstanceUid = LodGroupNode->GetUniqueID();
+					}
+					FInterchangeLodSceneNodeContainer& InstancedSceneNodes = MeshInstance.SceneNodePerLodIndex.FindOrAdd(LodIndex);
+					InstancedSceneNodes.BaseNodes.AddUnique(BaseNode);
+					MeshGeometry.ReferencingMeshInstanceUids.Add(MeshInstance.MeshInstanceUid);
+					MeshInstance.ReferencingMeshGeometryUids.Add(MeshUid);
+					MeshInstance.bReferenceSkinnedMesh |= MeshGeometry.MeshNode->IsSkinnedMesh();
+					MeshInstance.bReferenceMorphTarget |= MeshGeometry.MeshNode->IsMorphTarget();
+					MeshInstance.bHasMorphTargets |= MeshGeometry.MeshNode->GetMorphTargetDependeciesCount() > 0;
+				}
+				else
+				{
+					FInterchangeMeshInstance& MeshInstance = PipelineMeshesUtilities->MeshInstancesPerMeshInstanceUid.FindOrAdd(NodeUid);
+					MeshInstance.LodGroupNode = nullptr;
+					MeshInstance.MeshInstanceUid = NodeUid;
+					FInterchangeLodSceneNodeContainer& InstancedSceneNodes = MeshInstance.SceneNodePerLodIndex.FindOrAdd(LodIndex);
+					InstancedSceneNodes.BaseNodes.AddUnique(BaseNode);
+					MeshGeometry.ReferencingMeshInstanceUids.Add(MeshInstance.MeshInstanceUid);
+					MeshInstance.ReferencingMeshGeometryUids.Add(MeshUid);
+					MeshInstance.bReferenceSkinnedMesh |= MeshGeometry.MeshNode->IsSkinnedMesh();
+					MeshInstance.bReferenceMorphTarget |= MeshGeometry.MeshNode->IsMorphTarget();
+					MeshInstance.bHasMorphTargets |= MeshGeometry.MeshNode->GetMorphTargetDependeciesCount() > 0;
+				}
+
+				return true;
+			}
+
+			return false;
+		};
+
+
+	BaseNodeContainer->BreakableIterateNodesOfType<UInterchangeInstancedStaticMeshComponentNode>([&PipelineMeshesUtilities, &BaseNodeContainer, &ProcessMeshNode](const FString& ISMComponentNodeUid, UInterchangeInstancedStaticMeshComponentNode* ISMComponentNode)
+		{
+			FString AssetInstanceUid;
+			if (ISMComponentNode->GetCustomInstancedAssetUid(AssetInstanceUid))
+			{
+				if (!ProcessMeshNode(ISMComponentNodeUid, ISMComponentNode, AssetInstanceUid))
+				{
+					//if not mesh node it's likely LODContainerNode:
+					if (const UInterchangeMeshLODContainerNode* LODContainerNode = Cast<const UInterchangeMeshLODContainerNode>(BaseNodeContainer->GetNode(AssetInstanceUid)))
+					{
+						const FString LodPrefix = TEXT("LOD");
+
+						TArray<FString> LODMeshUids;
+						LODContainerNode->GetMeshLODNodeUids(LODMeshUids);
+
+						TArray<const UInterchangeBaseNode*> OrderedLODs;
+						OrderedLODs.Reserve(LODMeshUids.Num());
+
+						for (const FString& LODMeshUid : LODMeshUids)
+						{
+							if (const UInterchangeBaseNode* LODMeshNode = Cast<const UInterchangeBaseNode>(BaseNodeContainer->GetNode(LODMeshUid)))
+							{
+								FString LODMeshNodeName = LODMeshNode->GetDisplayLabel();
+								if (LODMeshNodeName.StartsWith(LodPrefix, ESearchCase::CaseSensitive))
 								{
-									FString LastChildUid = SceneNode->GetUniqueID();
-									do
+									FString LODXNumber = LODMeshNodeName.RightChop(3);
+									if (LODXNumber.IsNumeric())
 									{
-										if (ParentMeshSceneNode->IsSpecializedTypeContains(UE::Interchange::FSceneNodeStaticData::GetLodGroupSpecializeTypeString()))
+										int32 LodNumber = FPlatformString::Atoi(*LODXNumber);
+
+										//Add LOD at the correct index
+										if (LodNumber >= OrderedLODs.Num())
 										{
-											LodGroupNode = ParentMeshSceneNode;
-											TArray<FString> LodGroupChildrens = BaseNodeContainer->GetNodeChildrenUids(ParentMeshSceneNode->GetUniqueID());
-											for (int32 ChildLodIndex = 0; ChildLodIndex < LodGroupChildrens.Num(); ++ChildLodIndex)
-											{
-												const FString& ChildrenUid = LodGroupChildrens[ChildLodIndex];
-												if (ChildrenUid.Equals(LastChildUid))
-												{
-													LodIndex = ChildLodIndex;
-													break;
-												}
-											}
-											break;
+											int32 AddCount = LodNumber + 1 - OrderedLODs.Num();
+											OrderedLODs.AddDefaulted(AddCount);
 										}
-										LastChildUid = ParentMeshSceneNode->GetUniqueID();
-										ParentMeshSceneNode = Cast<const UInterchangeSceneNode>(BaseNodeContainer->GetNode(ParentMeshSceneNode->GetParentUid()));
-									} while (ParentMeshSceneNode);
+										OrderedLODs[LodNumber] = LODMeshNode;
+									}
 								}
 							}
+						}
 
-							if (MeshGeometry.MeshNode->IsA<UInterchangeGeometryCacheNode>())
+						for (int32 ChildLodIndex = OrderedLODs.Num() - 1; ChildLodIndex >= 0; ChildLodIndex--)
+						{
+							if (OrderedLODs[ChildLodIndex] == nullptr)
 							{
-								FInterchangeMeshInstance& MeshInstance = PipelineMeshesUtilities->MeshInstancesPerMeshInstanceUid.FindOrAdd(NodeUid);
-								MeshInstance.LodGroupNode = nullptr;
-								MeshInstance.MeshInstanceUid = NodeUid;
-								FInterchangeLodSceneNodeContainer& InstancedSceneNodes = MeshInstance.SceneNodePerLodIndex.FindOrAdd(LodIndex);
-								InstancedSceneNodes.SceneNodes.AddUnique(SceneNode);
-								MeshGeometry.ReferencingMeshInstanceUids.Add(MeshInstance.MeshInstanceUid);
-								MeshInstance.ReferencingMeshGeometryUids.Add(MeshUid);
-								MeshInstance.bReferenceSkinnedMesh = false;
-								MeshInstance.bReferenceMorphTarget = false;
-								MeshInstance.bHasMorphTargets = false;
-								MeshInstance.bIsAnimated = true;
+								OrderedLODs.RemoveAt(ChildLodIndex, EAllowShrinking::No);
 							}
-							else if (LodGroupNode)
+						}
+						//Shrink the array to the correct size
+						OrderedLODs.Shrink();
+
+						if (OrderedLODs.Num() > 0)
+						{
+
+							FInterchangeMeshInstance& MeshInstance = PipelineMeshesUtilities->MeshInstancesPerMeshInstanceUid.FindOrAdd(AssetInstanceUid);
+							if (MeshInstance.LodGroupNode != nullptr)
 							{
-								//We have a LOD
-								FInterchangeMeshInstance& MeshInstance = PipelineMeshesUtilities->MeshInstancesPerMeshInstanceUid.FindOrAdd(LodGroupNode->GetUniqueID());
-								if (MeshInstance.LodGroupNode != nullptr)
-								{
-									//This LodGroup was already created, verify everything is ok
-									checkSlow(MeshInstance.LodGroupNode == LodGroupNode);
-									checkSlow(MeshInstance.MeshInstanceUid.Equals(LodGroupNode->GetUniqueID()));
-								}
-								else
-								{
-									MeshInstance.LodGroupNode = LodGroupNode;
-									MeshInstance.MeshInstanceUid = LodGroupNode->GetUniqueID();
-								}
-								FInterchangeLodSceneNodeContainer& InstancedSceneNodes = MeshInstance.SceneNodePerLodIndex.FindOrAdd(LodIndex);
-								InstancedSceneNodes.SceneNodes.AddUnique(SceneNode);
-								MeshGeometry.ReferencingMeshInstanceUids.Add(MeshInstance.MeshInstanceUid);
-								MeshInstance.ReferencingMeshGeometryUids.Add(MeshUid);
-								MeshInstance.bReferenceSkinnedMesh |= MeshGeometry.MeshNode->IsSkinnedMesh();
-								MeshInstance.bReferenceMorphTarget |= MeshGeometry.MeshNode->IsMorphTarget();
-								MeshInstance.bHasMorphTargets |= MeshGeometry.MeshNode->GetMorphTargetDependeciesCount() > 0;
+								//This LodGroup was already created, verify everything is ok
+								checkSlow(MeshInstance.LodGroupNode == LODContainerNode);
+								checkSlow(MeshInstance.MeshInstanceUid.Equals(LODContainerNode->GetUniqueID()));
 							}
 							else
 							{
-								FInterchangeMeshInstance& MeshInstance = PipelineMeshesUtilities->MeshInstancesPerMeshInstanceUid.FindOrAdd(NodeUid);
-								MeshInstance.LodGroupNode = nullptr;
-								MeshInstance.MeshInstanceUid = NodeUid;
-								FInterchangeLodSceneNodeContainer& InstancedSceneNodes = MeshInstance.SceneNodePerLodIndex.FindOrAdd(LodIndex);
-								InstancedSceneNodes.SceneNodes.AddUnique(SceneNode);
+								MeshInstance.LodGroupNode = LODContainerNode;
+								MeshInstance.MeshInstanceUid = LODContainerNode->GetUniqueID();
+							}
+
+							for (int32 LODIndex = 0; LODIndex < OrderedLODs.Num(); LODIndex++)
+							{
+								FInterchangeLodSceneNodeContainer& InstancedSceneNodes = MeshInstance.SceneNodePerLodIndex.FindOrAdd(LODIndex);
+								InstancedSceneNodes.BaseNodes.AddUnique(OrderedLODs[LODIndex]);
+								FInterchangeMeshGeometry& MeshGeometry = PipelineMeshesUtilities->MeshGeometriesPerMeshUid.FindChecked(OrderedLODs[LODIndex]->GetUniqueID());
 								MeshGeometry.ReferencingMeshInstanceUids.Add(MeshInstance.MeshInstanceUid);
-								MeshInstance.ReferencingMeshGeometryUids.Add(MeshUid);
+								MeshInstance.ReferencingMeshGeometryUids.Add(OrderedLODs[LODIndex]->GetUniqueID());
 								MeshInstance.bReferenceSkinnedMesh |= MeshGeometry.MeshNode->IsSkinnedMesh();
 								MeshInstance.bReferenceMorphTarget |= MeshGeometry.MeshNode->IsMorphTarget();
 								MeshInstance.bHasMorphTargets |= MeshGeometry.MeshNode->GetMorphTargetDependeciesCount() > 0;
@@ -456,6 +528,40 @@ UInterchangePipelineMeshesUtilities* UInterchangePipelineMeshesUtilities::Create
 					}
 				}
 			}
+
+			return false;
+		});
+
+	//Find all translated scene node we need for this pipeline
+	bool bHasSockets = false;
+	BaseNodeContainer->IterateNodes(
+		[&PipelineMeshesUtilities, &BaseNodeContainer, &SceneMeshNodeUidsPerLodParentUidMap, &SkeletonRootNodeUids, &bHasSockets, &ProcessMeshNode](const FString& NodeUid, const UInterchangeBaseNode* Node)
+		{
+			if (Node->GetNodeContainerType() == EInterchangeNodeContainerType::TranslatedScene)
+			{
+				if (const UInterchangeSceneNode* SceneNode = Cast<UInterchangeSceneNode>(Node))
+				{
+					if (SceneNode->IsSpecializedTypeContains(UE::Interchange::FSceneNodeStaticData::GetJointSpecializeTypeString()))
+					{
+						const UInterchangeSceneNode* ParentJointNode = Cast<UInterchangeSceneNode>(BaseNodeContainer->GetNode(SceneNode->GetParentUid()));
+						if (!ParentJointNode || !ParentJointNode->IsSpecializedTypeContains(UE::Interchange::FSceneNodeStaticData::GetJointSpecializeTypeString()))
+						{
+							SkeletonRootNodeUids.Add(SceneNode->GetUniqueID());
+						}
+					}
+
+					if (IsSceneNodeASocket(SceneNode))
+					{
+						bHasSockets = true;
+					}
+
+					FString MeshUid;
+					if (SceneNode->GetCustomAssetInstanceUid(MeshUid))
+					{
+						ProcessMeshNode(NodeUid, SceneNode, MeshUid);
+					}
+				}
+			}
 		}
 	);
 
@@ -896,15 +1002,18 @@ FString UInterchangePipelineMeshesUtilities::GetMeshInstanceSkeletonRootUid(cons
 		return SkeletonRootUid;
 	}
 	const int32 BaseLodIndex = 0;
-	if (MeshInstance.SceneNodePerLodIndex[BaseLodIndex].SceneNodes.Num() > 0)
+	if (MeshInstance.SceneNodePerLodIndex[BaseLodIndex].BaseNodes.Num() > 0)
 	{
-		const UInterchangeSceneNode* SceneNode = MeshInstance.SceneNodePerLodIndex[BaseLodIndex].SceneNodes[0];
-		FString MeshNodeUid;
-		if (SceneNode->GetCustomAssetInstanceUid(MeshNodeUid))
+		const UInterchangeBaseNode* BaseNode = MeshInstance.SceneNodePerLodIndex[BaseLodIndex].BaseNodes[0];
+		if (const UInterchangeSceneNode* SceneNode = Cast<const UInterchangeSceneNode>(BaseNode))
 		{
-			if (SkeletonRootUidPerMeshUid.Contains(MeshNodeUid))
+			FString MeshNodeUid;
+			if (SceneNode->GetCustomAssetInstanceUid(MeshNodeUid))
 			{
-				SkeletonRootUid = SkeletonRootUidPerMeshUid.FindChecked(MeshNodeUid);
+				if (SkeletonRootUidPerMeshUid.Contains(MeshNodeUid))
+				{
+					SkeletonRootUid = SkeletonRootUidPerMeshUid.FindChecked(MeshNodeUid);
+				}
 			}
 		}
 	}
diff --git a/Engine/Plugins/Interchange/Runtime/Source/Pipelines/Public/InterchangeGenericMeshPipeline.h b/Engine/Plugins/Interchange/Runtime/Source/Pipelines/Public/InterchangeGenericMeshPipeline.h
--- a/Engine/Plugins/Interchange/Runtime/Source/Pipelines/Public/InterchangeGenericMeshPipeline.h
+++ b/Engine/Plugins/Interchange/Runtime/Source/Pipelines/Public/InterchangeGenericMeshPipeline.h
@@ -407,8 +407,9 @@ private:
 	/**
 	 * This function can create a UInterchangeStaticMeshFactoryNode
 	 * @param MeshUidsPerLodIndex - The MeshUids can represent a SceneNode pointing on a MeshNode or directly a MeshNode
+	 * @param BaseMeshUid - in case of multiple MeshUidsPerLodIndex generated from a LOD Container, the BaseMeshUid will point to the LOD Container for naming and UID purposes.
 	 */
-	UE_API UInterchangeStaticMeshFactoryNode* CreateStaticMeshFactoryNode(const TMap<int32, TArray<FString>>& MeshUidsPerLodIndex);
+	UE_API UInterchangeStaticMeshFactoryNode* CreateStaticMeshFactoryNode(const TMap<int32, TArray<FString>>& MeshUidsPerLodIndex, const FString& BaseMeshUid = TEXT(""));
 
 	/** This function can create a UInterchangeStaticMeshLodDataNode which represents the LOD data needed by the factory to create a lod mesh */
 	UE_API UInterchangeStaticMeshLodDataNode* CreateStaticMeshLodDataNode(const FString& NodeName, const FString& NodeUniqueID);
@@ -421,8 +422,9 @@ private:
 
 	/**
 	 * Return a reasonable UID and display label for a new mesh factory node.
+	 * @param BaseMeshUid - in case of multiple MeshUidsPerLodIndex generated from a LOD Container, the BaseMeshUid will point to the LOD Container for naming and UID purposes.
 	 */
-	UE_API bool MakeMeshFactoryNodeUidAndDisplayLabel(const TMap<int32, TArray<FString>>& MeshUidsPerLodIndex, int32 LodIndex, FString& NewMeshUid, FString& DisplayLabel);
+	UE_API bool MakeMeshFactoryNodeUidAndDisplayLabel(const TMap<int32, TArray<FString>>& MeshUidsPerLodIndex, int32 LodIndex, FString& NewMeshUid, FString& DisplayLabel, const FString& BaseMeshUid);
 
 	/* Static mesh API END                                                */
 	/************************************************************************/
diff --git a/Engine/Plugins/Interchange/Runtime/Source/Pipelines/Public/InterchangeGenericScenesPipeline.h b/Engine/Plugins/Interchange/Runtime/Source/Pipelines/Public/InterchangeGenericScenesPipeline.h
--- a/Engine/Plugins/Interchange/Runtime/Source/Pipelines/Public/InterchangeGenericScenesPipeline.h
+++ b/Engine/Plugins/Interchange/Runtime/Source/Pipelines/Public/InterchangeGenericScenesPipeline.h
@@ -20,6 +20,7 @@ class UInterchangeLevelFactoryNode;
 class UInterchangeStaticMeshFactoryNode;
 class UInterchangeLevelInstanceActorFactoryNode;
 class UInterchangeSceneNode;
+class UInterchangeSceneComponentNode;
 class UInterchangeSceneVariantSetsNode;
 class UInterchangeSceneImportAssetFactoryNode;
 class UWorld;
@@ -68,6 +69,13 @@ public:
 	UPROPERTY(EditAnywhere, BlueprintReadWrite, Category = "Scene", meta = (SubCategory = "Reimport Assets"))
 	bool bDeleteMissingAssets = false;
 
+	/* If enabled, Normal SceneComponents found under a given Actor will be ignored and their transforms baked into collapsed Instanced Meshes. */
+	UPROPERTY(EditAnywhere, BlueprintReadWrite, Category = "Scene", meta = (SubCategory = "Component properties", DisplayName = "Collapse ISM Scene Components"))
+	bool bCollapseISMSceneComponents = true;
+
+	UPROPERTY(EditAnywhere, BlueprintReadWrite, Category = "Scene", meta = (SubCategory = "Component properties", DisplayName = "Use Hierarchical ISM Components"))
+	bool bUseHierarchicalISMComponents = false;
+
 	/** BEGIN UInterchangePipelineBase overrides */
 	UE_API virtual void AdjustSettingsForContext(const FInterchangePipelineContextParams& ContextParams) override;
 
@@ -109,14 +117,39 @@ protected:
 	 */
 	UE_API virtual UInterchangeActorFactoryNode* CreateActorFactoryNode(const UInterchangeSceneNode* SceneNode, const UInterchangeBaseNode* TranslatedAssetNode) const;
 
+	/**
+	 * Return a new Component Factory Node.
+	 */
+	UE_API virtual void CreateComponentFactoryNode(const UInterchangeSceneComponentNode* TranslatedAssetNode) const;
+
 	/**
 	 * Use to set up the given factory node's attributes after its initialization.
 	 */
 	UE_API virtual void SetUpFactoryNode(UInterchangeActorFactoryNode* ActorFactoryNode, const UInterchangeSceneNode* SceneNode, const UInterchangeBaseNode* TranslatedAssetNode) const;
 
 	UE_API virtual void PopulateSceneNodesPreImport(TArray<UInterchangeSceneNode*>& SceneNodes,
+		TArray<UInterchangeSceneComponentNode*>& SceneComponentNodes,
 		TMap<FString, const UInterchangeStaticMeshFactoryNode*>& StaticMeshFactoryNodesBySceneNodeUid,
 		TSet<FString>& SkinnedMeshNodes) const;
+
+	/**
+	 * Use to set up the given factory node's component's attributes after its initialization.
+	 */
+	UE_API virtual void SetUpComponents(UInterchangeActorFactoryNode* ActorFactoryNode, const UInterchangeSceneNode* SceneNode) const;
+
+	/**
+	 * Genereates and sets up the collected Collapsed InstancedStaticMesh transforms.
+	 * Note: Collapsing will ignore User Defined Attributes from the nodes being collapsed.
+	 */
+	UE_API virtual void CreateAndSetUpCollapsedISMComponents(UInterchangeActorFactoryNode* ActorFactoryNode, const UInterchangeSceneNode* SceneNode) const;
+
+	/**
+	 * Collects InstancedStaticMeshComponentNodes collapses them to have 1 Component per Mesh with the required baked transforms on the InstanceTransforms.
+	 * Only collects ISM Components marked as Visible/(Not Marked as InVisible).
+	 */
+	UE_API virtual void CollectCollapsedInstancedStaticMeshTransforms(const TArray<FString>& SceneComponentUids, TMap<FString, TArray<FTransform>>& InstanceTransformsPerMeshes, FTransform ParentGlobalTransform) const;
+
+	UE_API virtual void SetUpComponentDependencies(UInterchangeActorFactoryNode* ActorFactoryNode, const TArray<FString>& SceneComponentUids) const;
 	
 	/** Disable this option to not convert Standard(Perspective) to Physical Cameras*/
 	UPROPERTY(EditAnywhere, BlueprintReadWrite, Category = "Common Skeletal Meshes and Animations")
diff --git a/Engine/Plugins/Interchange/Runtime/Source/Pipelines/Public/InterchangePipelineMeshesUtilities.h b/Engine/Plugins/Interchange/Runtime/Source/Pipelines/Public/InterchangePipelineMeshesUtilities.h
--- a/Engine/Plugins/Interchange/Runtime/Source/Pipelines/Public/InterchangePipelineMeshesUtilities.h
+++ b/Engine/Plugins/Interchange/Runtime/Source/Pipelines/Public/InterchangePipelineMeshesUtilities.h
@@ -27,10 +27,10 @@ struct FInterchangeLodSceneNodeContainer
 	GENERATED_BODY()
 
 	/**
-	 * Each scene node here represents a mesh scene node. If it represents a LOD group, there may be more then one mesh scene node for a specific LOD index.
+	 * Each scene node here represents a mesh scene/component node. If it represents a LOD group, there may be more then one mesh scene node for a specific LOD index.
 	 */
 	UPROPERTY(EditAnywhere, Category = "Interchange | Pipeline | MeshesInstance")
-	TArray<TObjectPtr<const UInterchangeSceneNode>> SceneNodes;
+	TArray<TObjectPtr<const UInterchangeBaseNode>> BaseNodes;
 };
 
 /*
@@ -61,7 +61,7 @@ struct FInterchangeMeshInstance
 	 * This member is null unless the mesh instance represents a LOD group.
 	 */
 	UPROPERTY(EditAnywhere, Category = "Interchange | Pipeline | MeshesInstance")
-	TObjectPtr<const UInterchangeSceneNode> LodGroupNode;
+	TObjectPtr<const UInterchangeBaseNode> LodGroupNode;
 
 	UPROPERTY(EditAnywhere, Category = "Interchange | Pipeline | MeshesInstance")
 	bool bReferenceSkinnedMesh;
@@ -361,7 +361,7 @@ public:
 	{
 		CurrentDataContext = Context;
 	}
-	
+
 protected:
 	TMap<FString, FInterchangeMeshGeometry> MeshGeometriesPerMeshUid;
 	TMap<FString, FInterchangeMeshInstance> MeshInstancesPerMeshInstanceUid;
