[Interchange] Removed logs submitted by mistake in CL-44559046

diff --git a/Engine/Plugins/Interchange/Runtime/Source/Import/Private/Scene/InterchangeStaticMeshActorFactory.cpp b/Engine/Plugins/Interchange/Runtime/Source/Import/Private/Scene/InterchangeStaticMeshActorFactory.cpp
--- a/Engine/Plugins/Interchange/Runtime/Source/Import/Private/Scene/InterchangeStaticMeshActorFactory.cpp
+++ b/Engine/Plugins/Interchange/Runtime/Source/Import/Private/Scene/InterchangeStaticMeshActorFactory.cpp
@@ -9,7 +9,6 @@
 #include "Components/StaticMeshComponent.h"
 #include "Engine/StaticMesh.h"
 #include "Engine/StaticMeshActor.h"
-#include "InterchangeImportLog.h"
 
 #include UE_INLINE_GENERATED_CPP_BY_NAME(InterchangeStaticMeshActorFactory)
 
@@ -44,8 +43,6 @@ UObject* UInterchangeStaticMeshActorFactory::ProcessActor(AActor& SpawnedActor,
 
 		if (bHasGeometricTransform)
 		{
-			UE_LOG(LogInterchangeImport, Warning, TEXT("BOUNDS bHasGeometricTransform"));
-
 			UStaticMeshComponent* GeometricTransformMeshComponent = nullptr;
 			if (StaticMeshComponent->GetNumChildrenComponents() > 0)
 			{
@@ -73,7 +70,6 @@ UObject* UInterchangeStaticMeshActorFactory::ProcessActor(AActor& SpawnedActor,
 
 				if (const UInterchangeFactoryBaseNode* MeshNode = ActorHelper::FindAssetInstanceFactoryNode(&NodeContainer, &FactoryNode))
 				{
-					UE_LOG(LogInterchangeImport, Warning, TEXT("BOUNDS Process Actor"));
 					ApplyStaticMeshToComponent(MeshNode, GeometricTransformMeshComponent, NodeContainer, MeshActorFactoryNode);
 				}
 
@@ -94,7 +90,6 @@ UObject* UInterchangeStaticMeshActorFactory::ProcessActor(AActor& SpawnedActor,
 		{
 			StaticMeshComponent->UnregisterComponent();
 
-			UE_LOG(LogInterchangeImport, Warning, TEXT("BOUNDS Process Actor"));
 			ApplyStaticMeshToComponent(MeshNode, StaticMeshComponent, NodeContainer, MeshActorFactoryNode);
 
 			StaticMeshComponent->RegisterComponent();
@@ -121,7 +116,6 @@ void UInterchangeStaticMeshActorFactory::SetupObject_GameThread(const FSetupObje
 			const UInterchangeStaticMeshFactoryNode* StaticMeshFactoryNode = TargetNodeUids.IsEmpty() ? nullptr : Cast<UInterchangeStaticMeshFactoryNode>(Arguments.NodeContainer->GetFactoryNode(TargetNodeUids[0]));
 			if (StaticMeshFactoryNode)
 			{
-				UE_LOG(LogInterchangeImport, Warning, TEXT("BOUNDS SetupObject_GameThread"));
 				ApplyStaticMeshToComponent(StaticMeshFactoryNode, StaticMeshComponent, *Arguments.NodeContainer, Cast<UInterchangeMeshActorFactoryNode>(Arguments.FactoryNode));
 			}
 		}
@@ -178,8 +172,6 @@ void UInterchangeStaticMeshActorFactory::ApplyStaticMeshToComponent(const UInter
 			StaticMeshComponent->SetStaticMesh(StaticMesh);
 		}
 
-		UE_LOG(LogInterchangeImport, Warning, TEXT("BOUNDS static mesh '%s' : '%0.3f, %0.3f, %0.3f'."), *StaticMesh->GetName(), StaticMesh->GetBounds().BoxExtent.X, StaticMesh->GetBounds().BoxExtent.Y, StaticMesh->GetBounds().BoxExtent.Z);
-
 		// UE-306557: At this point, bounds might not be calculated yet but the navigation system will try to add it to the Octree, causing an Empty Bounds warning.
 		// To avoid this warning, we manually calculate the bounds here.
 		StaticMesh->CalculateExtendedBounds();
