Interchange: Implement Nanite triangle threshold.

diff --git a/Engine/Plugins/Interchange/Runtime/Source/Import/Private/Mesh/InterchangeStaticMeshFactory.cpp b/Engine/Plugins/Interchange/Runtime/Source/Import/Private/Mesh/InterchangeStaticMeshFactory.cpp
--- a/Engine/Plugins/Interchange/Runtime/Source/Import/Private/Mesh/InterchangeStaticMeshFactory.cpp
+++ b/Engine/Plugins/Interchange/Runtime/Source/Import/Private/Mesh/InterchangeStaticMeshFactory.cpp
@@ -627,9 +627,6 @@ UInterchangeFactoryBase::FImportAssetResult UInterchangeStaticMeshFactory::EndIm
 		return ImportAssetResult;
 	}
 
-	const UClass* StaticMeshClass = StaticMeshFactoryNode->GetObjectClass();
-	check(StaticMeshClass && StaticMeshClass->IsChildOf(GetFactoryClass()));
-
 	// create an asset if it doesn't exist
 	UObject* ExistingAsset = StaticFindObject(nullptr, Arguments.Parent, *Arguments.AssetName);
 
diff --git a/Engine/Plugins/Interchange/Runtime/Source/Pipelines/Private/InterchangeGenericMeshPipeline.cpp b/Engine/Plugins/Interchange/Runtime/Source/Pipelines/Private/InterchangeGenericMeshPipeline.cpp
--- a/Engine/Plugins/Interchange/Runtime/Source/Pipelines/Private/InterchangeGenericMeshPipeline.cpp
+++ b/Engine/Plugins/Interchange/Runtime/Source/Pipelines/Private/InterchangeGenericMeshPipeline.cpp
@@ -453,6 +453,11 @@ void UInterchangeGenericMeshPipeline::ExecutePipeline(UInterchangeBaseNodeContai
 	ExecutePreImportPipelineGeometryCache();
 }
 
+void UInterchangeGenericMeshPipeline::ExecutePostFactoryPipeline(const UInterchangeBaseNodeContainer* InBaseNodeContainer, const FString& NodeKey, UObject* CreatedAsset, bool bIsAReimport)
+{
+	ExecutePostFactoryPipelineStaticMesh(InBaseNodeContainer, NodeKey, CreatedAsset, bIsAReimport);
+}
+
 void UInterchangeGenericMeshPipeline::ExecutePostImportPipeline(const UInterchangeBaseNodeContainer* InBaseNodeContainer, const FString& FactoryNodeKey, UObject* CreatedAsset, bool bIsAReimport)
 {
 	//We do not use the provided base container since ExecutePreImportPipeline cache it
diff --git a/Engine/Plugins/Interchange/Runtime/Source/Pipelines/Private/InterchangeGenericStaticMeshPipeline.cpp b/Engine/Plugins/Interchange/Runtime/Source/Pipelines/Private/InterchangeGenericStaticMeshPipeline.cpp
--- a/Engine/Plugins/Interchange/Runtime/Source/Pipelines/Private/InterchangeGenericStaticMeshPipeline.cpp
+++ b/Engine/Plugins/Interchange/Runtime/Source/Pipelines/Private/InterchangeGenericStaticMeshPipeline.cpp
@@ -1,11 +1,12 @@
 // Copyright Epic Games, Inc. All Rights Reserved.
+
 #include "InterchangeGenericMeshPipeline.h"
 
 #include "AssetRegistry/AssetRegistryModule.h"
 #include "Async/Async.h"
 #include "CoreMinimal.h"
-#include "InterchangeMeshLODContainerNode.h"
 #include "InterchangeMaterialFactoryNode.h"
+#include "InterchangeMeshLODContainerNode.h"
 #include "InterchangeMeshNode.h"
 #include "InterchangePipelineLog.h"
 #include "InterchangePipelineMeshesUtilities.h"
@@ -13,16 +14,17 @@
 #include "InterchangeSceneNode.h"
 #include "InterchangeStaticMeshFactoryNode.h"
 #include "InterchangeStaticMeshLodDataNode.h"
-#include "InterchangeSourceData.h"
-#include "Misc/Paths.h"
 #include "Nodes/InterchangeBaseNode.h"
 #include "Nodes/InterchangeBaseNodeContainer.h"
+#include "Nodes/InterchangeSourceNode.h"
 #include "Nodes/InterchangeUserDefinedAttribute.h"
+
+#include "Async/Async.h"
+#include "CoreMinimal.h"
 #include "Tasks/Task.h"
 #include "UObject/Object.h"
 #include "UObject/ObjectMacros.h"
 
-
 namespace UE::Interchange::Private
 {
 	FString GetNodeName(TObjectPtr<UInterchangePipelineMeshesUtilities> PipelineMeshesUtilities, const UInterchangeBaseNodeContainer& NodeContainer, const FString& NodeUid)
@@ -263,6 +265,12 @@ void UInterchangeGenericMeshPipeline::ExecutePreImportPipelineStaticMesh()
 		});
 #endif
 
+	if (UInterchangeSourceNode* SourceNode = UInterchangeSourceNode::FindOrCreateUniqueInstance(BaseNodeContainer))
+	{
+		// Keep a single triangle threshold for the entire scene
+		SourceNode->SetCustomNaniteTriangleThreshold(NaniteTriangleThreshold);
+	}
+
 	if (bImportStaticMeshes && (CommonMeshesProperties->ForceAllMeshAsType == EInterchangeForceMeshType::IFMT_None || CommonMeshesProperties->ForceAllMeshAsType == EInterchangeForceMeshType::IFMT_StaticMesh))
 	{
 		if (bCombineStaticMeshes)
@@ -438,6 +446,77 @@ void UInterchangeGenericMeshPipeline::ExecutePreImportPipelineStaticMesh()
 	}
 }
 
+namespace UE::InterchangeStaticMeshFactory::Private
+{
+	void ApplyNaniteTriangleThreshold(
+		const UInterchangeStaticMeshFactoryNode* FactoryNode,
+		UStaticMesh* StaticMesh,
+		const UInterchangeBaseNodeContainer* NodeContainer
+	)
+	{
+#if WITH_EDITORONLY_DATA
+		if (!FactoryNode || !NodeContainer || !StaticMesh || !StaticMesh->GetNaniteSettings().bEnabled)
+		{
+			return;
+		}
+
+		int64 TriangleThreshold = 0;
+		if (const UInterchangeSourceNode* SourceNode = UInterchangeSourceNode::GetUniqueInstance(NodeContainer))
+		{
+			SourceNode->GetCustomNaniteTriangleThreshold(TriangleThreshold);
+		}
+		if (TriangleThreshold <= 0)
+		{
+			return;
+		}
+
+		const int32 LodIndex = 0;
+		FMeshDescription* MeshDescription = StaticMesh->GetMeshDescription(LodIndex);
+		if (!MeshDescription)
+		{
+			return;
+		}
+
+		if (MeshDescription->Triangles().Num() < TriangleThreshold)
+		{
+			StaticMesh->GetNaniteSettings().bEnabled = false;
+		}
+#endif	  // #if WITH_EDITORONLY_DATA
+	}
+}	 // namespace UE::InterchangeStaticMeshFactory::Private
+
+void UInterchangeGenericMeshPipeline::ExecutePostFactoryPipelineStaticMesh(
+	const UInterchangeBaseNodeContainer* InBaseNodeContainer,
+	const FString& NodeKey,
+	UObject* CreatedAsset,
+	bool bIsAReimport
+)
+{
+	UStaticMesh* StaticMesh = Cast<UStaticMesh>(CreatedAsset);
+	if (!StaticMesh)
+	{
+		return;
+	}
+
+	UInterchangeStaticMeshFactoryNode* StaticMeshFactoryNode = Cast<UInterchangeStaticMeshFactoryNode>(InBaseNodeContainer->GetFactoryNode(NodeKey));
+	if (!StaticMeshFactoryNode)
+	{
+		return;
+	}
+
+	// If we're reimporting and being told to leave the properties alone, just early out instead
+	if (bIsAReimport)
+	{
+		EReimportStrategyFlags Strategy = StaticMeshFactoryNode->GetReimportStrategyFlags();
+		if (Strategy == EReimportStrategyFlags::ApplyNoProperties)
+		{
+			return;
+		}
+	}
+
+	UE::InterchangeStaticMeshFactory::Private::ApplyNaniteTriangleThreshold(StaticMeshFactoryNode, StaticMesh, InBaseNodeContainer);
+}
+
 bool UInterchangeGenericMeshPipeline::MakeMeshFactoryNodeUidAndDisplayLabel(const TMap<int32, TArray<FString>>& MeshUidsPerLodIndex, int32 LodIndex, FString& NewNodeUid, FString& DisplayLabel, const FString& BaseMeshUid)
 {
 	int32 SceneNodeCount = 0;
diff --git a/Engine/Plugins/Interchange/Runtime/Source/Pipelines/Public/InterchangeGenericMeshPipeline.h b/Engine/Plugins/Interchange/Runtime/Source/Pipelines/Public/InterchangeGenericMeshPipeline.h
--- a/Engine/Plugins/Interchange/Runtime/Source/Pipelines/Public/InterchangeGenericMeshPipeline.h
+++ b/Engine/Plugins/Interchange/Runtime/Source/Pipelines/Public/InterchangeGenericMeshPipeline.h
@@ -102,10 +102,20 @@ public:
 
 	//////	Static Meshes Build settings Properties //////
 
-	/** If enabled, imported meshes will be rendered by Nanite at runtime. Make sure your meshes and materials meet the requirements for Nanite. */
+	/**
+	 * If enabled, imported meshes will be rendered by Nanite at runtime. Make sure your meshes and materials meet the requirements for Nanite.
+	 * See also NaniteTriangleThreshold
+	 */
 	UPROPERTY(EditAnywhere, BlueprintReadWrite, Category = "Static Meshes", meta = (SubCategory = "Build"))
 	bool bBuildNanite = true;
 
+	/**
+	 * Minimum triangle count a mesh needs to have in order to get Nanite enabled for it when BuildNanite is enabled.
+	 * When zero it means Nanite will always be enabled for all meshes when BuildNanite is enabled.
+	 */
+	UPROPERTY(EditAnywhere, BlueprintReadWrite, Category = "Static Meshes", meta = (SubCategory = "Build", EditCondition = "bBuildNanite"))
+	int64 NaniteTriangleThreshold = 0;
+
 	/** If enabled, builds a reversed index buffer for each static mesh. */
 	UPROPERTY(EditAnywhere, BlueprintReadWrite, Category = "Static Meshes", meta = (SubCategory = "Build"))
 	bool bBuildReversedIndexBuffer = false;
@@ -314,6 +324,8 @@ public:
 protected:
 	UE_API virtual void ExecutePipeline(UInterchangeBaseNodeContainer* InBaseNodeContainer, const TArray<UInterchangeSourceData*>& InSourceDatas, const FString& ContentBasePath) override;
 
+	UE_API virtual void ExecutePostFactoryPipeline(const UInterchangeBaseNodeContainer* InBaseNodeContainer, const FString& NodeKey, UObject* CreatedAsset, bool bIsAReimport) override;
+
 	UE_API virtual void ExecutePostImportPipeline(const UInterchangeBaseNodeContainer* InBaseNodeContainer, const FString& NodeKey, UObject* CreatedAsset, bool bIsAReimport) override;
 
 	virtual bool CanExecuteOnAnyThread(EInterchangePipelineTask PipelineTask) override
@@ -401,6 +413,8 @@ private:
 	 */
 	UE_API void ExecutePreImportPipelineStaticMesh();
 
+	UE_API void ExecutePostFactoryPipelineStaticMesh(const UInterchangeBaseNodeContainer* InBaseNodeContainer, const FString& NodeKey, UObject* CreatedAsset, bool bIsAReimport);
+
 	/** Static mesh factory assets nodes */
 	TArray<UInterchangeStaticMeshFactoryNode*> StaticMeshFactoryNodes;
 
diff --git a/Engine/Source/Runtime/Interchange/Core/Private/Nodes/InterchangeSourceNode.cpp b/Engine/Source/Runtime/Interchange/Core/Private/Nodes/InterchangeSourceNode.cpp
--- a/Engine/Source/Runtime/Interchange/Core/Private/Nodes/InterchangeSourceNode.cpp
+++ b/Engine/Source/Runtime/Interchange/Core/Private/Nodes/InterchangeSourceNode.cpp
@@ -205,12 +205,12 @@ bool UInterchangeSourceNode::SetCustomUseAssetTypeSubPathSuffix(const bool& Attr
 	IMPLEMENT_NODE_ATTRIBUTE_SETTER_NODELEGATE(UseAssetTypeSubPathSuffix, bool);
 }
 
-INTERCHANGECORE_API bool UInterchangeSourceNode::GetCustomReimportStrategyFlags(uint8& AttributeValue) const
+bool UInterchangeSourceNode::GetCustomReimportStrategyFlags(uint8& AttributeValue) const
 {
 	IMPLEMENT_NODE_ATTRIBUTE_GETTER(ReimportStrategyFlags, uint8);
 }
 
-INTERCHANGECORE_API bool UInterchangeSourceNode::SetCustomReimportStrategyFlags(uint8 AttributeValue)
+bool UInterchangeSourceNode::SetCustomReimportStrategyFlags(uint8 AttributeValue)
 {
 	IMPLEMENT_NODE_ATTRIBUTE_SETTER_NODELEGATE(ReimportStrategyFlags, uint8)
 }
@@ -225,3 +225,12 @@ bool UInterchangeSourceNode::SetCustomSkeletalMeshFrontAxis(uint8 AttributeValue
 	IMPLEMENT_NODE_ATTRIBUTE_SETTER_NODELEGATE(SkeletalMeshFrontAxis, uint8);
 }
 
+bool UInterchangeSourceNode::GetCustomNaniteTriangleThreshold(int64& AttributeValue) const
+{
+	IMPLEMENT_NODE_ATTRIBUTE_GETTER(NaniteTriangleThreshold, int64);
+}
+
+bool UInterchangeSourceNode::SetCustomNaniteTriangleThreshold(const int64& AttributeValue)
+{
+	IMPLEMENT_NODE_ATTRIBUTE_SETTER_NODELEGATE(NaniteTriangleThreshold, int64)
+}
diff --git a/Engine/Source/Runtime/Interchange/Core/Public/Nodes/InterchangeSourceNode.h b/Engine/Source/Runtime/Interchange/Core/Public/Nodes/InterchangeSourceNode.h
--- a/Engine/Source/Runtime/Interchange/Core/Public/Nodes/InterchangeSourceNode.h
+++ b/Engine/Source/Runtime/Interchange/Core/Public/Nodes/InterchangeSourceNode.h
@@ -199,6 +199,14 @@ public:
 	UFUNCTION(BlueprintCallable, Category = "Interchange | Node | Source")
 	INTERCHANGECORE_API bool SetCustomSkeletalMeshFrontAxis(uint8 AttributeValue);
 
+	/** Gets the minimum triangle count a mesh needs to have in order to get Nanite enabled for it when bBuildNanite is true on the static mesh pipeline. */
+	UFUNCTION(BlueprintCallable, Category = "Interchange | Node | Source")
+	INTERCHANGECORE_API bool GetCustomNaniteTriangleThreshold(int64& MinNumTriangles) const;
+
+	/** Sets the minimum triangle count a mesh needs to have in order to get Nanite enabled for it when bBuildNanite is true on the static mesh pipeline. */
+	UFUNCTION(BlueprintCallable, Category = "Interchange | Node | Source")
+	INTERCHANGECORE_API bool SetCustomNaniteTriangleThreshold(const int64& MinNumTriangles);
+
 private:
 	IMPLEMENT_NODE_ATTRIBUTE_KEY(SourceFrameRateNumerator);
 	IMPLEMENT_NODE_ATTRIBUTE_KEY(SourceFrameRateDenominator);
@@ -213,6 +221,7 @@ private:
 	IMPLEMENT_NODE_ATTRIBUTE_KEY(UseAssetTypeSubPathSuffix);
 	IMPLEMENT_NODE_ATTRIBUTE_KEY(ReimportStrategyFlags);
 	IMPLEMENT_NODE_ATTRIBUTE_KEY(SkeletalMeshFrontAxis);
+	IMPLEMENT_NODE_ATTRIBUTE_KEY(NaniteTriangleThreshold);
 
 	// Extra InformationTo show in the Config Panel.
 	UE::Interchange::TMapAttributeHelper<FString, FString> ExtraInformation;
