[Interchange] Prevent unnecessary MeshDescription copy inside StaticMesh Factory

While profiling for memory usage, I noticed a spike inside the StaticMesh factory caused by what seems to be an extra copy of a MeshDescription.

The factory already has copies of the MeshDescription inside `TMap<int32, FLodPayloads> PayloadsPerLodIndex;` and the idea is to combines the mesh descriptions in-place instead of a copy.

My biggest concern about this change is that calling `UInterchangeStaticMeshFactory::ImportAsset_Async` twice will not work unless `UInterchangeStaticMeshFactory::CreatePayloadTasks` is called before. This is  because `PayloadsPerLodIndex` was already "consumed" and needs to be repopulated.

Interchange tests verification : Pass (689a441f593d93678dc2cbc0)

Before and After (testing using one big mesh of 9M triangles (11M vertices))
level import : 14Gb ---> 10.6Gb
asset import: 13.9Gb ---> 10.5Gb

For Moore-Lane Assembly USDA:
level import: 27Gb ---> 25.7Gb

#rnx

diff --git a/Engine/Plugins/Interchange/Runtime/Source/Import/Private/Mesh/InterchangeStaticMeshFactory.cpp b/Engine/Plugins/Interchange/Runtime/Source/Import/Private/Mesh/InterchangeStaticMeshFactory.cpp
--- a/Engine/Plugins/Interchange/Runtime/Source/Import/Private/Mesh/InterchangeStaticMeshFactory.cpp
+++ b/Engine/Plugins/Interchange/Runtime/Source/Import/Private/Mesh/InterchangeStaticMeshFactory.cpp
@@ -414,6 +414,12 @@ UInterchangeFactoryBase::FImportAssetResult UInterchangeStaticMeshFactory::Impor
 	int32 CurrentLodIndex = 0;
 	for (int32 LodIndex = 0; LodIndex < LodCount; ++LodIndex)
 	{
+		if (!PayloadsPerLodIndex.Contains(LodIndex))
+		{
+			UE_LOG(LogInterchangeImport, Error, TEXT("LOD %d do not have any valid payload to create a mesh when importing StaticMesh asset %s."), LodIndex, *Arguments.AssetName);
+			continue;
+		}
+
 		FString LodUniqueId = LodDataUniqueIds[LodIndex];
 		const UInterchangeStaticMeshLodDataNode* LodDataNode = Cast<UInterchangeStaticMeshLodDataNode>(Arguments.NodeContainer->GetNode(LodUniqueId));
 		if (!LodDataNode)
@@ -433,26 +439,14 @@ UInterchangeFactoryBase::FImportAssetResult UInterchangeStaticMeshFactory::Impor
 		
 		// Merge the SourceObjectName attribute, if it exists on both the target and source.
 		AppendSettings.bMergeObjectName = true;
-				
 
-		TArray<FString> MeshUids;
-		LodDataNode->GetMeshUids(MeshUids);
-
-		FLodPayloads LodPayloads;
-		if (PayloadsPerLodIndex.Contains(LodIndex))
-		{
-			// Fill the lod mesh description using all combined mesh parts
-			LodPayloads = PayloadsPerLodIndex.FindChecked(LodIndex);
-		}
-		else
-		{
-			UE_LOG(LogInterchangeImport, Error, TEXT("LOD %d do not have any valid payload to create a mesh when importing StaticMesh asset %s."), LodIndex, *Arguments.AssetName);
-			continue;
-		}
+		// Fill the lod mesh description using all combined mesh parts
+		FLodPayloads* LodPayloads = &PayloadsPerLodIndex.FindChecked(LodIndex);
 
 		// Just move the mesh description from the first valid payload then append the rest
 		bool bFirstValidMoved = false;
-		for (TPair<FInterchangeMeshPayLoadKey, UE::Interchange::FMeshPayload>& KeyAndPayload : LodPayloads.MeshPayloadPerKey)
+
+		for (TPair<FInterchangeMeshPayLoadKey, UE::Interchange::FMeshPayload>& KeyAndPayload : LodPayloads->MeshPayloadPerKey)
 		{
 			const TOptional<UE::Interchange::FMeshPayloadData>& LodMeshPayload = KeyAndPayload.Value.PayloadData;
 			if (!LodMeshPayload.IsSet())
@@ -468,7 +462,7 @@ UInterchangeFactoryBase::FImportAssetResult UInterchangeStaticMeshFactory::Impor
 				{
 					continue;
 				}
-				LodMeshDescription = MoveTemp(MeshDescription);
+				LodMeshDescription = MoveTemp(MeshDescription); // This is emptying the mesh description inside PayloadsPerLodIndex
 				bFirstValidMoved = true;
 			}
 			else
@@ -565,13 +559,17 @@ UInterchangeFactoryBase::FImportAssetResult UInterchangeStaticMeshFactory::Impor
 				}
 
 				using namespace UE::Interchange::Private::MeshHelper;
-				bImportedCustomCollision |= ImportBoxCollision(Arguments, LodPayloads.CollisionBoxPayloadPerKey, StaticMesh, bForceGeneration);
-				bImportedCustomCollision |= ImportCapsuleCollision(Arguments, LodPayloads.CollisionCapsulePayloadPerKey, StaticMesh);
-				bImportedCustomCollision |= ImportSphereCollision(Arguments, LodPayloads.CollisionSpherePayloadPerKey, StaticMesh, bForceGeneration);
-				bImportedCustomCollision |= ImportConvexCollision(Arguments, LodPayloads.CollisionConvexPayloadPerKey, StaticMesh, LodDataNode);
+				bImportedCustomCollision |= ImportBoxCollision(Arguments, LodPayloads->CollisionBoxPayloadPerKey, StaticMesh, bForceGeneration);
+				bImportedCustomCollision |= ImportCapsuleCollision(Arguments, LodPayloads->CollisionCapsulePayloadPerKey, StaticMesh);
+				bImportedCustomCollision |= ImportSphereCollision(Arguments, LodPayloads->CollisionSpherePayloadPerKey, StaticMesh, bForceGeneration);
+				bImportedCustomCollision |= ImportConvexCollision(Arguments, LodPayloads->CollisionConvexPayloadPerKey, StaticMesh, LodDataNode);
 			}
 		}
 
+		// Payload has been processed and moved to ImportAssetObjectData.LodMeshDescriptions.
+		// Remove it to make sure it's not reused.
+		PayloadsPerLodIndex.Remove(LodIndex); 
+
 		CurrentLodIndex++;
 	}
 
