Interchange FBX:
- Fix material inconsistencies with the Material Instances use.
- Fix scalar parameter not correctly set with the Material Instance node when having a texture.

#virtualized

diff --git a/Engine/Plugins/Interchange/Runtime/Source/Parsers/Fbx/Private/FbxMaterial.cpp b/Engine/Plugins/Interchange/Runtime/Source/Parsers/Fbx/Private/FbxMaterial.cpp
--- a/Engine/Plugins/Interchange/Runtime/Source/Parsers/Fbx/Private/FbxMaterial.cpp
+++ b/Engine/Plugins/Interchange/Runtime/Source/Parsers/Fbx/Private/FbxMaterial.cpp
@@ -15,8 +15,6 @@
 #include "Misc/Paths.h"
 #include "Nodes/InterchangeBaseNodeContainer.h"
 
-
-
 #define LOCTEXT_NAMESPACE "InterchangeFbxMaterial"
 
 namespace UE
@@ -160,6 +158,7 @@ namespace UE
 				if (const UInterchangeTexture2DNode* TextureNode = CreateTexture2DNode(FbxTexture, NodeContainer, MaterialInstanceNode))
 				{
 					MaterialInstanceNode->AddTextureParameterValue(InputName.ToString() + TEXT("Map"), TextureNode->GetUniqueID());
+					MaterialInstanceNode->AddScalarParameterValue(InputName.ToString() + TEXT("MapWeight"), 1.f);
 				}
 				else
 				{
