Interchange MaterialX:
- Fix possible nullptr when retrieve the graph parent

diff --git a/Engine/Plugins/Interchange/Runtime/Source/Import/Private/MaterialX/MaterialXUtils/MaterialXSurfaceShaderAbstract.cpp b/Engine/Plugins/Interchange/Runtime/Source/Import/Private/MaterialX/MaterialXUtils/MaterialXSurfaceShaderAbstract.cpp
--- a/Engine/Plugins/Interchange/Runtime/Source/Import/Private/MaterialX/MaterialXUtils/MaterialXSurfaceShaderAbstract.cpp
+++ b/Engine/Plugins/Interchange/Runtime/Source/Import/Private/MaterialX/MaterialXUtils/MaterialXSurfaceShaderAbstract.cpp
@@ -1062,7 +1062,7 @@ void FMaterialXSurfaceShaderAbstract::ConnectHeightToNormalInputToOutput(const F
 
 			SetMatchingInputsNames(Connect.UpstreamNode);
 
-			mx::NodeGraphPtr Graph = Connect.UpstreamNode->getParent()->asA<mx::NodeGraph>();
+			mx::GraphElementPtr Graph = Connect.UpstreamNode->getParent()->asA<mx::GraphElement>();
 			Graph->removeNode(ConnectedNode->getName());
 
 			using namespace UE::Interchange::Materials::Standard::Nodes;
@@ -1131,7 +1131,7 @@ void FMaterialXSurfaceShaderAbstract::ConnectBlurInputToOutput(const FConnectNod
 
 			SetMatchingInputsNames(Connect.UpstreamNode);
 
-			mx::NodeGraphPtr Graph = Connect.UpstreamNode->getParent()->asA<mx::NodeGraph>();
+			mx::GraphElementPtr Graph = Connect.UpstreamNode->getParent()->asA<mx::GraphElement>();
 			Graph->removeNode(ConnectedNode->getName());
 
 			using namespace UE::Interchange::Materials::Standard::Nodes;
