Interchange USD: Added support for groom import in Interchange.
Implemented groom import for Interchange USD.

diff --git a/Engine/Plugins/Interchange/Extensions/OpenUSD/Source/Import/InterchangeOpenUSDImport.build.cs b/Engine/Plugins/Interchange/Extensions/OpenUSD/Source/Import/InterchangeOpenUSDImport.build.cs
--- a/Engine/Plugins/Interchange/Extensions/OpenUSD/Source/Import/InterchangeOpenUSDImport.build.cs
+++ b/Engine/Plugins/Interchange/Extensions/OpenUSD/Source/Import/InterchangeOpenUSDImport.build.cs
@@ -28,6 +28,7 @@ namespace UnrealBuildTool.Rules
 				{
 					"DynamicMesh",				// For FMeshMapMaker
 					"GeometryCore",				// For FDynamicMesh3
+					"HairStrandsCore",
 					"InterchangeEngine", 		// For UInterchangeManager::GetInterchangeManager
 					"InterchangeFactoryNodes",
 					"InterchangeNodes", 		// For UInterchangeShaderGraphNode
diff --git a/Engine/Plugins/Interchange/Extensions/OpenUSD/Source/Import/Private/InterchangeUsdTranslator.cpp b/Engine/Plugins/Interchange/Extensions/OpenUSD/Source/Import/Private/InterchangeUsdTranslator.cpp
--- a/Engine/Plugins/Interchange/Extensions/OpenUSD/Source/Import/Private/InterchangeUsdTranslator.cpp
+++ b/Engine/Plugins/Interchange/Extensions/OpenUSD/Source/Import/Private/InterchangeUsdTranslator.cpp
@@ -7,6 +7,7 @@
 #include "UnrealUSDWrapper.h"
 #include "USDConversionUtils.h"
 #include "USDGeomMeshConversion.h"
+#include "USDGroomConversion.h"
 #include "USDLayerUtils.h"
 #include "USDLightConversion.h"
 #include "USDLog.h"
@@ -37,6 +38,7 @@
 #include "UsdWrappers/VtValue.h"
 
 #include "InterchangeCameraNode.h"
+#include "InterchangeGroomNode.h"
 #include "InterchangeLightNode.h"
 #include "InterchangeMeshLODContainerNode.h"
 #include "InterchangeManager.h"
@@ -64,6 +66,7 @@
 
 #include "Async/ParallelFor.h"
 #include "Components/HeterogeneousVolumeComponent.h"
+#include "GroomCacheData.h"
 #include "HAL/IConsoleManager.h"
 #include "Internationalization/Regex.h"
 #include "Modules/ModuleManager.h"
@@ -108,6 +111,7 @@ namespace UE::InterchangeUsdTranslator::Private
 	const static FString AnimationTrackPrefix = TEXT("\\AnimationTrack\\");
 	const static FString BonePrefix = TEXT("\\Bone\\");
 	const static FString CameraPrefix = TEXT("\\Camera\\");
+	const static FString GroomPrefix = TEXT("\\Groom\\");
 	const static FString LightPrefix = TEXT("\\Light\\");
 	const static FString LODPrefix = TEXT("\\LOD\\");
 	const static FString MaterialPrefix = TEXT("\\Material\\");
@@ -344,6 +348,11 @@ namespace UE::InterchangeUsdTranslator::Private
 			bool& bOutNeedsFrameTrack
 		);
 
+		FString AddGroomNode(
+			const UE::FUsdPrim& Prim,
+			UInterchangeBaseNodeContainer& NodeContainer
+		);
+
 		/**
 		 * Search for a MaterialX file embedded in the file, in case there's none, create a ShaderGraph using MaterialX
 		 * @param GeomPropValueNames - Array of the <geompropvalue> node names from the shader graph that has been converted to <image> nodes (used
@@ -3277,6 +3286,31 @@ namespace UE::InterchangeUsdTranslator::Private
 		return VolumeAssetNodeUids.Array();
 	}
 
+	FString UInterchangeUSDTranslatorImpl::AddGroomNode(
+		const UE::FUsdPrim& Prim,
+		UInterchangeBaseNodeContainer& NodeContainer
+	)
+	{
+		TRACE_CPUPROFILER_EVENT_SCOPE(UInterchangeUSDTranslatorImpl::AddGroomNode)
+
+		const FString PrimPath = UsdUtils::GetPrototypePrimPath(Prim).GetString();
+		const FString NodeUid = GroomPrefix + PrimPath;
+		const FString NodeName(Prim.GetName().ToString());
+
+		// Check if Node already exist with this ID
+		if (const UInterchangeMeshNode* Node = Cast<UInterchangeMeshNode>(NodeContainer.GetNode(NodeUid)))
+		{
+			return NodeUid;
+		}
+
+		UInterchangeGroomNode* GroomNode = NewObject<UInterchangeGroomNode>(&NodeContainer);
+		NodeContainer.SetupNode(GroomNode, NodeUid, NodeName, EInterchangeNodeContainerType::TranslatedAsset);
+		GroomNode->SetAssetName(NodeName);
+		GroomNode->SetPayloadKey(PrimPath);
+
+		return NodeUid;
+	}
+
 	void AddTrackSetNode(UInterchangeUSDTranslatorImpl& Impl, UInterchangeBaseNodeContainer& NodeContainer)
 	{
 		// For now we only want a single track set (i.e. LevelSequence) per stage.
@@ -4136,7 +4170,7 @@ namespace UE::InterchangeUsdTranslator::Private
 					Func(FullKey, Pair.Value, Node);
 				}
 			}
-		};
+			};
 
 		const FString ConcatParentKey;
 		ParseMetadataMap(Prim.GetAllAuthoredMetadata(), ConcatParentKey);
@@ -4195,9 +4229,13 @@ namespace UE::InterchangeUsdTranslator::Private
 		{
 			AssetNodeUids = TranslatorImpl.AddVolumeNodes(Prim, NodeContainer, VolumeMaterialInstanceUid, bNeedsVolumeTrack);
 		}
+		else if ((Prim.IsA(TEXT("BasisCurves")) || Prim.IsA(TEXT("Xform"))) && Prim.HasAPI(TEXT("GroomAPI")))
+		{
+			AssetNodeUids = {TranslatorImpl.AddGroomNode(Prim, NodeContainer)};
+		}
 		else
 		{
-			const static TSet<FName> KnownUnsupported = {TEXT("SpatialAudio"), TEXT("BasisCurves")};
+			const static TSet<FName> KnownUnsupported = {TEXT("SpatialAudio")};
 
 			if (KnownUnsupported.Contains(TypeName))
 			{
@@ -5898,4 +5936,32 @@ TOptional<UE::Interchange::FVolumePayloadData> UInterchangeUSDTranslator::GetVol
 	return {};
 }
 
+TOptional<UE::Interchange::FGroomPayloadData> UInterchangeUSDTranslator::GetGroomPayloadData(const FInterchangeGroomPayloadKey& PayloadKey) const
+{
+	TRACE_CPUPROFILER_EVENT_SCOPE(UInterchangeUSDTranslator::GetGroomPayloadData)
+
+#if USE_USD_SDK
+	const FString& PrimPath = PayloadKey.UniqueId;
+	UE::FUsdPrim Prim = Impl->UsdStage.GetPrimAtPath(UE::FSdfPath{ *PrimPath });
+
+	FGroomAnimationInfo AnimInfo;
+	UE::Interchange::FGroomPayloadData PayloadData;
+
+	const bool bSuccess = UsdToUnreal::ConvertGroomHierarchy(
+		Prim,
+		UsdUtils::GetEarliestTimeCode(),
+		FTransform::Identity,
+		PayloadData.HairDescription,
+		&AnimInfo
+	);
+
+	if (bSuccess && PayloadData.HairDescription.IsValid())
+	{
+		return PayloadData;
+	}
+#endif // USE_USD_SDK
+
+	return {};
+}
+
 #undef LOCTEXT_NAMESPACE
diff --git a/Engine/Plugins/Interchange/Extensions/OpenUSD/Source/Import/Public/InterchangeUsdTranslator.h b/Engine/Plugins/Interchange/Extensions/OpenUSD/Source/Import/Public/InterchangeUsdTranslator.h
--- a/Engine/Plugins/Interchange/Extensions/OpenUSD/Source/Import/Public/InterchangeUsdTranslator.h
+++ b/Engine/Plugins/Interchange/Extensions/OpenUSD/Source/Import/Public/InterchangeUsdTranslator.h
@@ -3,6 +3,7 @@
 #pragma once
 
 #include "Animation/InterchangeAnimationPayloadInterface.h"
+#include "Groom/InterchangeGroomPayloadInterface.h"
 #include "InterchangeTranslatorBase.h"
 #include "Mesh/InterchangeMeshPayloadInterface.h"
 #include "Texture/InterchangeBlockedTexturePayloadInterface.h"
@@ -93,6 +94,7 @@ class UInterchangeUSDTranslator
 	, public IInterchangeBlockedTexturePayloadInterface
 	, public IInterchangeAnimationPayloadInterface
 	, public IInterchangeVolumePayloadInterface
+	, public IInterchangeGroomPayloadInterface
 {
 	GENERATED_BODY()
 
@@ -144,6 +146,8 @@ public:
 	) const override;
 
 	virtual TOptional<UE::Interchange::FVolumePayloadData> GetVolumePayloadData(const UE::Interchange::FVolumePayloadKey& PayloadKey) const override;
+
+	virtual TOptional<UE::Interchange::FGroomPayloadData> GetGroomPayloadData(const FInterchangeGroomPayloadKey& PayloadKey) const override;
 	/** End Interchange payload interfaces */
 
 private:
diff --git a/Engine/Plugins/Interchange/Runtime/Interchange.uplugin b/Engine/Plugins/Interchange/Runtime/Interchange.uplugin
--- a/Engine/Plugins/Interchange/Runtime/Interchange.uplugin
+++ b/Engine/Plugins/Interchange/Runtime/Interchange.uplugin
@@ -87,6 +87,10 @@
 		}
 	],
 	"Plugins": [
+		{
+			"Name": "HairStrands",
+			"Enabled": true
+		},
 		{
 			"Name": "GeometryCache",
 			"Enabled": true
diff --git a/Engine/Plugins/Interchange/Runtime/Source/FactoryNodes/InterchangeFactoryNodes.Build.cs b/Engine/Plugins/Interchange/Runtime/Source/FactoryNodes/InterchangeFactoryNodes.Build.cs
--- a/Engine/Plugins/Interchange/Runtime/Source/FactoryNodes/InterchangeFactoryNodes.Build.cs
+++ b/Engine/Plugins/Interchange/Runtime/Source/FactoryNodes/InterchangeFactoryNodes.Build.cs
@@ -25,6 +25,7 @@ namespace UnrealBuildTool.Rules
 			PrivateDependencyModuleNames.AddRange(
 				new string[]
 				{
+					"HairStrandsCore",
 					"GeometryCache",
 					"LevelSequence",
 				}
diff --git a/Engine/Plugins/Interchange/Runtime/Source/FactoryNodes/Private/InterchangeGroomFactoryNode.cpp b/Engine/Plugins/Interchange/Runtime/Source/FactoryNodes/Private/InterchangeGroomFactoryNode.cpp
new file mode 100644
--- /dev/null
+++ b/Engine/Plugins/Interchange/Runtime/Source/FactoryNodes/Private/InterchangeGroomFactoryNode.cpp
@@ -0,0 +1,18 @@
+// Copyright Epic Games, Inc. All Rights Reserved.
+
+#include "InterchangeGroomFactoryNode.h"
+
+#include "GroomAsset.h"
+
+#include UE_INLINE_GENERATED_CPP_BY_NAME(InterchangeGroomFactoryNode)
+
+FString UInterchangeGroomFactoryNode::GetTypeName() const
+{
+	const FString TypeName = TEXT("GroomFactoryNode");
+	return TypeName;
+}
+
+UClass* UInterchangeGroomFactoryNode::GetObjectClass() const
+{
+	return UGroomAsset::StaticClass();
+}
diff --git a/Engine/Plugins/Interchange/Runtime/Source/FactoryNodes/Public/InterchangeGroomFactoryNode.h b/Engine/Plugins/Interchange/Runtime/Source/FactoryNodes/Public/InterchangeGroomFactoryNode.h
new file mode 100644
--- /dev/null
+++ b/Engine/Plugins/Interchange/Runtime/Source/FactoryNodes/Public/InterchangeGroomFactoryNode.h
@@ -0,0 +1,21 @@
+// Copyright Epic Games, Inc. All Rights Reserved.
+
+#pragma once
+
+#include "Nodes/InterchangeFactoryBaseNode.h"
+
+#include "InterchangeGroomFactoryNode.generated.h"
+
+#define UE_API INTERCHANGEFACTORYNODES_API
+
+UCLASS(MinimalAPI, BlueprintType)
+class UInterchangeGroomFactoryNode : public UInterchangeFactoryBaseNode
+{
+	GENERATED_BODY()
+
+public:
+	UE_API virtual FString GetTypeName() const override;
+	UE_API virtual class UClass* GetObjectClass() const override;
+};
+
+#undef UE_API
\ No newline at end of file
diff --git a/Engine/Plugins/Interchange/Runtime/Source/Import/InterchangeImport.Build.cs b/Engine/Plugins/Interchange/Runtime/Source/Import/InterchangeImport.Build.cs
--- a/Engine/Plugins/Interchange/Runtime/Source/Import/InterchangeImport.Build.cs
+++ b/Engine/Plugins/Interchange/Runtime/Source/Import/InterchangeImport.Build.cs
@@ -13,6 +13,7 @@ namespace UnrealBuildTool.Rules
 					"Core",
 					"CoreUObject",
 					"Engine",
+					"HairStrandsCore",
 					"InterchangeCore",
 					"InterchangeCommon",
 					"InterchangeDispatcher",
@@ -57,6 +58,7 @@ namespace UnrealBuildTool.Rules
 					new string[]
 					{
 						"BSPUtils",
+						"HairStrandsEditor",
 						"InterchangeFbxParser",
 						"MaterialEditor",
 						"MaterialX",
diff --git a/Engine/Plugins/Interchange/Runtime/Source/Import/Private/Groom/InterchangeGroomFactory.cpp b/Engine/Plugins/Interchange/Runtime/Source/Import/Private/Groom/InterchangeGroomFactory.cpp
new file mode 100644
--- /dev/null
+++ b/Engine/Plugins/Interchange/Runtime/Source/Import/Private/Groom/InterchangeGroomFactory.cpp
@@ -0,0 +1,250 @@
+// Copyright Epic Games, Inc. All Rights Reserved.
+
+#include "Groom/InterchangeGroomFactory.h"
+
+#include "Groom/InterchangeGroomPayloadInterface.h"
+#include "GroomAsset.h"
+#include "GroomBuilder.h"
+#include "GroomImportOptions.h"
+#if WITH_EDITOR
+#include "HairStrandsImporter.h"
+#endif
+
+#include "InterchangeGroomFactoryNode.h"
+#include "InterchangeGroomNode.h"
+#include "InterchangeImportCommon.h"
+#include "InterchangeImportLog.h"
+#include "InterchangeSourceData.h"
+#include "InterchangeTranslatorBase.h"
+#include "Nodes/InterchangeBaseNodeContainer.h"
+
+#include UE_INLINE_GENERATED_CPP_BY_NAME(InterchangeGroomFactory)
+
+#define LOCTEXT_NAMESPACE "InterchangeGroomFactory"
+
+UClass* UInterchangeGroomFactory::GetFactoryClass() const
+{
+	return UGroomAsset::StaticClass();
+}
+
+void UInterchangeGroomFactory::CreatePayloadTasks(const FImportAssetObjectParams& Arguments, bool bAsync, TArray<TSharedPtr<UE::Interchange::FInterchangeTaskBase>>& PayloadTasks)
+{
+	TRACE_CPUPROFILER_EVENT_SCOPE(UInterchangeGroomFactory::CreatePayloadTasks);
+
+	UInterchangeGroomFactoryNode* GroomFactoryNode = Cast<UInterchangeGroomFactoryNode>(Arguments.AssetNode);
+	if (GroomFactoryNode == nullptr)
+	{
+		return;
+	}
+
+	const IInterchangeGroomPayloadInterface* GroomPayloadInterface = Cast<IInterchangeGroomPayloadInterface>(Arguments.Translator);
+	if (!GroomPayloadInterface)
+	{
+		UE_LOG(LogInterchangeImport, Error, TEXT("Cannot import groom asset. The translator does not implement IInterchangeGroomPayloadInterface."));
+		return;
+	}
+
+	TArray<FString> GroomUids;
+	GroomFactoryNode->GetTargetNodeUids(GroomUids);
+	if (GroomUids.Num() == 0)
+	{
+		return;
+	}
+
+	const FString& GroomUid = GroomUids[0]; // the groom pipeline adds only one target node
+	const UInterchangeBaseNode* Node = Arguments.NodeContainer->GetNode(GroomUid);
+	if (const UInterchangeGroomNode* GroomNode = Cast<const UInterchangeGroomNode>(Node))
+	{
+		TOptional<FInterchangeGroomPayloadKey> OptionalPayloadKey = GroomNode->GetPayloadKey();
+		if (!ensure(OptionalPayloadKey.IsSet()))
+		{
+			UE_LOG(LogInterchangeImport, Warning, TEXT("Empty groom reference payload when importing Groom asset %s."), *Arguments.AssetName);
+			return;
+		}
+
+		FInterchangeGroomPayloadKey& PayloadKey = OptionalPayloadKey.GetValue();
+
+		TSharedPtr<UE::Interchange::FInterchangeTaskLambda, ESPMode::ThreadSafe> TaskGetMeshPayload = MakeShared<UE::Interchange::FInterchangeTaskLambda, ESPMode::ThreadSafe>(
+			bAsync ? UE::Interchange::EInterchangeTaskThread::AsyncThread : UE::Interchange::EInterchangeTaskThread::GameThread,
+			[&GroomPayload = GroomPayload, GroomPayloadInterface, PayloadKey]()
+			{
+				TRACE_CPUPROFILER_EVENT_SCOPE(UInterchangeGroomFactory::GetGroomPayloadDataTask);
+				GroomPayload = GroomPayloadInterface->GetGroomPayloadData(PayloadKey);
+			});
+
+		PayloadTasks.Add(TaskGetMeshPayload);
+	}
+}
+
+UInterchangeFactoryBase::FImportAssetResult UInterchangeGroomFactory::BeginImportAsset_GameThread(const FImportAssetObjectParams& Arguments)
+{
+	TRACE_CPUPROFILER_EVENT_SCOPE(UInterchangeGroomFactory::BeginImportAsset_GameThread);
+
+	FImportAssetResult ImportAssetResult;
+
+#if WITH_EDITOR
+	auto LogGroomFactoryError = [this, &Arguments, &ImportAssetResult](const FText& Info)
+		{
+			UInterchangeResultError_Generic* Message = AddMessage<UInterchangeResultError_Generic>();
+			Message->SourceAssetName = Arguments.SourceData->GetFilename();
+			Message->DestinationAssetName = Arguments.AssetName;
+			Message->AssetType = GetFactoryClass();
+			Message->Text = FText::Format(
+				LOCTEXT("GroomFactory_Failure", "UInterchangeGroomFactory: Could not create Groom asset '{0}'. Reason: {1}"),
+				FText::FromString(Arguments.AssetName),
+				Info
+			);
+			ImportAssetResult.bIsFactorySkipAsset = true;
+		};
+
+	const IInterchangeGroomPayloadInterface* GroomPayloadInterface = Cast<IInterchangeGroomPayloadInterface>(Arguments.Translator);
+	if (!GroomPayloadInterface)
+	{
+		LogGroomFactoryError(LOCTEXT("GroomFactory_NoPayloadInterface", "The translator does not implement IInterchangeGroomPayloadInterface"));
+		return ImportAssetResult;
+	}
+
+	UInterchangeGroomFactoryNode* GroomFactoryNode = Cast<UInterchangeGroomFactoryNode>(Arguments.AssetNode);
+	if (!GroomFactoryNode)
+	{
+		LogGroomFactoryError(LOCTEXT("GroomFactory_AssetNodeNull", "Asset node parameter is not a UInterchangeGroomFactoryNode."));
+		return ImportAssetResult;
+	}
+
+	const UClass* ObjectClass = GroomFactoryNode->GetObjectClass();
+	if (!ObjectClass || !ObjectClass->IsChildOf(GetFactoryClass()))
+	{
+		LogGroomFactoryError(LOCTEXT("GroomFactory_NodeClassMissmatch", "Asset node parameter class doesn't derive the UGroomAsset class."));
+		return ImportAssetResult;
+	}
+
+	if (!GroomPayload.IsSet())
+	{
+		LogGroomFactoryError(LOCTEXT("GroomFactory_NoPayload", "No groom payload could be translated from the source."));
+		return ImportAssetResult;
+	}
+
+	FHairDescription& HairDescription = GroomPayload->HairDescription;
+	if (!HairDescription.IsValid())
+	{
+		LogGroomFactoryError(LOCTEXT("GroomFactory_InvalidHairDescription", "No valid hair description found in groom payload."));
+		return ImportAssetResult;
+	}
+
+	UObject* ExistingAsset = Arguments.ReimportObject;
+	if (!ExistingAsset)
+	{
+		FSoftObjectPath ReferenceObject;
+		if (GroomFactoryNode->GetCustomReferenceObject(ReferenceObject))
+		{
+			ExistingAsset = ReferenceObject.TryLoad();
+		}
+	}
+
+	UGroomAsset* GroomAsset = Cast<UGroomAsset>(ExistingAsset);
+	if (!GroomAsset)
+	{
+		GroomAsset = NewObject<UGroomAsset>(Arguments.Parent, *Arguments.AssetName, RF_Public | RF_Standalone);
+	}
+	if (!GroomAsset)
+	{
+		LogGroomFactoryError(LOCTEXT("GroomFactory_GroomCreationFailure", "Could not allocate new GroomAsset."));
+		return ImportAssetResult;
+	}
+
+	// Extract the groom groups info from the hair description to get the number of groups
+	FHairDescriptionGroups GroupsDescription;
+	FGroomBuilder::BuildHairDescriptionGroups(HairDescription, GroupsDescription);
+
+	TArray<FHairGroupsInterpolation> GroomInterpolationSettings;
+	GroomInterpolationSettings.SetNum(GroupsDescription.HairGroups.Num());
+	UGroomImportOptions* ImportOptions = CreateGroomImportOptions(GroupsDescription, GroomInterpolationSettings);
+
+	FHairImportContext HairImportContext(ImportOptions);
+	UGroomAsset* ImportedGroomAsset = FHairStrandsImporter::ImportHair(HairImportContext, HairDescription, GroomAsset);
+	if (ImportedGroomAsset)
+	{
+		ImportAssetResult.ImportedObject = ImportedGroomAsset;
+	}
+	else
+	{
+		GroomAsset->MarkAsGarbage();
+	}
+#endif // WITH_EDITOR
+
+	return ImportAssetResult;
+}
+
+void UInterchangeGroomFactory::SetupObject_GameThread(const FSetupObjectParams& Arguments)
+{
+	TRACE_CPUPROFILER_EVENT_SCOPE(UInterchangeGeometryCacheFactory::SetupObject_GameThread);
+
+	Super::SetupObject_GameThread(Arguments);
+
+#if WITH_EDITORONLY_DATA
+	if (ensure(Arguments.ImportedObject && Arguments.SourceData))
+	{
+		// We must call the Update of the asset source file in the main thread because UAssetImportData::Update execute some delegate we do not control
+		UGroomAsset* GroomAsset = CastChecked<UGroomAsset>(Arguments.ImportedObject);
+
+		UE::Interchange::FFactoryCommon::FUpdateImportAssetDataParameters UpdateImportAssetDataParameters(GroomAsset, GroomAsset->AssetImportData, Arguments.SourceData, Arguments.NodeUniqueID, Arguments.NodeContainer, Arguments.OriginalPipelines, Arguments.Translator);
+		GroomAsset->AssetImportData = UE::Interchange::FFactoryCommon::UpdateImportAssetData(UpdateImportAssetDataParameters);
+	}
+#endif
+}
+
+bool UInterchangeGroomFactory::GetSourceFilenames(const UObject* Object, TArray<FString>& OutSourceFilenames) const
+{
+#if WITH_EDITORONLY_DATA
+	if (const UGroomAsset* GroomAsset = Cast<UGroomAsset>(Object))
+	{
+		return UE::Interchange::FFactoryCommon::GetSourceFilenames(GroomAsset->AssetImportData, OutSourceFilenames);
+	}
+#endif
+
+	return false;
+}
+
+bool UInterchangeGroomFactory::SetSourceFilename(const UObject* Object, const FString& SourceFilename, int32 SourceIndex) const
+{
+#if WITH_EDITORONLY_DATA
+	if (const UGroomAsset* GroomAsset = Cast<UGroomAsset>(Object))
+	{
+		return UE::Interchange::FFactoryCommon::SetSourceFilename(GroomAsset->AssetImportData, SourceFilename, SourceIndex);
+	}
+#endif
+
+	return false;
+}
+
+void UInterchangeGroomFactory::BackupSourceData(const UObject* Object) const
+{
+#if WITH_EDITORONLY_DATA
+	if (const UGroomAsset* GroomAsset = Cast<UGroomAsset>(Object))
+	{
+		UE::Interchange::FFactoryCommon::BackupSourceData(GroomAsset->AssetImportData);
+	}
+#endif
+}
+
+void UInterchangeGroomFactory::ReinstateSourceData(const UObject* Object) const
+{
+#if WITH_EDITORONLY_DATA
+	if (const UGroomAsset* GroomAsset = Cast<UGroomAsset>(Object))
+	{
+		UE::Interchange::FFactoryCommon::ReinstateSourceData(GroomAsset->AssetImportData);
+	}
+#endif
+}
+
+void UInterchangeGroomFactory::ClearBackupSourceData(const UObject* Object) const
+{
+#if WITH_EDITORONLY_DATA
+	if (const UGroomAsset* GroomAsset = Cast<UGroomAsset>(Object))
+	{
+		UE::Interchange::FFactoryCommon::ClearBackupSourceData(GroomAsset->AssetImportData);
+	}
+#endif
+}
+
+#undef LOCTEXT_NAMESPACE
diff --git a/Engine/Plugins/Interchange/Runtime/Source/Import/Private/InterchangeImportModule.cpp b/Engine/Plugins/Interchange/Runtime/Source/Import/Private/InterchangeImportModule.cpp
--- a/Engine/Plugins/Interchange/Runtime/Source/Import/Private/InterchangeImportModule.cpp
+++ b/Engine/Plugins/Interchange/Runtime/Source/Import/Private/InterchangeImportModule.cpp
@@ -9,6 +9,7 @@
 #include "Engine/RendererSettings.h"
 #include "Fbx/InterchangeFbxTranslator.h"
 #include "Gltf/InterchangeGltfTranslator.h"
+#include "Groom/InterchangeGroomFactory.h"
 #include "InterchangeImportLog.h"
 #include "InterchangeManager.h"
 #include "Material/InterchangeMaterialFactory.h"
@@ -109,6 +110,7 @@ void FInterchangeImportModule::StartupModule()
 		InterchangeManager.RegisterFactory(UInterchangeSkeletalMeshFactory::StaticClass());
 		InterchangeManager.RegisterFactory(UInterchangeStaticMeshFactory::StaticClass());
 		InterchangeManager.RegisterFactory(UInterchangeGeometryCacheFactory::StaticClass());
+		InterchangeManager.RegisterFactory(UInterchangeGroomFactory::StaticClass());
 		InterchangeManager.RegisterFactory(UInterchangePhysicsAssetFactory::StaticClass());
 		InterchangeManager.RegisterFactory(UInterchangeActorFactory::StaticClass());
 		InterchangeManager.RegisterFactory(UInterchangeLevelFactory::StaticClass());
diff --git a/Engine/Plugins/Interchange/Runtime/Source/Import/Public/Groom/InterchangeGroomFactory.h b/Engine/Plugins/Interchange/Runtime/Source/Import/Public/Groom/InterchangeGroomFactory.h
new file mode 100644
--- /dev/null
+++ b/Engine/Plugins/Interchange/Runtime/Source/Import/Public/Groom/InterchangeGroomFactory.h
@@ -0,0 +1,36 @@
+// Copyright Epic Games, Inc. All Rights Reserved.
+
+#pragma once
+
+#include "Groom/InterchangeGroomPayloadData.h"
+#include "InterchangeFactoryBase.h"
+#include "UObject/ObjectMacros.h"
+
+#include "InterchangeGroomFactory.generated.h"
+
+#define UE_API INTERCHANGEIMPORT_API
+
+UCLASS(MinimalAPI, BlueprintType)
+class UInterchangeGroomFactory : public UInterchangeFactoryBase
+{
+	GENERATED_BODY()
+
+public:
+	// Begin UInterchangeFactoryBase interface
+	UE_API virtual UClass* GetFactoryClass() const override;
+	virtual EInterchangeFactoryAssetType GetFactoryAssetType() override	{ return EInterchangeFactoryAssetType::Grooms; }
+	UE_API virtual void CreatePayloadTasks(const FImportAssetObjectParams& Arguments, bool bAsync, TArray<TSharedPtr<UE::Interchange::FInterchangeTaskBase>>& PayloadTasks) override;
+	UE_API virtual FImportAssetResult BeginImportAsset_GameThread(const FImportAssetObjectParams& Arguments) override;
+	UE_API virtual void SetupObject_GameThread(const FSetupObjectParams& Arguments) override;
+	UE_API virtual bool GetSourceFilenames(const UObject* Object, TArray<FString>& OutSourceFilenames) const override;
+	UE_API virtual bool SetSourceFilename(const UObject* Object, const FString& SourceFilename, int32 SourceIndex) const override;
+	UE_API virtual void BackupSourceData(const UObject* Object) const override;
+	UE_API virtual void ReinstateSourceData(const UObject* Object) const override;
+	UE_API virtual void ClearBackupSourceData(const UObject* Object) const override;
+	// End UInterchangeFactoryBase interface
+
+private:
+	TOptional<UE::Interchange::FGroomPayloadData> GroomPayload;
+};
+
+#undef UE_API
\ No newline at end of file
diff --git a/Engine/Plugins/Interchange/Runtime/Source/Import/Public/Groom/InterchangeGroomPayloadData.h b/Engine/Plugins/Interchange/Runtime/Source/Import/Public/Groom/InterchangeGroomPayloadData.h
new file mode 100644
--- /dev/null
+++ b/Engine/Plugins/Interchange/Runtime/Source/Import/Public/Groom/InterchangeGroomPayloadData.h
@@ -0,0 +1,13 @@
+// Copyright Epic Games, Inc. All Rights Reserved.
+
+#pragma once
+
+#include "HairDescription.h"
+
+namespace UE::Interchange
+{
+	struct FGroomPayloadData
+	{
+		FHairDescription HairDescription;
+	};
+}
\ No newline at end of file
diff --git a/Engine/Plugins/Interchange/Runtime/Source/Import/Public/Groom/InterchangeGroomPayloadInterface.h b/Engine/Plugins/Interchange/Runtime/Source/Import/Public/Groom/InterchangeGroomPayloadInterface.h
new file mode 100644
--- /dev/null
+++ b/Engine/Plugins/Interchange/Runtime/Source/Import/Public/Groom/InterchangeGroomPayloadInterface.h
@@ -0,0 +1,31 @@
+// Copyright Epic Games, Inc. All Rights Reserved.
+
+#pragma once
+
+#include "Groom/InterchangeGroomPayloadData.h"
+#include "InterchangeGroomNode.h"
+#include "UObject/Interface.h"
+#include "UObject/ObjectMacros.h"
+
+#include "InterchangeGroomPayloadInterface.generated.h"
+
+UINTERFACE(MinimalAPI)
+class UInterchangeGroomPayloadInterface : public UInterface
+{
+	GENERATED_BODY()
+};
+
+/**
+ * Payload interface for groom data
+ */
+class IInterchangeGroomPayloadInterface
+{
+	GENERATED_BODY()
+
+public:
+	virtual TOptional<UE::Interchange::FGroomPayloadData> GetGroomPayloadData(const FInterchangeGroomPayloadKey& PayloadKey) const
+	{
+		return {};
+	}
+};
+
diff --git a/Engine/Plugins/Interchange/Runtime/Source/Nodes/Private/InterchangeGroomNode.cpp b/Engine/Plugins/Interchange/Runtime/Source/Nodes/Private/InterchangeGroomNode.cpp
new file mode 100644
--- /dev/null
+++ b/Engine/Plugins/Interchange/Runtime/Source/Nodes/Private/InterchangeGroomNode.cpp
@@ -0,0 +1,88 @@
+// Copyright Epic Games, Inc. All Rights Reserved.
+
+#include "InterchangeGroomNode.h"
+
+#include UE_INLINE_GENERATED_CPP_BY_NAME(InterchangeGroomNode)
+
+namespace UE::Interchange
+{
+	const FAttributeKey& FGroomNodeStaticData::PayloadKey()
+	{
+		static FAttributeKey AttributeKey(TEXT("__PayloadKey__"));
+		return AttributeKey;
+	}
+
+	const FAttributeKey& FGroomNodeStaticData::PayloadTypeKey()
+	{
+		static FAttributeKey AttributeKey(TEXT("__PayloadTypeKey__"));
+		return AttributeKey;
+	}
+}
+
+FString UInterchangeGroomNode::GetTypeName() const
+{
+	const FString TypeName = TEXT("GroomNode");
+	return TypeName;
+}
+
+const TOptional<FInterchangeGroomPayloadKey> UInterchangeGroomNode::GetPayloadKey() const
+{
+	FString PayloadKey;
+	EInterchangeGroomPayLoadType PayloadType;
+
+	//PayLoadKey
+	{
+		if (!Attributes->ContainAttribute(UE::Interchange::FGroomNodeStaticData::PayloadKey()))
+		{
+			return {};
+		}
+		UE::Interchange::FAttributeStorage::TAttributeHandle<FString> AttributeHandle = Attributes->GetAttributeHandle<FString>(UE::Interchange::FGroomNodeStaticData::PayloadKey());
+		if (!AttributeHandle.IsValid())
+		{
+			return {};
+		}
+		UE::Interchange::EAttributeStorageResult Result = AttributeHandle.Get(PayloadKey);
+		if (!IsAttributeStorageResultSuccess(Result))
+		{
+			LogAttributeStorageErrors(Result, TEXT("UInterchangeGroomNode.GetPayloadKey"), UE::Interchange::FGroomNodeStaticData::PayloadKey());
+			return {};
+		}
+	}
+
+	//PayLoadType
+	{
+		if (!Attributes->ContainAttribute(UE::Interchange::FGroomNodeStaticData::PayloadTypeKey()))
+		{
+			return {};
+		}
+		UE::Interchange::FAttributeStorage::TAttributeHandle<EInterchangeGroomPayLoadType> AttributeHandle = Attributes->GetAttributeHandle<EInterchangeGroomPayLoadType>(UE::Interchange::FGroomNodeStaticData::PayloadTypeKey());
+		if (!AttributeHandle.IsValid())
+		{
+			return {};
+		}
+
+		UE::Interchange::EAttributeStorageResult Result = AttributeHandle.Get(PayloadType);
+		if (!IsAttributeStorageResultSuccess(Result))
+		{
+			LogAttributeStorageErrors(Result, TEXT("UInterchangeGroomNode.GetPayloadTypeKey"), UE::Interchange::FGroomNodeStaticData::PayloadTypeKey());
+			return {};
+		}
+	}
+
+	return TOptional<FInterchangeGroomPayloadKey>(FInterchangeGroomPayloadKey(PayloadKey, PayloadType));
+}
+
+void UInterchangeGroomNode::SetPayloadKey(const FString& PayloadKey, EInterchangeGroomPayLoadType PayLoadType)
+{
+	UE::Interchange::EAttributeStorageResult Result = Attributes->RegisterAttribute(UE::Interchange::FGroomNodeStaticData::PayloadKey(), PayloadKey);
+	if (!IsAttributeStorageResultSuccess(Result))
+	{
+		LogAttributeStorageErrors(Result, TEXT("UInterchangeGroomNode.SetPayloadKey"), UE::Interchange::FGroomNodeStaticData::PayloadKey());
+		return;
+	}
+	Result = Attributes->RegisterAttribute(UE::Interchange::FGroomNodeStaticData::PayloadTypeKey(), PayLoadType);
+	if (!IsAttributeStorageResultSuccess(Result))
+	{
+		LogAttributeStorageErrors(Result, TEXT("UInterchangeGroomNode.SetPayloadTypeKey"), UE::Interchange::FGroomNodeStaticData::PayloadTypeKey());
+	}
+}
diff --git a/Engine/Plugins/Interchange/Runtime/Source/Nodes/Public/InterchangeGroomNode.h b/Engine/Plugins/Interchange/Runtime/Source/Nodes/Public/InterchangeGroomNode.h
new file mode 100644
--- /dev/null
+++ b/Engine/Plugins/Interchange/Runtime/Source/Nodes/Public/InterchangeGroomNode.h
@@ -0,0 +1,83 @@
+// Copyright Epic Games, Inc. All Rights Reserved.
+
+#pragma once
+
+#include "Nodes/InterchangeBaseNode.h"
+
+#include "InterchangeGroomNode.generated.h"
+
+#define UE_API INTERCHANGENODES_API
+
+namespace UE::Interchange
+{
+	struct FGroomNodeStaticData : public FBaseNodeStaticData
+	{
+		static UE_API const FAttributeKey& PayloadKey();
+		static UE_API const FAttributeKey& PayloadTypeKey();
+	};
+}
+
+UENUM(BlueprintType)
+enum class EInterchangeGroomPayLoadType : uint8
+{
+	STATIC		UMETA(DisplayName = "Groom Asset"),
+	ANIMATED	UMETA(DisplayName = "Groom Cache")
+};
+
+USTRUCT(BlueprintType)
+struct FInterchangeGroomPayloadKey
+{
+	GENERATED_BODY()
+
+public:
+
+	UPROPERTY(EditAnywhere, BlueprintReadWrite, Category = "Interchange | Groom")
+	FString UniqueId = "";
+
+	UPROPERTY(EditAnywhere, BlueprintReadWrite, Category = "Interchange | Groom")
+	EInterchangeGroomPayLoadType Type = EInterchangeGroomPayLoadType::STATIC;
+
+	UPROPERTY(EditAnywhere, BlueprintReadWrite, Category = "Interchange | Groom Cache")
+	int32 FrameNumber = 0;
+
+	FInterchangeGroomPayloadKey() {}
+
+	FInterchangeGroomPayloadKey(const FString& InUniqueId, const EInterchangeGroomPayLoadType& InType)
+		: UniqueId(InUniqueId)
+		, Type(InType)
+	{
+	}
+
+	FInterchangeGroomPayloadKey(const FString& InUniqueId, int32 InFrameNumber)
+		: UniqueId(InUniqueId)
+		, Type(EInterchangeGroomPayLoadType::ANIMATED)
+		, FrameNumber(InFrameNumber)
+	{
+	}
+
+	bool operator==(const FInterchangeGroomPayloadKey& Other) const
+	{
+		return UniqueId.Equals(Other.UniqueId) && Type == Other.Type && FrameNumber == Other.FrameNumber;
+	}
+
+	friend uint32 GetTypeHash(const FInterchangeGroomPayloadKey& InterchangeGroomPayLoadKey)
+	{
+		return GetTypeHash(InterchangeGroomPayLoadKey.UniqueId + FString::FromInt(static_cast<int32>(InterchangeGroomPayLoadKey.Type)) + FString::FromInt(InterchangeGroomPayLoadKey.FrameNumber));
+	}
+};
+
+UCLASS(MinimalAPI, BlueprintType)
+class UInterchangeGroomNode : public UInterchangeBaseNode
+{
+	GENERATED_BODY()
+
+public:
+	UE_API virtual FString GetTypeName() const override;
+
+	UE_API const TOptional<FInterchangeGroomPayloadKey> GetPayloadKey() const;
+
+	UFUNCTION(BlueprintCallable, Category = "Interchange | Node | Groom")
+	UE_API void SetPayloadKey(const FString& PayloadKey, EInterchangeGroomPayLoadType PayLoadType = EInterchangeGroomPayLoadType::STATIC);
+};
+
+#undef UE_API
\ No newline at end of file
diff --git a/Engine/Plugins/Interchange/Runtime/Source/Pipelines/InterchangePipelines.Build.cs b/Engine/Plugins/Interchange/Runtime/Source/Pipelines/InterchangePipelines.Build.cs
--- a/Engine/Plugins/Interchange/Runtime/Source/Pipelines/InterchangePipelines.Build.cs
+++ b/Engine/Plugins/Interchange/Runtime/Source/Pipelines/InterchangePipelines.Build.cs
@@ -29,6 +29,7 @@ namespace UnrealBuildTool.Rules
 					"Foliage",
 					"GeometryCache",
 					"GLTFCore",
+					"HairStrandsCore",
 					"InputCore",
 					"InterchangeImport",
 					"LevelSequence",
diff --git a/Engine/Plugins/Interchange/Runtime/Source/Pipelines/Private/InterchangeGenericAssetsPipeline.cpp b/Engine/Plugins/Interchange/Runtime/Source/Pipelines/Private/InterchangeGenericAssetsPipeline.cpp
--- a/Engine/Plugins/Interchange/Runtime/Source/Pipelines/Private/InterchangeGenericAssetsPipeline.cpp
+++ b/Engine/Plugins/Interchange/Runtime/Source/Pipelines/Private/InterchangeGenericAssetsPipeline.cpp
@@ -19,6 +19,7 @@
 #include "InterchangeBlueprintPipelineBase.h"
 #include "InterchangeCommonPipelineDataFactoryNode.h"
 #include "InterchangeGenericAnimationPipeline.h"
+#include "InterchangeGenericGroomPipeline.h"
 #include "InterchangeGenericMaterialPipeline.h"
 #include "InterchangeGenericMeshPipeline.h"
 #include "InterchangeGenericTexturePipeline.h"
@@ -76,6 +77,7 @@ UInterchangeGenericAssetsPipeline::UInterchangeGenericAssetsPipeline()
 	MeshPipeline = CreateDefaultSubobject<UInterchangeGenericMeshPipeline>("MeshPipeline");
 	MeshPipeline->CommonMeshesProperties = CommonMeshesProperties;
 	MeshPipeline->CommonSkeletalMeshesAndAnimationsProperties = CommonSkeletalMeshesAndAnimationsProperties;
+	GroomPipeline = CreateDefaultSubobject<UInterchangeGenericGroomPipeline>("GroomPipeline");
 	AnimationPipeline = CreateDefaultSubobject<UInterchangeGenericAnimationPipeline>("AnimationPipeline");
 	AnimationPipeline->CommonSkeletalMeshesAndAnimationsProperties = CommonSkeletalMeshesAndAnimationsProperties;
 	AnimationPipeline->CommonMeshesProperties = CommonMeshesProperties;
@@ -107,6 +109,11 @@ void UInterchangeGenericAssetsPipeline::PreDialogCleanup(const FName PipelineSta
 		MeshPipeline->PreDialogCleanup(PipelineStackName);
 	}
 
+	if (GroomPipeline)
+	{
+		GroomPipeline->PreDialogCleanup(PipelineStackName);
+	}
+
 	if (AnimationPipeline)
 	{
 		AnimationPipeline->PreDialogCleanup(PipelineStackName);
@@ -187,6 +194,11 @@ void UInterchangeGenericAssetsPipeline::AdjustSettingsForContext(const FIntercha
 		MeshPipeline->AdjustSettingsForContext(ContextParams);
 	}
 
+	if (GroomPipeline)
+	{
+		GroomPipeline->AdjustSettingsForContext(ContextParams);
+	}
+
 	if (AnimationPipeline)
 	{
 		AnimationPipeline->AdjustSettingsForContext(ContextParams);
@@ -394,6 +406,11 @@ void UInterchangeGenericAssetsPipeline::FilterPropertiesFromTranslatedData(UInte
 			HideFullCategory(UInterchangeGenericAnimationPipeline::GetPipelineCategory(nullptr), AnimationPipeline);
 		}
 	}
+
+	if (GroomPipeline)
+	{
+		GroomPipeline->FilterPropertiesFromTranslatedData(InBaseNodeContainer);
+	}
 }
 
 bool UInterchangeGenericAssetsPipeline::IsPropertyChangeNeedRefresh(const FPropertyChangedEvent& PropertyChangedEvent) const
@@ -428,6 +445,10 @@ void UInterchangeGenericAssetsPipeline::GetSupportAssetClasses(TArray<UClass*>&
 	{
 		MeshPipeline->GetSupportAssetClasses(PipelineSupportAssetClasses);
 	}
+	if (GroomPipeline)
+	{
+		GroomPipeline->GetSupportAssetClasses(PipelineSupportAssetClasses);
+	}
 	if (MaterialPipeline)
 	{
 		MaterialPipeline->GetSupportAssetClasses(PipelineSupportAssetClasses);
@@ -456,7 +477,15 @@ bool UInterchangeGenericAssetsPipeline::GetPropertyPossibleValues(const FName Pr
 			return true;
 		}
 	}
-	
+
+	if (GroomPipeline && PropertyPathString.StartsWith(UInterchangeGenericGroomPipeline::StaticClass()->GetPathName()))
+	{
+		if (GroomPipeline->GetPropertyPossibleValues(PropertyPath, PossibleValues))
+		{
+			return true;
+		}
+	}
+
 	if (AnimationPipeline && PropertyPathString.StartsWith(UInterchangeGenericAnimationPipeline::StaticClass()->GetPathName()))
 	{
 		if (AnimationPipeline->GetPropertyPossibleValues(PropertyPath, PossibleValues))
@@ -1202,6 +1231,10 @@ void UInterchangeGenericAssetsPipeline::ExecutePipeline(UInterchangeBaseNodeCont
 		{
 			MeshPipeline->SetResultsContainer(Results);
 		}
+		if (GroomPipeline)
+		{
+			GroomPipeline->SetResultsContainer(Results);
+		}
 		if (AnimationPipeline)
 		{
 			AnimationPipeline->SetResultsContainer(Results);
@@ -1264,6 +1297,10 @@ void UInterchangeGenericAssetsPipeline::ExecutePipeline(UInterchangeBaseNodeCont
 	{
 		MeshPipeline->ScriptedExecutePipeline(InBaseNodeContainer, InSourceDatas, ContentBasePath);
 	}
+	if (GroomPipeline)
+	{
+		GroomPipeline->ScriptedExecutePipeline(InBaseNodeContainer, InSourceDatas, ContentBasePath);
+	}
 	if (AnimationPipeline)
 	{
 		AnimationPipeline->ScriptedExecutePipeline(InBaseNodeContainer, InSourceDatas, ContentBasePath);
@@ -1315,6 +1352,10 @@ void UInterchangeGenericAssetsPipeline::ExecutePostFactoryPipeline(const UInterc
 	{
 		MeshPipeline->ScriptedExecutePostFactoryPipeline(InBaseNodeContainer, NodeKey, CreatedAsset, bIsAReimport);
 	}
+	if (GroomPipeline)
+	{
+		GroomPipeline->ScriptedExecutePostFactoryPipeline(InBaseNodeContainer, NodeKey, CreatedAsset, bIsAReimport);
+	}
 	if (AnimationPipeline)
 	{
 		AnimationPipeline->ScriptedExecutePostFactoryPipeline(InBaseNodeContainer, NodeKey, CreatedAsset, bIsAReimport);
@@ -1331,6 +1372,10 @@ void UInterchangeGenericAssetsPipeline::ExecutePostImportPipeline(const UInterch
 	{
 		MeshPipeline->ScriptedExecutePostImportPipeline(InBaseNodeContainer, NodeKey, CreatedAsset, bIsAReimport);
 	}
+	if (GroomPipeline)
+	{
+		GroomPipeline->ScriptedExecutePostImportPipeline(InBaseNodeContainer, NodeKey, CreatedAsset, bIsAReimport);
+	}
 	if (AnimationPipeline)
 	{
 		AnimationPipeline->ScriptedExecutePostImportPipeline(InBaseNodeContainer, NodeKey, CreatedAsset, bIsAReimport);
@@ -1821,6 +1866,11 @@ void UInterchangeGenericAssetsPipeline::SetReimportSourceIndex(UClass* ReimportO
 		MeshPipeline->ScriptedSetReimportSourceIndex(ReimportObjectClass, SourceFileIndex);
 	}
 
+	if (GroomPipeline)
+	{
+		GroomPipeline->ScriptedSetReimportSourceIndex(ReimportObjectClass, SourceFileIndex);
+	}
+
 	if (AnimationPipeline)
 	{
 		AnimationPipeline->ScriptedSetReimportSourceIndex(ReimportObjectClass, SourceFileIndex);
diff --git a/Engine/Plugins/Interchange/Runtime/Source/Pipelines/Private/InterchangeGenericGroomPipeline.cpp b/Engine/Plugins/Interchange/Runtime/Source/Pipelines/Private/InterchangeGenericGroomPipeline.cpp
new file mode 100644
--- /dev/null
+++ b/Engine/Plugins/Interchange/Runtime/Source/Pipelines/Private/InterchangeGenericGroomPipeline.cpp
@@ -0,0 +1,73 @@
+// Copyright Epic Games, Inc. All Rights Reserved.
+
+#include "InterchangeGenericGroomPipeline.h"
+
+#include "GroomAsset.h"
+#include "InterchangeGroomNode.h"
+#include "InterchangeGroomFactoryNode.h"
+#include "InterchangePipelineHelper.h"
+#include "Nodes/InterchangeFactoryBaseNode.h"
+#include "Nodes/InterchangeSourceNode.h"
+#include "Nodes/InterchangeUserDefinedAttribute.h"
+
+#include UE_INLINE_GENERATED_CPP_BY_NAME(InterchangeGenericGroomPipeline)
+
+FString UInterchangeGenericGroomPipeline::GetPipelineCategory(UClass* AssetClass)
+{
+	return TEXT("Grooms");
+}
+
+void UInterchangeGenericGroomPipeline::ExecutePipeline(UInterchangeBaseNodeContainer* InBaseNodeContainer, const TArray<UInterchangeSourceData*>& SourceDatas, const FString& ContentBasePath)
+{
+	if (!bEnableGroomTypesImport)
+	{
+		return;
+	}
+
+	if (!InBaseNodeContainer)
+	{
+		return;
+	}
+
+	// Find all the translated groom nodes we need for this pipeline
+	TArray<UInterchangeGroomNode*> GroomNodes;
+	if (bImportGrooms)
+	{
+		InBaseNodeContainer->IterateNodes(
+			[&GroomNodes](const FString& NodeUid, UInterchangeBaseNode* Node)
+			{
+				if (UInterchangeGroomNode* GroomNode = Cast<UInterchangeGroomNode>(Node))
+				{
+					GroomNodes.Add(GroomNode);
+				}
+			}
+		);
+
+		// Create groom factory nodes
+		for (const UInterchangeGroomNode* GroomNode : GroomNodes)
+		{
+			const FString FactoryNodeUid = UInterchangeFactoryBaseNode::BuildFactoryNodeUid(GroomNode->GetUniqueID());
+
+			UInterchangeGroomFactoryNode* GroomFactoryNode = NewObject<UInterchangeGroomFactoryNode>(InBaseNodeContainer);
+			InBaseNodeContainer->SetupNode(GroomFactoryNode, FactoryNodeUid, GroomNode->GetDisplayLabel(), EInterchangeNodeContainerType::FactoryData);
+
+			UInterchangeSourceNode* SourceNode = UInterchangeSourceNode::FindOrCreateUniqueInstance(InBaseNodeContainer);
+			UE::Interchange::PipelineHelper::FillSubPathFromSourceNode(GroomFactoryNode, SourceNode);
+
+			const bool bAddSourceNodeName = false;
+			UInterchangeUserDefinedAttributesAPI::DuplicateAllUserDefinedAttribute(GroomNode, GroomFactoryNode, bAddSourceNodeName);
+
+			GroomFactoryNode->AddTargetNodeUid(GroomNode->GetUniqueID());
+			GroomNode->AddTargetNodeUid(GroomFactoryNode->GetUniqueID());
+		}
+	}
+}
+
+#if WITH_EDITOR
+
+void UInterchangeGenericGroomPipeline::GetSupportAssetClasses(TArray<UClass*>& PipelineSupportAssetClasses) const
+{
+	PipelineSupportAssetClasses.Add(UGroomAsset::StaticClass());
+}
+
+#endif // WITH_EDITOR
\ No newline at end of file
diff --git a/Engine/Plugins/Interchange/Runtime/Source/Pipelines/Private/InterchangePipelineHelper.cpp b/Engine/Plugins/Interchange/Runtime/Source/Pipelines/Private/InterchangePipelineHelper.cpp
--- a/Engine/Plugins/Interchange/Runtime/Source/Pipelines/Private/InterchangePipelineHelper.cpp
+++ b/Engine/Plugins/Interchange/Runtime/Source/Pipelines/Private/InterchangePipelineHelper.cpp
@@ -12,6 +12,7 @@
 #include "Framework/Application/SlateApplication.h"
 #include "Framework/Docking/TabManager.h"
 #include "GeometryCache.h"
+#include "GroomAsset.h"
 #include "FoliageType_InstancedStaticMesh.h"
 #include "HAL/PlatformApplicationMisc.h"
 #include "InterchangePipelineBase.h"
@@ -124,6 +125,10 @@ namespace UE::Interchange::PipelineHelper
 			{
 				Suffix = TEXT("GeometryCaches");
 			}
+			else if (Class->IsChildOf(UGroomAsset::StaticClass()))
+			{
+				Suffix = TEXT("Grooms");
+			}
 			else if (Class->IsChildOf(UWorld::StaticClass()))
 			{
 				Suffix = TEXT("Levels");
diff --git a/Engine/Plugins/Interchange/Runtime/Source/Pipelines/Public/InterchangeGenericAssetsPipeline.h b/Engine/Plugins/Interchange/Runtime/Source/Pipelines/Public/InterchangeGenericAssetsPipeline.h
--- a/Engine/Plugins/Interchange/Runtime/Source/Pipelines/Public/InterchangeGenericAssetsPipeline.h
+++ b/Engine/Plugins/Interchange/Runtime/Source/Pipelines/Public/InterchangeGenericAssetsPipeline.h
@@ -22,6 +22,7 @@
 class STableViewBase;
 class ITableRow;
 class UInterchangeGenericAnimationPipeline;
+class UInterchangeGenericGroomPipeline;
 class UInterchangeGenericMaterialPipeline;
 class UInterchangeGenericMeshPipeline;
 class UInterchangeGenericTexturePipeline;
@@ -103,6 +104,10 @@ public:
 	UPROPERTY(VisibleAnywhere, BlueprintReadWrite, Instanced, Category = "Materials")
 	TObjectPtr<UInterchangeGenericMaterialPipeline> MaterialPipeline;
 
+	//////	GROOMS_CATEGORY Properties //////
+	UPROPERTY(VisibleAnywhere, BlueprintReadWrite, Instanced, Category = "Grooms")
+	TObjectPtr<UInterchangeGenericGroomPipeline> GroomPipeline;
+
 	UE_API virtual void PreDialogCleanup(const FName PipelineStackName) override;
 
 	UE_API virtual bool IsSettingsAreValid(TOptional<FText>& OutInvalidReason) const override;
diff --git a/Engine/Plugins/Interchange/Runtime/Source/Pipelines/Public/InterchangeGenericGroomPipeline.h b/Engine/Plugins/Interchange/Runtime/Source/Pipelines/Public/InterchangeGenericGroomPipeline.h
new file mode 100644
--- /dev/null
+++ b/Engine/Plugins/Interchange/Runtime/Source/Pipelines/Public/InterchangeGenericGroomPipeline.h
@@ -0,0 +1,43 @@
+// Copyright Epic Games, Inc. All Rights Reserved.
+
+#pragma once
+
+#include "InterchangePipelineBase.h"
+#include "Nodes/InterchangeBaseNodeContainer.h"
+#include "UObject/ObjectMacros.h"
+
+#include "InterchangeGenericGroomPipeline.generated.h"
+
+#define UE_API INTERCHANGEPIPELINES_API
+
+UCLASS(MinimalAPI, BlueprintType, editinlinenew)
+class UInterchangeGenericGroomPipeline : public UInterchangePipelineBase
+{
+	GENERATED_BODY()
+
+public:
+	static UE_API FString GetPipelineCategory(UClass* AssetClass);
+
+	/** The name of the pipeline that will be displayed in the import dialog. */
+	UPROPERTY(EditAnywhere, BlueprintReadWrite, Category = "Grooms", meta = (StandAlonePipelineProperty = "True", PipelineInternalEditionData = "True"))
+	FString PipelineDisplayName;
+
+	/** If enabled, allow the import of groom-type assets. */
+	UPROPERTY(EditAnywhere, BlueprintReadWrite, Category = "Grooms")
+	bool bEnableGroomTypesImport = false;
+
+	/** If enabled, import all groom assets found in the source. */
+	UPROPERTY(EditAnywhere, BlueprintReadWrite, Category = "Grooms")
+	bool bImportGrooms = true;
+
+protected:
+
+	UE_API virtual void ExecutePipeline(UInterchangeBaseNodeContainer* BaseNodeContainer, const TArray<UInterchangeSourceData*>& SourceDatas, const FString& ContentBasePath) override;
+
+public:
+#if WITH_EDITOR
+	UE_API virtual void GetSupportAssetClasses(TArray<UClass*>& PipelineSupportAssetClasses) const override;
+#endif
+};
+
+#undef UE_API
\ No newline at end of file
diff --git a/Engine/Source/Runtime/Interchange/Core/Public/InterchangeFactoryBase.h b/Engine/Source/Runtime/Interchange/Core/Public/InterchangeFactoryBase.h
--- a/Engine/Source/Runtime/Interchange/Core/Public/InterchangeFactoryBase.h
+++ b/Engine/Source/Runtime/Interchange/Core/Public/InterchangeFactoryBase.h
@@ -27,7 +27,8 @@ enum class EInterchangeFactoryAssetType : uint8
 	Materials,
 	Meshes,
 	Animations,
-	Physics
+	Physics,
+	Grooms
 };
 
 /**
