Fixed profiler tag.

diff --git a/Engine/Plugins/Interchange/Runtime/Source/Import/Private/Groom/InterchangeGroomFactory.cpp b/Engine/Plugins/Interchange/Runtime/Source/Import/Private/Groom/InterchangeGroomFactory.cpp
--- a/Engine/Plugins/Interchange/Runtime/Source/Import/Private/Groom/InterchangeGroomFactory.cpp
+++ b/Engine/Plugins/Interchange/Runtime/Source/Import/Private/Groom/InterchangeGroomFactory.cpp
@@ -177,7 +177,7 @@ UInterchangeFactoryBase::FImportAssetResult UInterchangeGroomFactory::BeginImpor
 
 void UInterchangeGroomFactory::SetupObject_GameThread(const FSetupObjectParams& Arguments)
 {
-	TRACE_CPUPROFILER_EVENT_SCOPE(UInterchangeGeometryCacheFactory::SetupObject_GameThread);
+	TRACE_CPUPROFILER_EVENT_SCOPE(UInterchangeGroomFactory::SetupObject_GameThread);
 
 	Super::SetupObject_GameThread(Arguments);
 
