Fixed clang variable shadowing error between EInterchangeSkeletonCompareSection::Count and EHairDescriptionType::Count.

#rnx

diff --git a/Engine/Plugins/Interchange/Runtime/Source/Pipelines/Private/InterchangeGenericAssetsPipeline.cpp b/Engine/Plugins/Interchange/Runtime/Source/Pipelines/Private/InterchangeGenericAssetsPipeline.cpp
--- a/Engine/Plugins/Interchange/Runtime/Source/Pipelines/Private/InterchangeGenericAssetsPipeline.cpp
+++ b/Engine/Plugins/Interchange/Runtime/Source/Pipelines/Private/InterchangeGenericAssetsPipeline.cpp
@@ -2094,8 +2094,8 @@ FReply SInterchangeGenericAssetMaterialConflictWidget::OnKeyDown(const FGeometry
 
 void SInterchangeGenericAssetSkeletonConflictWidget::Construct(const FArguments& InArgs)
 {
-	bShowSectionFlag[EInterchangeSkeletonCompareSection::Skeleton] = true;
-	bShowSectionFlag[EInterchangeSkeletonCompareSection::References] = true;
+	bShowSectionFlag[(uint8)EInterchangeSkeletonCompareSection::Skeleton] = true;
+	bShowSectionFlag[(uint8)EInterchangeSkeletonCompareSection::References] = true;
 
 	AssetReferencingSkeleton = InArgs._AssetReferencingSkeleton;
 	Joints = InArgs._Joints;
@@ -2202,18 +2202,18 @@ FReply SInterchangeGenericAssetSkeletonConflictWidget::OnExpandToConflict()
 
 FReply SInterchangeGenericAssetSkeletonConflictWidget::SetSectionVisible(EInterchangeSkeletonCompareSection SectionIndex)
 {
-	bShowSectionFlag[SectionIndex] = !bShowSectionFlag[SectionIndex];
+	bShowSectionFlag[(uint8)SectionIndex] = !bShowSectionFlag[(uint8)SectionIndex];
 	return FReply::Handled();
 }
 
 EVisibility SInterchangeGenericAssetSkeletonConflictWidget::IsSectionVisible(EInterchangeSkeletonCompareSection SectionIndex)
 {
-	return bShowSectionFlag[SectionIndex] ? EVisibility::Visible : EVisibility::Collapsed;
+	return bShowSectionFlag[(uint8)SectionIndex] ? EVisibility::Visible : EVisibility::Collapsed;
 }
 
 const FSlateBrush* SInterchangeGenericAssetSkeletonConflictWidget::GetCollapsableArrow(EInterchangeSkeletonCompareSection SectionIndex) const
 {
-	return bShowSectionFlag[SectionIndex] ? FAppStyle::GetBrush("Symbols.DownArrow") : FAppStyle::GetBrush("Symbols.RightArrow");
+	return bShowSectionFlag[(uint8)SectionIndex] ? FAppStyle::GetBrush("Symbols.DownArrow") : FAppStyle::GetBrush("Symbols.RightArrow");
 }
 
 TSharedPtr<SWidget> SInterchangeGenericAssetSkeletonConflictWidget::ConstructSkeletonComparison()
diff --git a/Engine/Plugins/Interchange/Runtime/Source/Pipelines/Public/InterchangeGenericAssetsPipeline.h b/Engine/Plugins/Interchange/Runtime/Source/Pipelines/Public/InterchangeGenericAssetsPipeline.h
--- a/Engine/Plugins/Interchange/Runtime/Source/Pipelines/Public/InterchangeGenericAssetsPipeline.h
+++ b/Engine/Plugins/Interchange/Runtime/Source/Pipelines/Public/InterchangeGenericAssetsPipeline.h
@@ -281,7 +281,7 @@ protected:
 	TSharedPtr<SListView<TSharedPtr<FListItem>>> MaterialList;
 };
 
-enum EInterchangeSkeletonCompareSection
+enum class EInterchangeSkeletonCompareSection : uint8
 {
 	Skeleton = 0,
 	References,
@@ -337,7 +337,7 @@ private:
 
 	//////////////////////////////////////////////////////////////////////////
 	//Collapse generic
-	bool bShowSectionFlag[EInterchangeSkeletonCompareSection::Count];
+	bool bShowSectionFlag[(uint8)EInterchangeSkeletonCompareSection::Count];
 	FReply OnExpandToConflict();
 	FReply SetSectionVisible(EInterchangeSkeletonCompareSection SectionIndex);
 	EVisibility IsSectionVisible(EInterchangeSkeletonCompareSection SectionIndex);
