[Interchange] Component Name uniqueness fixes for Actor Factory.

#rnx

diff --git a/Engine/Plugins/Interchange/Runtime/Source/Import/Private/Scene/InterchangeActorFactory.cpp b/Engine/Plugins/Interchange/Runtime/Source/Import/Private/Scene/InterchangeActorFactory.cpp
--- a/Engine/Plugins/Interchange/Runtime/Source/Import/Private/Scene/InterchangeActorFactory.cpp
+++ b/Engine/Plugins/Interchange/Runtime/Source/Import/Private/Scene/InterchangeActorFactory.cpp
@@ -89,79 +89,104 @@ void UInterchangeActorFactory::ExecuteResetObjectProperties(const UInterchangeBa
 
 namespace UE::Interchange::Private
 {
-	USceneComponent* CreateSceneComponent(AActor* Owner, UClass* ComponentClass, USceneComponent* ParentComponent, const FString& DisplayName)
+	//Per Actor (important for Name uniqueness tracker.
+	struct FComponentProcessHelper
 	{
-		USceneComponent* SceneComponent = NewObject<USceneComponent>(Owner, ComponentClass, *DisplayName);
-		SceneComponent->SetupAttachment(ParentComponent);
-		SceneComponent->RegisterComponent();
-		Owner->AddInstanceComponent(SceneComponent);
-		return SceneComponent;
-	};
+		AActor* Owner;
+		const UInterchangeBaseNodeContainer* NodeContainer;
+		TMap<FString, uint64> NameUniquenessTracker;
+
+		FComponentProcessHelper() = delete;
+		FComponentProcessHelper(AActor* InOwner, const UInterchangeBaseNodeContainer* InNodeContainer)
+			: Owner(InOwner)
+			, NodeContainer(InNodeContainer)
+		{
+		}
 
-	void ProcessComponents(AActor* Owner, const UInterchangeBaseNodeContainer* NodeContainer, const TArray<FString>& ComponentUids, USceneComponent* ParentComponent, EComponentMobility::Type MobilityToSet)
-	{
-		for (const FString& ComponentUid : ComponentUids)
+		USceneComponent* CreateSceneComponent(UClass* ComponentClass, USceneComponent* ParentComponent, const FString& DisplayName)
 		{
-			if (const UInterchangeSceneComponentFactoryNode* SceneComponentFactoryNode = Cast<const UInterchangeSceneComponentFactoryNode>(NodeContainer->GetNode(ComponentUid)))
+			FString UniqueDisplayName = DisplayName;
+			if (NameUniquenessTracker.Contains(UniqueDisplayName))
 			{
-				UClass* ComponentClass = SceneComponentFactoryNode->GetObjectClass();
-				FString Name = SceneComponentFactoryNode->GetDisplayLabel();
-				
-				USceneComponent* CreatedSceneComponent = CreateSceneComponent(Owner, ComponentClass, ParentComponent, Name);
+				UniqueDisplayName += TEXT("_") + FString::FromInt(++NameUniquenessTracker[UniqueDisplayName]);
+			}
+			else
+			{
+				NameUniquenessTracker.Add(UniqueDisplayName, 0);
+			}
 
-				if (const UInterchangeInstancedStaticMeshComponentFactoryNode * ISMComponentNode = Cast<const UInterchangeInstancedStaticMeshComponentFactoryNode >(NodeContainer->GetNode(ComponentUid)))
+			USceneComponent* SceneComponent = NewObject<USceneComponent>(Owner, ComponentClass, *UniqueDisplayName);
+			SceneComponent->SetupAttachment(ParentComponent);
+			SceneComponent->RegisterComponent();
+			Owner->AddInstanceComponent(SceneComponent);
+			return SceneComponent;
+		};
+
+		void ProcessComponents(const TArray<FString>& ComponentUids, USceneComponent* ParentComponent, EComponentMobility::Type MobilityToSet)
+		{
+			for (const FString& ComponentUid : ComponentUids)
+			{
+				if (const UInterchangeSceneComponentFactoryNode* SceneComponentFactoryNode = Cast<const UInterchangeSceneComponentFactoryNode>(NodeContainer->GetNode(ComponentUid)))
 				{
-					if (UInstancedStaticMeshComponent* ISMComponent = Cast<UInstancedStaticMeshComponent>(CreatedSceneComponent))
+					UClass* ComponentClass = SceneComponentFactoryNode->GetObjectClass();
+					FString Name = SceneComponentFactoryNode->GetDisplayLabel();
+
+					USceneComponent* CreatedSceneComponent = CreateSceneComponent(ComponentClass, ParentComponent, Name);
+
+					if (const UInterchangeInstancedStaticMeshComponentFactoryNode* ISMComponentNode = Cast<const UInterchangeInstancedStaticMeshComponentFactoryNode >(NodeContainer->GetNode(ComponentUid)))
 					{
-						TArray<FTransform> InstanceTransforms;
-						ISMComponentNode->GetInstanceTransforms(InstanceTransforms);
-						for (const FTransform& InstanceTransform : InstanceTransforms)
+						if (UInstancedStaticMeshComponent* ISMComponent = Cast<UInstancedStaticMeshComponent>(CreatedSceneComponent))
 						{
-							ISMComponent->AddInstance(InstanceTransform);
-						}
+							TArray<FTransform> InstanceTransforms;
+							ISMComponentNode->GetInstanceTransforms(InstanceTransforms);
+							for (const FTransform& InstanceTransform : InstanceTransforms)
+							{
+								ISMComponent->AddInstance(InstanceTransform);
+							}
 
-						FString AssetInstanceUid;
-						ISMComponentNode->GetCustomInstancedAssetUid(AssetInstanceUid);
+							FString AssetInstanceUid;
+							ISMComponentNode->GetCustomInstancedAssetUid(AssetInstanceUid);
 
-						if (const UInterchangeStaticMeshFactoryNode* StaticMeshFactoryNode = Cast<const UInterchangeStaticMeshFactoryNode>(NodeContainer->GetNode(AssetInstanceUid)))
-						{
-							FSoftObjectPath ReferenceObject;
-							StaticMeshFactoryNode->GetCustomReferenceObject(ReferenceObject);
-							if (UStaticMesh* StaticMesh = Cast<UStaticMesh>(ReferenceObject.TryLoad()))
+							if (const UInterchangeStaticMeshFactoryNode* StaticMeshFactoryNode = Cast<const UInterchangeStaticMeshFactoryNode>(NodeContainer->GetNode(AssetInstanceUid)))
 							{
-								ISMComponent->SetStaticMesh(StaticMesh);
+								FSoftObjectPath ReferenceObject;
+								StaticMeshFactoryNode->GetCustomReferenceObject(ReferenceObject);
+								if (UStaticMesh* StaticMesh = Cast<UStaticMesh>(ReferenceObject.TryLoad()))
+								{
+									ISMComponent->SetStaticMesh(StaticMesh);
+								}
 							}
 						}
 					}
-				}
 
-				//General setups:
-				{
-					FTransform LocalTransform;
-					SceneComponentFactoryNode->GetCustomLocalTransform(LocalTransform);
+					//General setups:
+					{
+						FTransform LocalTransform;
+						SceneComponentFactoryNode->GetCustomLocalTransform(LocalTransform);
 
-					CreatedSceneComponent->SetMobility(EComponentMobility::Type::Movable); //so that RelativeTransform can be set
-					CreatedSceneComponent->SetRelativeTransform(LocalTransform);
-					CreatedSceneComponent->SetMobility(MobilityToSet);
+						CreatedSceneComponent->SetMobility(EComponentMobility::Type::Movable); //so that RelativeTransform can be set
+						CreatedSceneComponent->SetRelativeTransform(LocalTransform);
+						CreatedSceneComponent->SetMobility(MobilityToSet);
 
-					bool bComponentVisibility = true;
-					if (SceneComponentFactoryNode->GetCustomComponentVisibility(bComponentVisibility))
-					{
-						CreatedSceneComponent->SetVisibility(bComponentVisibility);
-					}
+						bool bComponentVisibility = true;
+						if (SceneComponentFactoryNode->GetCustomComponentVisibility(bComponentVisibility))
+						{
+							CreatedSceneComponent->SetVisibility(bComponentVisibility);
+						}
 
-					TArray<FString> ChildrenComponentUids;
-					SceneComponentFactoryNode->GetComponentUids(ChildrenComponentUids);
+						TArray<FString> ChildrenComponentUids;
+						SceneComponentFactoryNode->GetComponentUids(ChildrenComponentUids);
 
-					ProcessComponents(Owner, NodeContainer, ChildrenComponentUids, CreatedSceneComponent, MobilityToSet);
+						ProcessComponents(ChildrenComponentUids, CreatedSceneComponent, MobilityToSet);
+					}
+				}
+				else
+				{
+					ensureMsgf(false, TEXT("Unexpected Component class type in ComponentUids."));
 				}
-			}
-			else
-			{
-				ensureMsgf(false, TEXT("Unexpected Component class type in ComponentUids."));
 			}
 		}
-	}
+	};
 }
 
 
@@ -192,7 +217,8 @@ UObject* UInterchangeActorFactory::ImportSceneObject_GameThread(const UInterchan
 				TArray<FString> ComponentUids;
 				FactoryNode->GetComponentUids(ComponentUids);
 
-				UE::Interchange::Private::ProcessComponents(SpawnedActor, CreateSceneObjectsParams.NodeContainer, ComponentUids, RootComponent, CachedMobility);
+				UE::Interchange::Private::FComponentProcessHelper ComponentProcessHelper(SpawnedActor, CreateSceneObjectsParams.NodeContainer);
+				ComponentProcessHelper.ProcessComponents(ComponentUids, RootComponent, CachedMobility);
 				
 				// Restore mobility value
 				if (CachedMobility != EComponentMobility::Type::Movable)
