[Interchange] Component Processor coment fix in Actor Factory.

#rnx

diff --git a/Engine/Plugins/Interchange/Runtime/Source/Import/Private/Scene/InterchangeActorFactory.cpp b/Engine/Plugins/Interchange/Runtime/Source/Import/Private/Scene/InterchangeActorFactory.cpp
--- a/Engine/Plugins/Interchange/Runtime/Source/Import/Private/Scene/InterchangeActorFactory.cpp
+++ b/Engine/Plugins/Interchange/Runtime/Source/Import/Private/Scene/InterchangeActorFactory.cpp
@@ -89,7 +89,7 @@ void UInterchangeActorFactory::ExecuteResetObjectProperties(const UInterchangeBa
 
 namespace UE::Interchange::Private
 {
-	//Per Actor (important for Name uniqueness tracker.
+	//Per Actor (important for Name uniqueness tracker).
 	struct FComponentProcessHelper
 	{
 		AActor* Owner;
