Interchange FBX:
- fix transparent material instances not correctly working

diff --git a/Engine/Plugins/Interchange/Runtime/Source/Parsers/Fbx/Private/FbxMaterial.cpp b/Engine/Plugins/Interchange/Runtime/Source/Parsers/Fbx/Private/FbxMaterial.cpp
--- a/Engine/Plugins/Interchange/Runtime/Source/Parsers/Fbx/Private/FbxMaterial.cpp
+++ b/Engine/Plugins/Interchange/Runtime/Source/Parsers/Fbx/Private/FbxMaterial.cpp
@@ -26,8 +26,6 @@ namespace UE
 			UInterchangeMaterialInstanceNode* FFbxMaterial::CreateMaterialInstanceNode(UInterchangeBaseNodeContainer& NodeContainer, const FString& NodeUID, const FString& NodeName)
 			{
 				UInterchangeMaterialInstanceNode* MaterialInstanceNode = NewObject<UInterchangeMaterialInstanceNode>(&NodeContainer);
-				// By default we create a Phong material, depending later of the different input hooked we'll change it
-				MaterialInstanceNode->SetCustomParent(TEXT("/InterchangeAssets/Materials/PhongSurfaceMaterial.PhongSurfaceMaterial"));
 				NodeContainer.SetupNode(MaterialInstanceNode, NodeUID, NodeName, EInterchangeNodeContainerType::TranslatedAsset);
 
 				return MaterialInstanceNode;
@@ -209,38 +207,41 @@ namespace UE
 				ProcessCustomAttributes(Parser, SurfaceMaterial, MaterialInstanceNode);
 
 				auto ShouldConvertProperty = [](FBXSDK_NAMESPACE::FbxProperty& MaterialProperty) -> bool
-					{
-						bool bShouldConvertProperty = false;
+				{
+					bool bShouldConvertProperty = false;
 
-						if (MaterialProperty.IsValid())
-						{
-							// FbxProperty::HasDefaultValue(..) can return true while the property has textures attached to it.
-							bShouldConvertProperty = MaterialProperty.GetSrcObjectCount<FBXSDK_NAMESPACE::FbxTexture>() > 0
-								|| !FBXSDK_NAMESPACE::FbxProperty::HasDefaultValue(MaterialProperty);
-						}
+					if (MaterialProperty.IsValid())
+					{
+						// FbxProperty::HasDefaultValue(..) can return true while the property has textures attached to it.
+						bShouldConvertProperty = MaterialProperty.GetSrcObjectCount<FBXSDK_NAMESPACE::FbxTexture>() > 0
+							|| !FBXSDK_NAMESPACE::FbxProperty::HasDefaultValue(MaterialProperty);
+					}
 
-						return bShouldConvertProperty;
-					};
+					return bShouldConvertProperty;
+				};
 
 				auto GetFactor = [&SurfaceMaterial](const char* FactorName)
-					{
-						FBXSDK_NAMESPACE::FbxProperty Property = SurfaceMaterial->FindProperty(FactorName);
-						return Property.IsValid() ? (float)Property.Get<FbxDouble>() : 1.;
-					};
+				{
+					FBXSDK_NAMESPACE::FbxProperty Property = SurfaceMaterial->FindProperty(FactorName);
+					return Property.IsValid() ? (float)Property.Get<FbxDouble>() : 1.;
+				};
 
 				auto ConnectInput = [&](FName InputName, const char* FbxPropertyName, float Factor, TVariant<FLinearColor, float>& DefaultValue, bool bInverse) -> bool
+				{
+					FBXSDK_NAMESPACE::FbxProperty MaterialProperty = SurfaceMaterial->FindProperty(FbxPropertyName);
+					if (ShouldConvertProperty(MaterialProperty))
 					{
-						FBXSDK_NAMESPACE::FbxProperty MaterialProperty = SurfaceMaterial->FindProperty(FbxPropertyName);
-						if (ShouldConvertProperty(MaterialProperty))
-						{
-							return ConvertPropertyToMaterialInstanceNode(NodeContainer, MaterialInstanceNode, MaterialProperty, Factor, InputName, DefaultValue, bInverse);
-						}
+						return ConvertPropertyToMaterialInstanceNode(NodeContainer, MaterialInstanceNode, MaterialProperty, Factor, InputName, DefaultValue, bInverse);
+					}
 
-						return false;
-					};
+					return false;
+				};
 
 				bool bHasInput = false;
 
+				// By default we create a Phong material, depending later of the different input hooked we'll change it
+				FString ParentMaterial{ TEXT("/InterchangeAssets/Materials/PhongSurfaceMaterial.PhongSurfaceMaterial") };
+
 				// Diffuse
 				{
 					const float Factor = GetFactor(FbxSurfaceMaterial::sDiffuseFactor);
@@ -290,24 +291,24 @@ namespace UE
 
 						if (MaterialProperty.GetSrcObjectCount<FbxFileTexture>() > 0)
 						{
-							MaterialInstanceNode->SetCustomParent(TEXT("/InterchangeAssets/Materials/PhongSurfaceTransparentMaterial.PhongSurfaceTransparentMaterial"));
+							ParentMaterial = TEXT("/InterchangeAssets/Materials/PhongSurfaceTranslucentMaterial.PhongSurfaceTranslucentMaterial");
 							TVariant<FLinearColor, float> DefaultValue;
 							DefaultValue.Set<float>(0.f); // Opaque
 							FName InputName = Phong::Parameters::Opacity;
 							// The texture is hooked to the OpacityMask when transparency is with a texture
 							if (MaterialProperty.Get<FbxDouble>() == 0.)
 							{
-								InputName = Phong::Parameters::OpacityMask;
-								MaterialInstanceNode->SetCustomParent(TEXT("/InterchangeAssets/Materials/PhongSurfaceMaskMaterial.PhongSurfaceMaskMaterial"));
+								InputName = Phong::Parameters::OpacityMask;								
+								ParentMaterial = TEXT("/InterchangeAssets/Materials/PhongSurfaceMaskMaterial.PhongSurfaceMaskMaterial");
 							}
 							bHasInput |= ConvertPropertyToMaterialInstanceNode(NodeContainer, MaterialInstanceNode, MaterialProperty, Factor, InputName, DefaultValue, true);
 						}
 						else
 						{
-							const float OpacityScalar = 1.f - (MaterialProperty.Get<float>() * Factor);
+							const float OpacityScalar = 1.f - Factor;
 							if (OpacityScalar < 1.f)
 							{
-								MaterialInstanceNode->SetCustomParent(TEXT("/InterchangeAssets/Materials/PhongSurfaceTransparentMaterial.PhongSurfaceTransparentMaterial"));
+								ParentMaterial = TEXT("/InterchangeAssets/Materials/PhongSurfaceTranslucentMaterial.PhongSurfaceTranslucentMaterial");
 								MaterialInstanceNode->AddScalarParameterValue(Phong::Parameters::Opacity.ToString(), OpacityScalar);
 							}
 						}
@@ -337,9 +338,11 @@ namespace UE
 
 				if (!bIsPhong)
 				{
-					MaterialInstanceNode->SetCustomParent(TEXT("/InterchangeAssets/Materials/LambertSurfaceMaterial.LambertSurfaceMaterial"));
+					ParentMaterial = TEXT("/InterchangeAssets/Materials/LambertSurfaceMaterial.LambertSurfaceMaterial");
 				}
 
+				MaterialInstanceNode->SetCustomParent(ParentMaterial);
+
 				// If no valid property found, create a material anyway
 				TArray<FString> InputNames;
 				if (!bHasInput)
