MaterialX:
- Renamed OpenPBR inputs to v1.3 (such as specular_roughness_anisotropy, coat_roughness_anisotropy, etc.)
- Add both transmission and subsurface inputs in both Opaque and Translucents material functions for consistencies (StandardSurface and OpenPBR)
- Add missing inputs in the MaterialX pipeline

#virtualized

diff --git a/Engine/Plugins/Interchange/Runtime/Source/Common/Public/MaterialX/InterchangeMaterialXDefinitions.h b/Engine/Plugins/Interchange/Runtime/Source/Common/Public/MaterialX/InterchangeMaterialXDefinitions.h
--- a/Engine/Plugins/Interchange/Runtime/Source/Common/Public/MaterialX/InterchangeMaterialXDefinitions.h
+++ b/Engine/Plugins/Interchange/Runtime/Source/Common/Public/MaterialX/InterchangeMaterialXDefinitions.h
@@ -389,6 +389,7 @@ MATERIALX_NAMESPACE_BEGIN
 			static constexpr const char* GeometryNormal = "geometry_normal";
 			static constexpr const char* GeometryCoatNormal = "geometry_coat_normal";
 			static constexpr const char* GeometryTangent = "geometry_tangent";
+			static constexpr const char* GeometryCoatTangent = "geometry_coat_tangent";
 		}
 
 		namespace DefaultValue
@@ -433,6 +434,7 @@ MATERIALX_NAMESPACE_BEGIN
 			static const FVector GeometryNormal{ 0, 0, 1 };
 			static const FVector GeometryCoatNormal{ 0, 0, 1 };
 			static const FVector GeometryTangent{ 0, 1, 0 };
+			static const FVector GeometryCoatTangent{ 0, 1, 0 };
 		}
 	}
 
diff --git a/Engine/Plugins/Interchange/Runtime/Source/Pipelines/Private/InterchangeMaterialXPipeline.cpp b/Engine/Plugins/Interchange/Runtime/Source/Pipelines/Private/InterchangeMaterialXPipeline.cpp
--- a/Engine/Plugins/Interchange/Runtime/Source/Pipelines/Private/InterchangeMaterialXPipeline.cpp
+++ b/Engine/Plugins/Interchange/Runtime/Source/Pipelines/Private/InterchangeMaterialXPipeline.cpp
@@ -3,7 +3,7 @@
 #include "InterchangeMaterialXPipeline.h"
 
 #include "InterchangeImportModule.h"
-#include "InterchangeMaterialDefinitions.h"
+#include "MaterialX/InterchangeMaterialXDefinitions.h"
 #include "InterchangeMaterialFactoryNode.h"
 #include "InterchangePipelineLog.h"
 
@@ -20,6 +20,8 @@
 #define MATERIALX_MATERIALFUNCTION_PATH(Name) \
 	!MaterialXSettings->bIsSubstrateEnabled ? Name##FunctionsPath : Name##SubstratePath
 
+namespace mx = MaterialX;
+
 namespace
 {
 	constexpr const TCHAR* OpenPBRSurfaceFunctionsPath = TEXT("/InterchangeAssets/Functions/MX_OpenPBR_Opaque.MX_OpenPBR_Opaque");
@@ -68,7 +70,7 @@ TMap<EInterchangeMaterialXSettings, TPair<TSet<FName>, TSet<FName>>> UMaterialXP
 #endif // WITH_EDITOR
 
 UMaterialXPipelineSettings::UMaterialXPipelineSettings()
-{
+{	
 	using namespace UE::Interchange::Materials;
 #if WITH_EDITOR
 	if(HasAnyFlags(EObjectFlags::RF_ClassDefaultObject))
@@ -82,41 +84,47 @@ UMaterialXPipelineSettings::UMaterialXPipelineSettings()
 				{
 				// OpenPBRSurface Inputs
 				TSet<FName>{
-					OpenPBRSurface::Parameters::BaseWeight,
-					OpenPBRSurface::Parameters::BaseColor,
-					OpenPBRSurface::Parameters::BaseRoughness,
-					OpenPBRSurface::Parameters::BaseMetalness,
-					OpenPBRSurface::Parameters::SpecularWeight,
-					OpenPBRSurface::Parameters::SpecularColor,
-					OpenPBRSurface::Parameters::SpecularRoughness,
-					OpenPBRSurface::Parameters::SpecularIOR,
-					OpenPBRSurface::Parameters::SpecularIORLevel,
-					OpenPBRSurface::Parameters::SpecularAnisotropy,
-					OpenPBRSurface::Parameters::SpecularRotation,
-					OpenPBRSurface::Parameters::SubsurfaceWeight,
-					OpenPBRSurface::Parameters::SubsurfaceColor,
-					OpenPBRSurface::Parameters::SubsurfaceRadius,
-					OpenPBRSurface::Parameters::SubsurfaceRadiusScale,
-					OpenPBRSurface::Parameters::SubsurfaceAnisotropy,
-					OpenPBRSurface::Parameters::FuzzWeight,
-					OpenPBRSurface::Parameters::FuzzColor,
-					OpenPBRSurface::Parameters::FuzzRoughness,
-					OpenPBRSurface::Parameters::CoatWeight,
-					OpenPBRSurface::Parameters::CoatColor,
-					OpenPBRSurface::Parameters::CoatRoughness,
-					OpenPBRSurface::Parameters::CoatAnisotropy,
-					OpenPBRSurface::Parameters::CoatRotation,
-					OpenPBRSurface::Parameters::CoatIOR,
-					OpenPBRSurface::Parameters::CoatIORLevel,
-					OpenPBRSurface::Parameters::GeometryCoatNormal,
-					OpenPBRSurface::Parameters::ThinFilmThickness,
-					OpenPBRSurface::Parameters::ThinFilmIOR,
-					OpenPBRSurface::Parameters::EmissionLuminance,
-					OpenPBRSurface::Parameters::EmissionColor,
-					OpenPBRSurface::Parameters::GeometryNormal,
-					OpenPBRSurface::Parameters::GeometryTangent,
-					OpenPBRSurface::Parameters::GeometryOpacity,
-					OpenPBRSurface::Parameters::GeometryThinWalled,
+					mx::OpenPBRSurface::Input::BaseWeight,
+					mx::OpenPBRSurface::Input::BaseColor,
+					mx::OpenPBRSurface::Input::BaseDiffuseRoughness,
+					mx::OpenPBRSurface::Input::BaseMetalness,
+					mx::OpenPBRSurface::Input::SpecularWeight,
+					mx::OpenPBRSurface::Input::SpecularColor,
+					mx::OpenPBRSurface::Input::SpecularRoughness,
+					mx::OpenPBRSurface::Input::SpecularIOR,
+					mx::OpenPBRSurface::Input::SpecularRoughnessAnisotropy,
+					mx::OpenPBRSurface::Input::TransmissionWeight,
+					mx::OpenPBRSurface::Input::TransmissionColor,
+					mx::OpenPBRSurface::Input::TransmissionDepth,
+					mx::OpenPBRSurface::Input::TransmissionScatter,
+					mx::OpenPBRSurface::Input::TransmissionScatterAnisotropy,
+					mx::OpenPBRSurface::Input::TransmissionDispersionScale,
+					mx::OpenPBRSurface::Input::TransmissionDispersionAbbeNumber,
+					mx::OpenPBRSurface::Input::SubsurfaceWeight,
+					mx::OpenPBRSurface::Input::SubsurfaceColor,
+					mx::OpenPBRSurface::Input::SubsurfaceRadius,
+					mx::OpenPBRSurface::Input::SubsurfaceRadiusScale,
+					mx::OpenPBRSurface::Input::SubsurfaceScatterAnisotropy,
+					mx::OpenPBRSurface::Input::FuzzWeight,
+					mx::OpenPBRSurface::Input::FuzzColor,
+					mx::OpenPBRSurface::Input::FuzzRoughness,
+					mx::OpenPBRSurface::Input::CoatWeight,
+					mx::OpenPBRSurface::Input::CoatColor,
+					mx::OpenPBRSurface::Input::CoatRoughness,
+					mx::OpenPBRSurface::Input::CoatRoughnessAnisotropy,
+					mx::OpenPBRSurface::Input::CoatIOR,
+					mx::OpenPBRSurface::Input::CoatDarkening,
+					mx::OpenPBRSurface::Input::ThinFilmWeight,
+					mx::OpenPBRSurface::Input::ThinFilmThickness,
+					mx::OpenPBRSurface::Input::ThinFilmIOR,
+					mx::OpenPBRSurface::Input::EmissionLuminance,
+					mx::OpenPBRSurface::Input::EmissionColor,
+					mx::OpenPBRSurface::Input::GeometryOpacity,
+					mx::OpenPBRSurface::Input::GeometryThinWalled,
+					mx::OpenPBRSurface::Input::GeometryNormal,
+					mx::OpenPBRSurface::Input::GeometryCoatNormal,
+					mx::OpenPBRSurface::Input::GeometryTangent,
+					mx::OpenPBRSurface::Input::GeometryCoatTangent,
 				},
 				// OpenPBRSurface Outputs
 				!bIsSubstrateEnabled ?
@@ -148,43 +156,47 @@ UMaterialXPipelineSettings::UMaterialXPipelineSettings()
 				{
 					// OpenPBRSurfaceTransmission Inputs
 					TSet<FName>{
-						OpenPBRSurface::Parameters::BaseWeight,
-						OpenPBRSurface::Parameters::BaseColor,
-						OpenPBRSurface::Parameters::BaseRoughness,
-						OpenPBRSurface::Parameters::BaseMetalness,
-						OpenPBRSurface::Parameters::SpecularWeight,
-						OpenPBRSurface::Parameters::SpecularColor,
-						OpenPBRSurface::Parameters::SpecularRoughness,
-						OpenPBRSurface::Parameters::SpecularIOR,
-						OpenPBRSurface::Parameters::SpecularIORLevel,
-						OpenPBRSurface::Parameters::SpecularAnisotropy,
-						OpenPBRSurface::Parameters::SpecularRotation,
-						OpenPBRSurface::Parameters::TransmissionWeight,
-						OpenPBRSurface::Parameters::TransmissionColor,
-						OpenPBRSurface::Parameters::TransmissionDepth,
-						OpenPBRSurface::Parameters::TransmissionDispersionScale,
-						OpenPBRSurface::Parameters::TransmissionDispersionAbbeNumber,
-						OpenPBRSurface::Parameters::TransmissionScatter,
-						OpenPBRSurface::Parameters::TransmissionScatterAnisotropy,
-						OpenPBRSurface::Parameters::FuzzWeight,
-						OpenPBRSurface::Parameters::FuzzColor,
-						OpenPBRSurface::Parameters::FuzzRoughness,
-						OpenPBRSurface::Parameters::CoatWeight,
-						OpenPBRSurface::Parameters::CoatColor,
-						OpenPBRSurface::Parameters::CoatRoughness,
-						OpenPBRSurface::Parameters::CoatAnisotropy,
-						OpenPBRSurface::Parameters::CoatRotation,
-						OpenPBRSurface::Parameters::CoatIOR,
-						OpenPBRSurface::Parameters::CoatIORLevel,
-						OpenPBRSurface::Parameters::GeometryCoatNormal,
-						OpenPBRSurface::Parameters::ThinFilmThickness,
-						OpenPBRSurface::Parameters::ThinFilmIOR,
-						OpenPBRSurface::Parameters::EmissionLuminance,
-						OpenPBRSurface::Parameters::EmissionColor,
-						OpenPBRSurface::Parameters::GeometryNormal,
-						OpenPBRSurface::Parameters::GeometryTangent,
-						OpenPBRSurface::Parameters::GeometryOpacity,
-						OpenPBRSurface::Parameters::GeometryThinWalled,
+						mx::OpenPBRSurface::Input::BaseWeight,
+						mx::OpenPBRSurface::Input::BaseColor,
+						mx::OpenPBRSurface::Input::BaseDiffuseRoughness,
+						mx::OpenPBRSurface::Input::BaseMetalness,
+						mx::OpenPBRSurface::Input::SpecularWeight,
+						mx::OpenPBRSurface::Input::SpecularColor,
+						mx::OpenPBRSurface::Input::SpecularRoughness,
+						mx::OpenPBRSurface::Input::SpecularIOR,
+						mx::OpenPBRSurface::Input::SpecularRoughnessAnisotropy,
+						mx::OpenPBRSurface::Input::TransmissionWeight,
+						mx::OpenPBRSurface::Input::TransmissionColor,
+						mx::OpenPBRSurface::Input::TransmissionDepth,
+						mx::OpenPBRSurface::Input::TransmissionScatter,
+						mx::OpenPBRSurface::Input::TransmissionScatterAnisotropy,
+						mx::OpenPBRSurface::Input::TransmissionDispersionScale,
+						mx::OpenPBRSurface::Input::TransmissionDispersionAbbeNumber,
+						mx::OpenPBRSurface::Input::SubsurfaceWeight,
+						mx::OpenPBRSurface::Input::SubsurfaceColor,
+						mx::OpenPBRSurface::Input::SubsurfaceRadius,
+						mx::OpenPBRSurface::Input::SubsurfaceRadiusScale,
+						mx::OpenPBRSurface::Input::SubsurfaceScatterAnisotropy,
+						mx::OpenPBRSurface::Input::FuzzWeight,
+						mx::OpenPBRSurface::Input::FuzzColor,
+						mx::OpenPBRSurface::Input::FuzzRoughness,
+						mx::OpenPBRSurface::Input::CoatWeight,
+						mx::OpenPBRSurface::Input::CoatColor,
+						mx::OpenPBRSurface::Input::CoatRoughness,
+						mx::OpenPBRSurface::Input::CoatRoughnessAnisotropy,
+						mx::OpenPBRSurface::Input::CoatIOR,
+						mx::OpenPBRSurface::Input::CoatDarkening,
+						mx::OpenPBRSurface::Input::ThinFilmWeight,
+						mx::OpenPBRSurface::Input::ThinFilmThickness,
+						mx::OpenPBRSurface::Input::ThinFilmIOR,
+						mx::OpenPBRSurface::Input::EmissionLuminance,
+						mx::OpenPBRSurface::Input::EmissionColor,
+						mx::OpenPBRSurface::Input::GeometryOpacity,
+						mx::OpenPBRSurface::Input::GeometryThinWalled,
+						mx::OpenPBRSurface::Input::GeometryNormal,
+						mx::OpenPBRSurface::Input::GeometryCoatNormal,
+						mx::OpenPBRSurface::Input::GeometryTangent,
+						mx::OpenPBRSurface::Input::GeometryCoatTangent,
 					},
 					// OpenPBRSurfaceTransmission Outputs
 					!bIsSubstrateEnabled ?
@@ -212,40 +224,48 @@ UMaterialXPipelineSettings::UMaterialXPipelineSettings()
 				{
 					// StandardSurface Inputs
 					TSet<FName>{
-						StandardSurface::Parameters::Base,
-						StandardSurface::Parameters::BaseColor,
-						StandardSurface::Parameters::DiffuseRoughness,
-						StandardSurface::Parameters::Metalness,
-						StandardSurface::Parameters::Specular,
-						StandardSurface::Parameters::SpecularRoughness,
-						StandardSurface::Parameters::SpecularIOR,
-						StandardSurface::Parameters::SpecularAnisotropy,
-						StandardSurface::Parameters::SpecularRotation,
-						StandardSurface::Parameters::Subsurface,
-						StandardSurface::Parameters::SubsurfaceColor,
-						StandardSurface::Parameters::SubsurfaceRadius,
-						StandardSurface::Parameters::SubsurfaceScale,
-						StandardSurface::Parameters::Sheen,
-						StandardSurface::Parameters::SheenColor,
-						StandardSurface::Parameters::SheenRoughness,
-						StandardSurface::Parameters::Coat,
-						StandardSurface::Parameters::CoatColor,
-						StandardSurface::Parameters::CoatRoughness,
-						StandardSurface::Parameters::CoatNormal,
-						StandardSurface::Parameters::ThinFilmThickness,
-						StandardSurface::Parameters::ThinFilmIOR,
-						StandardSurface::Parameters::Emission,
-						StandardSurface::Parameters::EmissionColor,
-						StandardSurface::Parameters::Normal,
-						StandardSurface::Parameters::Tangent,
-						StandardSurface::Parameters::ThinWalled,
-						bIsSubstrateEnabled ? StandardSurface::Parameters::SpecularColor : FName{},
-						StandardSurface::Parameters::CoatIOR,
-						bIsSubstrateEnabled ? StandardSurface::Parameters::CoatAnisotropy : FName{},
-						bIsSubstrateEnabled ? StandardSurface::Parameters::CoatRotation : FName{},
-						bIsSubstrateEnabled ? StandardSurface::Parameters::ThinFilmIOR : FName{},
-						bIsSubstrateEnabled ? StandardSurface::Parameters::Opacity : FName{},
-						bIsSubstrateEnabled ? StandardSurface::Parameters::SubsurfaceAnisotropy : FName{}
+						mx::StandardSurface::Input::Base,
+						mx::StandardSurface::Input::BaseColor,
+						mx::StandardSurface::Input::DiffuseRoughness,
+						mx::StandardSurface::Input::Metalness,
+						mx::StandardSurface::Input::Specular,
+						mx::StandardSurface::Input::SpecularColor,
+						mx::StandardSurface::Input::SpecularRoughness,
+						mx::StandardSurface::Input::SpecularIOR,
+						mx::StandardSurface::Input::SpecularAnisotropy,
+						mx::StandardSurface::Input::SpecularRotation,
+						mx::StandardSurface::Input::Transmission,
+						mx::StandardSurface::Input::TransmissionColor,
+						mx::StandardSurface::Input::TransmissionDepth,
+						mx::StandardSurface::Input::TransmissionScatter,
+						mx::StandardSurface::Input::TransmissionScatterAnisotropy,
+						mx::StandardSurface::Input::TransmissionDispersion,
+						mx::StandardSurface::Input::TransmissionExtraRoughness,
+						mx::StandardSurface::Input::Subsurface,
+						mx::StandardSurface::Input::SubsurfaceColor,
+						mx::StandardSurface::Input::SubsurfaceRadius,
+						mx::StandardSurface::Input::SubsurfaceScale,
+						mx::StandardSurface::Input::SubsurfaceAnisotropy,
+						mx::StandardSurface::Input::Sheen,
+						mx::StandardSurface::Input::SheenColor,
+						mx::StandardSurface::Input::SheenRoughness,
+						mx::StandardSurface::Input::Coat,
+						mx::StandardSurface::Input::CoatColor,
+						mx::StandardSurface::Input::CoatRoughness,
+						mx::StandardSurface::Input::CoatAnisotropy,
+						mx::StandardSurface::Input::CoatRotation,
+						mx::StandardSurface::Input::CoatIOR,
+						mx::StandardSurface::Input::CoatNormal,
+						mx::StandardSurface::Input::CoatAffectColor,
+						mx::StandardSurface::Input::CoatAffectRoughness,
+						mx::StandardSurface::Input::ThinFilmThickness,
+						mx::StandardSurface::Input::ThinFilmIOR,
+						mx::StandardSurface::Input::Emission,
+						mx::StandardSurface::Input::EmissionColor,
+						mx::StandardSurface::Input::Opacity,
+						mx::StandardSurface::Input::ThinWalled,
+						mx::StandardSurface::Input::Normal,
+						mx::StandardSurface::Input::Tangent,
 					},
 					// StandardSurface Outputs
 					!bIsSubstrateEnabled ?
@@ -277,45 +297,48 @@ UMaterialXPipelineSettings::UMaterialXPipelineSettings()
 				{
 					// StandardSurfaceTransmission Inputs
 					TSet<FName>{
-						StandardSurface::Parameters::Base,
-						StandardSurface::Parameters::BaseColor,
-						StandardSurface::Parameters::DiffuseRoughness,
-						StandardSurface::Parameters::Metalness,
-						StandardSurface::Parameters::Specular,
-						bIsSubstrateEnabled ? StandardSurface::Parameters::SpecularColor : FName{},
-						StandardSurface::Parameters::SpecularRoughness,
-						StandardSurface::Parameters::SpecularIOR,
-						StandardSurface::Parameters::SpecularAnisotropy,
-						StandardSurface::Parameters::SpecularRotation,
-						StandardSurface::Parameters::Transmission,
-						StandardSurface::Parameters::TransmissionColor,
-						StandardSurface::Parameters::TransmissionDepth,
-						StandardSurface::Parameters::TransmissionScatter,
-						StandardSurface::Parameters::TransmissionScatterAnisotropy,
-						StandardSurface::Parameters::TransmissionDispersion,
-						StandardSurface::Parameters::TransmissionExtraRoughness,
-						!bIsSubstrateEnabled ? StandardSurface::Parameters::Subsurface : FName{},
-						!bIsSubstrateEnabled ? StandardSurface::Parameters::SubsurfaceColor : FName{},
-						!bIsSubstrateEnabled ? StandardSurface::Parameters::SubsurfaceRadius : FName{},
-						!bIsSubstrateEnabled ? StandardSurface::Parameters::SubsurfaceScale : FName{},
-						StandardSurface::Parameters::Sheen,
-						StandardSurface::Parameters::SheenColor,
-						StandardSurface::Parameters::SheenRoughness,
-						StandardSurface::Parameters::Coat,
-						StandardSurface::Parameters::CoatColor,
-						StandardSurface::Parameters::CoatRoughness,
-						bIsSubstrateEnabled ? StandardSurface::Parameters::CoatAnisotropy : FName{},
-						bIsSubstrateEnabled ? StandardSurface::Parameters::CoatRotation : FName{},
-						StandardSurface::Parameters::CoatIOR,
-						StandardSurface::Parameters::CoatNormal,
-						StandardSurface::Parameters::ThinFilmThickness,
-						StandardSurface::Parameters::ThinFilmIOR,
-						bIsSubstrateEnabled ? StandardSurface::Parameters::ThinFilmIOR : FName{},
-						StandardSurface::Parameters::Emission,
-						StandardSurface::Parameters::EmissionColor,
-						bIsSubstrateEnabled ? StandardSurface::Parameters::Opacity : FName{},
-						StandardSurface::Parameters::Normal,
-						StandardSurface::Parameters::Tangent
+						mx::StandardSurface::Input::Base,
+						mx::StandardSurface::Input::BaseColor,
+						mx::StandardSurface::Input::DiffuseRoughness,
+						mx::StandardSurface::Input::Metalness,
+						mx::StandardSurface::Input::Specular,
+						mx::StandardSurface::Input::SpecularColor,
+						mx::StandardSurface::Input::SpecularRoughness,
+						mx::StandardSurface::Input::SpecularIOR,
+						mx::StandardSurface::Input::SpecularAnisotropy,
+						mx::StandardSurface::Input::SpecularRotation,
+						mx::StandardSurface::Input::Transmission,
+						mx::StandardSurface::Input::TransmissionColor,
+						mx::StandardSurface::Input::TransmissionDepth,
+						mx::StandardSurface::Input::TransmissionScatter,
+						mx::StandardSurface::Input::TransmissionScatterAnisotropy,
+						mx::StandardSurface::Input::TransmissionDispersion,
+						mx::StandardSurface::Input::TransmissionExtraRoughness,
+						mx::StandardSurface::Input::Subsurface,
+						mx::StandardSurface::Input::SubsurfaceColor,
+						mx::StandardSurface::Input::SubsurfaceRadius,
+						mx::StandardSurface::Input::SubsurfaceScale,
+						mx::StandardSurface::Input::SubsurfaceAnisotropy,
+						mx::StandardSurface::Input::Sheen,
+						mx::StandardSurface::Input::SheenColor,
+						mx::StandardSurface::Input::SheenRoughness,
+						mx::StandardSurface::Input::Coat,
+						mx::StandardSurface::Input::CoatColor,
+						mx::StandardSurface::Input::CoatRoughness,
+						mx::StandardSurface::Input::CoatAnisotropy,
+						mx::StandardSurface::Input::CoatRotation,
+						mx::StandardSurface::Input::CoatIOR,
+						mx::StandardSurface::Input::CoatNormal,
+						mx::StandardSurface::Input::CoatAffectColor,
+						mx::StandardSurface::Input::CoatAffectRoughness,
+						mx::StandardSurface::Input::ThinFilmThickness,
+						mx::StandardSurface::Input::ThinFilmIOR,
+						mx::StandardSurface::Input::Emission,
+						mx::StandardSurface::Input::EmissionColor,
+						mx::StandardSurface::Input::Opacity,
+						mx::StandardSurface::Input::ThinWalled,
+						mx::StandardSurface::Input::Normal,
+						mx::StandardSurface::Input::Tangent,
 					},
 					// StandardSurfaceTransmission Outputs
 					!bIsSubstrateEnabled ?
@@ -332,7 +355,10 @@ UMaterialXPipelineSettings::UMaterialXPipelineSettings()
 						PBRMR::Parameters::Refraction,
 						ThinTranslucent::Parameters::TransmissionColor
 					}	:
-					TSet<FName>{StandardSurface::SubstrateMaterial::Outputs::Translucent}
+					TSet<FName>{
+						StandardSurface::SubstrateMaterial::Outputs::Opaque, // It should be named Translucent, but it got renamed somehow
+						StandardSurface::SubstrateMaterial::Outputs::Opacity,
+					}
 				}
 			},
 			{
