[In Progress] Interchange CAD Importer plugin
- Added ability to manage multiple workers to the InterchangeDispatcher
- This is done to share the dispatcher's code between the InterchangeWorker and the InterchangeCADWorker.

#rnx

diff --git a/Engine/Plugins/Interchange/Runtime/Source/Dispatcher/Private/FBX/InterchangeDispatcherFBXTasks.cpp b/Engine/Plugins/Interchange/Runtime/Source/Dispatcher/Private/FBX/InterchangeDispatcherFBXTasks.cpp
new file mode 100644
--- /dev/null
+++ b/Engine/Plugins/Interchange/Runtime/Source/Dispatcher/Private/FBX/InterchangeDispatcherFBXTasks.cpp
@@ -0,0 +1,47 @@
+// Copyright Epic Games, Inc. All Rights Reserved.
+
+#include "FBX/InterchangeDispatcherFBXTasks.h"
+
+#include "Dom/JsonObject.h"
+
+namespace UE::Interchange
+{
+	const TCHAR* FJsonFBXLoadSourceCmd::TaskName = TEXT("LoadFBX");
+	const TCHAR* FJsonFBXLoadSourceCmd::ConvertSceneJsonKey = TEXT("ConvertScene");
+	const TCHAR* FJsonFBXLoadSourceCmd::ForceFrontXAxisJsonKey = TEXT("ForceFrontXAxis");
+	const TCHAR* FJsonFBXLoadSourceCmd::ConvertSceneUnitJsonKey = TEXT("ConvertSceneUnit");
+	const TCHAR* FJsonFBXLoadSourceCmd::KeepFbxNamespaceJsonKey = TEXT("KeepFbxNamespace");
+
+	TSharedPtr<FJsonObject> FJsonFBXLoadSourceCmd::GetActionDataObject() const
+	{
+		TSharedPtr<FJsonObject> ActionDataObject = FJsonLoadSourceCmd::GetActionDataObject();
+		ActionDataObject->SetBoolField(ConvertSceneJsonKey, bConvertScene);
+		ActionDataObject->SetBoolField(ForceFrontXAxisJsonKey, bForceFrontXAxis);
+		ActionDataObject->SetBoolField(ConvertSceneUnitJsonKey, bConvertSceneUnit);
+		ActionDataObject->SetBoolField(KeepFbxNamespaceJsonKey, bKeepFbxNamespace);
+
+		return ActionDataObject;
+	}
+
+	bool FJsonFBXLoadSourceCmd::IsActionDataObjectValid(const FJsonObject& ActionDataObject)
+	{
+		if (!(ActionDataObject.TryGetBoolField(ConvertSceneJsonKey, bConvertScene)))
+		{
+			return false;
+		}
+		if (!(ActionDataObject.TryGetBoolField(ForceFrontXAxisJsonKey, bForceFrontXAxis)))
+		{
+			return false;
+		}
+		if (!(ActionDataObject.TryGetBoolField(ConvertSceneUnitJsonKey, bConvertSceneUnit)))
+		{
+			return false;
+		}
+		if (!(ActionDataObject.TryGetBoolField(KeepFbxNamespaceJsonKey, bKeepFbxNamespace)))
+		{
+			return false;
+		}
+
+		return true;
+	}
+}//ns UE
diff --git a/Engine/Plugins/Interchange/Runtime/Source/Dispatcher/Private/InterchangeDispatcher.cpp b/Engine/Plugins/Interchange/Runtime/Source/Dispatcher/Private/InterchangeDispatcher.cpp
--- a/Engine/Plugins/Interchange/Runtime/Source/Dispatcher/Private/InterchangeDispatcher.cpp
+++ b/Engine/Plugins/Interchange/Runtime/Source/Dispatcher/Private/InterchangeDispatcher.cpp
@@ -8,6 +8,7 @@
 #include "InterchangeDispatcherConfig.h"
 #include "InterchangeDispatcherLog.h"
 #include "InterchangeDispatcherTask.h"
+#include "Misc/App.h"
 #include "Misc/Paths.h"
 #include "Misc/ScopeLock.h"
 
@@ -16,19 +17,40 @@ namespace UE
 	namespace Interchange
 	{
 
-		FInterchangeDispatcher::FInterchangeDispatcher(const FString& InResultFolder)
-			: NextTaskIndex(0)
+		FInterchangeDispatcher::FInterchangeDispatcher(const FString& InResultFolder, int32 DesiredWorkerCount)
+			: ResultFolder(InResultFolder)
+			, NextTaskIndex(0)
 			, CompletedTaskCount(0)
-			, ResultFolder(InResultFolder)
 		{
+			// Set up number of workers to initiate based on main application's workload
+			if (DesiredWorkerCount >= 0)
+			{
+				if (DesiredWorkerCount == 0)
+				{
+					DesiredWorkerCount = FPlatformMisc::NumberOfCores();
+				}
+
+				const int32 MaxNumberOfWorkers = FPlatformMisc::NumberOfCores();
+
+				constexpr double RecommandedRamPerWorkers = 6.;
+				constexpr double OneGigaByte = 1024. * 1024. * 1024.;
+				const double AvailableRamGB = (double)(FPlatformMemory::GetStats().AvailablePhysical / OneGigaByte);
+
+				const int32 RecommandedNumberOfWorkers = (int32)(AvailableRamGB / RecommandedRamPerWorkers + 0.5);
+
+				int32 WorkerCount = FMath::Min(DesiredWorkerCount, FMath::Min(MaxNumberOfWorkers, RecommandedNumberOfWorkers));
+
+				WorkerHandlers.SetNum(WorkerCount);
+			}
 		}
 
 		int32 FInterchangeDispatcher::AddTask(const FString& JsonDescription)
 		{
-			if (!WorkerHandler->IsAlive())
+			if (!IsValid())
 			{
 				return INDEX_NONE;
 			}
+
 			FScopeLock Lock(&TaskPoolCriticalSection);
 			int32 TaskIndex = TaskPool.Emplace(JsonDescription);
 			TaskPool[TaskIndex].Index = TaskIndex;
@@ -136,22 +158,26 @@ namespace UE
 		void FInterchangeDispatcher::StartProcess()
 		{
 			//Start the process
-			SpawnHandler();
+			SpawnHandlers();
 		}
 
 		void FInterchangeDispatcher::StopProcess(bool bBlockUntilTerminated)
 		{
-			if (IsHandlerAlive())
+			for (TUniquePtr<FInterchangeWorkerHandler>& WorkerHandler : WorkerHandlers)
 			{
-				if (bBlockUntilTerminated)
+				if (WorkerHandler.IsValid() && WorkerHandler->IsAlive())
 				{
-					WorkerHandler->StopBlocking();
-				}
-				else
-				{
-					WorkerHandler->Stop();
+					if (bBlockUntilTerminated)
+					{
+						WorkerHandler->StopBlocking();
+					}
+					else
+					{
+						WorkerHandler->Stop();
+					}
 				}
 			}
+
 			EmptyQueueTasks();
 		}
 
@@ -166,12 +192,12 @@ namespace UE
 				IFileManager::Get().DeleteDirectory(*ResultFolder, RequireExists, Tree);
 			}
 			//Terminate the process
-			CloseHandler();
+			CloseHandlers();
 		}
 
 		void FInterchangeDispatcher::WaitAllTaskToCompleteExecution()
 		{
-			if (!WorkerHandler.IsValid())
+			if (!IsValid())
 			{
 				UE_LOG(LogInterchangeDispatcher, Error, TEXT("Cannot execute tasks before starting the process"));
 				return;
@@ -181,7 +207,7 @@ namespace UE
 			while (!IsOver())
 			{
 
-				if (!IsHandlerAlive())
+				if (!IsValid())
 				{
 					break;
 				}
@@ -207,68 +233,37 @@ namespace UE
 			return CompletedTaskCount == TaskPool.Num();
 		}
 
-		bool FInterchangeDispatcher::IsInterchangeWorkerRunning()
+		bool FInterchangeDispatcher::IsValid()
 		{
-			if (IsHandlerAlive())
+			for (TUniquePtr<FInterchangeWorkerHandler>& WorkerHandler : WorkerHandlers)
 			{
-				WaitAllTaskToCompleteExecution();
-				return GetInterchangeWorkerFatalError().IsEmpty();
+				if (WorkerHandler.IsValid() && WorkerHandler->IsAlive())
+				{
+					return true;
+				}
 			}
+
 			return false;
 		}
 
-		//this is a static function
-		bool FInterchangeDispatcher::IsInterchangeWorkerAvailable()
+		void FInterchangeDispatcher::SpawnHandlers()
 		{
-			FString RandomGuid = FGuid::NewGuid().ToString(EGuidFormats::Base36Encoded);
-			const FString TmpResultFolder = FPaths::ConvertRelativePathToFull(FPaths::ProjectSavedDir() + TEXT("Interchange/Temp/") + RandomGuid);
-			TUniquePtr<UE::Interchange::FInterchangeDispatcher> Dispatcher = MakeUnique<UE::Interchange::FInterchangeDispatcher>(TmpResultFolder);
-			bool bWorkerValid = false;
-			if (ensure(Dispatcher.IsValid()))
+			for (TUniquePtr<FInterchangeWorkerHandler>& WorkerHandler : WorkerHandlers)
 			{
-				Dispatcher->StartProcess();
-				//Wait the conection handshake is done
-				//The interchange worker will validate the version and issue a error message if the connection is incorrect
-				//If version is good the interchange worker will issue a ping command
-				while (Dispatcher->IsHandlerAlive() && Dispatcher->GetInterchangeWorkerFatalError().IsEmpty())
-				{
-					if (Dispatcher->WorkerHandler->IsPingCommandReceived())
-					{
-						break;
-					}
-					else
-					{
-						FPlatformProcess::Sleep(0.001f);
-					}
-				}
-				if (Dispatcher->IsHandlerAlive())
-				{
-					Dispatcher->WaitAllTaskToCompleteExecution();
-					bWorkerValid = Dispatcher->GetInterchangeWorkerFatalError().IsEmpty();
-				}
-				Dispatcher->CloseHandler();
+				WorkerHandler = MakeUnique<FInterchangeWorkerHandler>(*this, ResultFolder);
 			}
-			return bWorkerValid;
 		}
 
-		void FInterchangeDispatcher::SpawnHandler()
+		void FInterchangeDispatcher::CloseHandlers()
 		{
-			WorkerHandler = MakeUnique<FInterchangeWorkerHandler>(*this, ResultFolder);
-			WorkerHandler->OnWorkerHandlerExitLoop.AddLambda([this]()
+			for (TUniquePtr<FInterchangeWorkerHandler>& WorkerHandler : WorkerHandlers)
 			{
-				EmptyQueueTasks();
-			});
-		}
-
-		bool FInterchangeDispatcher::IsHandlerAlive()
-		{
-			return WorkerHandler.IsValid() && WorkerHandler->IsAlive();
-		}
-
-		void FInterchangeDispatcher::CloseHandler()
-		{
-			StopProcess(false);
-			WorkerHandler.Reset();
+				if (WorkerHandler.IsValid() && WorkerHandler->IsAlive())
+				{
+					WorkerHandler->Stop();
+					WorkerHandler.Reset();
+				}
+			}
 		}
 
 		void FInterchangeDispatcher::EmptyQueueTasks()
@@ -284,5 +279,31 @@ namespace UE
 			}
 		}
 
+		const FString& FInterchangeDispatcher::GetWorkerApplicationPath()
+		{
+			if (WorkerApplicationPath.IsEmpty())
+			{
+				WorkerApplicationPath = [&]()
+					{
+
+						const FString InterchangeWorkerApplicationName = GetWorkerApplicationName();
+						FString Path = FPlatformProcess::GenerateApplicationPath(InterchangeWorkerApplicationName, FApp::GetBuildConfiguration());
+						if (!FPaths::FileExists(Path))
+						{
+							//Force the development build if the path is not right
+							Path = FPlatformProcess::GenerateApplicationPath(InterchangeWorkerApplicationName, EBuildConfiguration::Development);
+						}
+
+						if (!FPaths::FileExists(Path))
+						{
+							UE_LOG(LogInterchangeDispatcher, Display, TEXT("InterchangeWorker executable not found. Expected location: %s"), *FPaths::ConvertRelativePathToFull(Path));
+						}
+						return Path;
+					}();
+			}
+
+			return WorkerApplicationPath;
+		}
+
 	} //ns Interchange
 }//ns UE
\ No newline at end of file
diff --git a/Engine/Plugins/Interchange/Runtime/Source/Dispatcher/Private/InterchangeDispatcherTask.cpp b/Engine/Plugins/Interchange/Runtime/Source/Dispatcher/Private/InterchangeDispatcherTask.cpp
--- a/Engine/Plugins/Interchange/Runtime/Source/Dispatcher/Private/InterchangeDispatcherTask.cpp
+++ b/Engine/Plugins/Interchange/Runtime/Source/Dispatcher/Private/InterchangeDispatcherTask.cpp
@@ -13,23 +13,32 @@ namespace UE
 {
 	namespace Interchange
 	{
+
+		TSharedPtr<FJsonObject> FJsonLoadSourceCmd::GetActionDataObject() const
+		{
+			TSharedPtr<FJsonObject> ActionDataObject = MakeShared<FJsonObject>();
+			ActionDataObject->SetStringField(GetSourceFilenameJsonKey(), GetSourceFilename());
+
+			return ActionDataObject;
+		}
+
 		FString FJsonLoadSourceCmd::ToJson() const
 		{
 			//Code should not do a ToJson if the data was not set before
 			ensure(bIsDataInitialize);
 
-			TSharedPtr<FJsonObject> CmdObject = MakeShared<FJsonObject>();
-			TSharedPtr<FJsonObject> ActionDataObject = MakeShared<FJsonObject>();
 			//CmdObject
+			TSharedPtr<FJsonObject> CmdObject = MakeShared<FJsonObject>();
 			CmdObject->SetStringField(GetCommandIDJsonKey(), GetAction());
 			CmdObject->SetStringField(GetTranslatorIDJsonKey(), GetTranslatorID());
-			ActionDataObject->SetStringField(GetSourceFilenameJsonKey(), GetSourceFilename());
-			ActionDataObject->SetBoolField(GetDoesConvertSceneJsonKey(), GetDoesConvertScene());
-			ActionDataObject->SetBoolField(GetDoesForceFrontXAxisJsonKey(), GetDoesForceFrontXAxis());
-			ActionDataObject->SetBoolField(GetDoesConvertSceneUnitJsonKey(), GetDoesConvertSceneUnit());
-			ActionDataObject->SetBoolField(GetDoesKeepFbxNamespaceJsonKey(), GetDoesKeepFbxNamespace());
+			
+			//Get ActionDataObject
+			TSharedPtr<FJsonObject> ActionDataObject = GetActionDataObject();
+
+			//Add ActionDataObject
 			CmdObject->SetObjectField(GetCommandDataJsonKey(), ActionDataObject);
 
+			//Get Json string from CmdObject
 			FString LoadSourceCmd;
 			TSharedRef< TJsonWriter< TCHAR, TPrettyJsonPrintPolicy<TCHAR> > > JsonWriter = TJsonWriterFactory< TCHAR, TPrettyJsonPrintPolicy<TCHAR> >::Create(&LoadSourceCmd);
 			if (!FJsonSerializer::Serialize(CmdObject.ToSharedRef(), JsonWriter))
@@ -79,26 +88,10 @@ namespace UE
 			{
 				return false;
 			}
-			if (!((*ActionDataObject)->TryGetBoolField(GetDoesConvertSceneJsonKey(), bConvertScene)))
-			{
-				return false;
-			}
-			if (!((*ActionDataObject)->TryGetBoolField(GetDoesForceFrontXAxisJsonKey(), bForceFrontXAxis)))
-			{
-				return false;
-			}
-			if (!((*ActionDataObject)->TryGetBoolField(GetDoesConvertSceneUnitJsonKey(), bConvertSceneUnit)))
-			{
-				return false;
-			}
-			if (!((*ActionDataObject)->TryGetBoolField(GetDoesKeepFbxNamespaceJsonKey(), bKeepFbxNamespace)))
-			{
-				return false;
-			}
 
-			//Since we filled the data from the json file, set the data has been initialize.
-			bIsDataInitialize = true;
-			return true;
+			//Since we filled the data from the json file, set the data has been initialized.
+			bIsDataInitialize = IsActionDataObjectValid(*(*ActionDataObject));
+			return bIsDataInitialize;
 		}
 
 		FString FJsonLoadSourceCmd::JsonResultParser::ToJson() const
diff --git a/Engine/Plugins/Interchange/Runtime/Source/Dispatcher/Private/InterchangeWorkerHandler.cpp b/Engine/Plugins/Interchange/Runtime/Source/Dispatcher/Private/InterchangeWorkerHandler.cpp
--- a/Engine/Plugins/Interchange/Runtime/Source/Dispatcher/Private/InterchangeWorkerHandler.cpp
+++ b/Engine/Plugins/Interchange/Runtime/Source/Dispatcher/Private/InterchangeWorkerHandler.cpp
@@ -20,29 +20,6 @@ namespace UE
 {
 	namespace Interchange
 	{
-		static FString GetWorkerExecutablePath()
-		{
-			static FString ProcessorPath = [&]()
-			{
-				
-				const FString InterchangeWorkerApplicationName = TEXT("InterchangeWorker");
-				FString Path = FPlatformProcess::GenerateApplicationPath(InterchangeWorkerApplicationName, FApp::GetBuildConfiguration());
-				if (!FPaths::FileExists(Path))
-				{
-					//Force the development build if the path is not right
-					Path = FPlatformProcess::GenerateApplicationPath(InterchangeWorkerApplicationName, EBuildConfiguration::Development);
-				}
-
-				if (!FPaths::FileExists(Path))
-				{
-					UE_LOG(LogInterchangeDispatcher, Display, TEXT("InterchangeWorker executable not found. Expected location: %s"), *FPaths::ConvertRelativePathToFull(Path));
-				}
-				return Path;
-			}();
-
-			return ProcessorPath;
-		}
-
 		FInterchangeWorkerHandler::FInterchangeWorkerHandler(FInterchangeDispatcher& InDispatcher, FString& InResultFolder)
 			: Dispatcher(InDispatcher)
 			, WorkerState(EWorkerState::Uninitialized)
@@ -51,7 +28,7 @@ namespace UE
 			, bShouldTerminate(false)
 		{
 			//Use a randomGuid to generate unique name
-			int32 UniqueID = (FPlatformTime::Cycles64() & 0x00000000EFFFFFFF);
+			const int32 UniqueID = (FPlatformTime::Cycles64() & 0x00000000EFFFFFFF);
 			ThreadName = FString(TEXT("InterchangeWorkerHandler_")) + FString::FromInt(UniqueID);
 			IOThread = FThread(*ThreadName, [this]() { Run(); });
 			LastProgressMessageTime = FPlatformTime::Seconds();
@@ -67,7 +44,7 @@ namespace UE
 			uint32 WorkerProcessId = 0;
 
 			ensure(ErrorState == EWorkerErrorState::Ok);
-			FString ProcessorPath = GetWorkerExecutablePath();
+			FString ProcessorPath = Dispatcher.GetWorkerApplicationPath();
 			if (FPaths::FileExists(ProcessorPath))
 			{
 				int32 ListenPort = NetworkInterface.GetListeningPort();
@@ -248,11 +225,11 @@ namespace UE
 							break;
 						}
 
-						WorkerState = EWorkerState::Processing;
+						WorkerState = EWorkerState::Idle;
 						break;
 					}
 
-					case EWorkerState::Processing:
+					case EWorkerState::Idle:
 					{
 						ValidateConnection();
 						if (ErrorState == EWorkerErrorState::WorkerProcess_Lost)
@@ -265,8 +242,33 @@ namespace UE
 							if (StartNewTask())
 							{
 								WorkerState = EWorkerState::Processing;
+								break;
+							}
+
+							//Need to terminate?
+							if (bShouldTerminate)
+							{
+								UE_LOG(LogInterchangeDispatcher, Verbose, TEXT("Exit loop gracefully"));
+								WorkerState = EWorkerState::Closing;
+							}
+							// No task to execute, wait a little
+							else if (WorkerState == EWorkerState::Idle)
+							{
+								FPlatformProcess::Sleep(Config::IdleLoopDelay);
 							}
+						}
+						break;
+					}
 
+					case EWorkerState::Processing:
+					{
+						ValidateConnection();
+						if (ErrorState == EWorkerErrorState::WorkerProcess_Lost)
+						{
+							WorkerState = EWorkerState::Closing;
+						}
+						else
+						{
 							// consume task
 							if (TSharedPtr<ICommand> Command = CommandIO.GetNextCommand(Config::IdleLoopDelay))
 							{
@@ -279,6 +281,11 @@ namespace UE
 								UE_LOG(LogInterchangeDispatcher, Verbose, TEXT("Exit loop gracefully"));
 								WorkerState = EWorkerState::Closing;
 							}
+							// Task is being executed, sleep before rechecking
+							else if (WorkerState == EWorkerState::Processing)
+							{
+								FPlatformProcess::Sleep(Config::ProcessingLoopDelay);
+							}
 						}
 						break;
 					}
@@ -351,7 +358,7 @@ namespace UE
 
 		bool FInterchangeWorkerHandler::IsAlive() const
 		{
-			return WorkerState != EWorkerState::Terminated;
+			return !EnumHasAnyFlags(EWorkerState(WorkerState), EWorkerState::Terminated | EWorkerState::Closing);
 		}
 
 		void FInterchangeWorkerHandler::ProcessCommand(ICommand& Command)
@@ -401,6 +408,8 @@ namespace UE
 				CurrentTasks.Remove(CompletedTaskCommand.TaskIndex);
 			}
 			Dispatcher.SetTaskState(CompletedTaskCommand.TaskIndex, CompletedTaskCommand.ProcessResult, CompletedTaskCommand.JSonResult, CompletedTaskCommand.JSonMessages);
+
+			WorkerState = EWorkerState::Idle;
 		}
 
 		void FInterchangeWorkerHandler::ProcessCommand(FCompletedQueryTaskProgressCommand& CompletedQueryTaskProgressCommand)
diff --git a/Engine/Plugins/Interchange/Runtime/Source/Dispatcher/Public/FBX/InterchangeDispatcherFBXTasks.h b/Engine/Plugins/Interchange/Runtime/Source/Dispatcher/Public/FBX/InterchangeDispatcherFBXTasks.h
new file mode 100644
--- /dev/null
+++ b/Engine/Plugins/Interchange/Runtime/Source/Dispatcher/Public/FBX/InterchangeDispatcherFBXTasks.h
@@ -0,0 +1,73 @@
+// Copyright Epic Games, Inc. All Rights Reserved.
+
+#pragma once
+
+#include "InterchangeDispatcherTask.h"
+
+namespace UE::Interchange
+{
+
+	class FJsonFBXLoadSourceCmd : public FJsonLoadSourceCmd
+	{
+	public:
+		FJsonFBXLoadSourceCmd() : FJsonLoadSourceCmd() {}
+
+		FJsonFBXLoadSourceCmd(const FString& InTranslatorID
+			, const FString& InSourceFilename
+			, const bool InbConvertScene
+			, const bool InbForceFrontXAxis
+			, const bool InbConvertSceneUnit
+			, const bool InbKeepFbxNamespace)
+			: FJsonLoadSourceCmd(InTranslatorID, InSourceFilename)
+			, bConvertScene(InbConvertScene)
+			, bForceFrontXAxis(InbForceFrontXAxis)
+			, bConvertSceneUnit(InbConvertSceneUnit)
+			, bKeepFbxNamespace(InbKeepFbxNamespace)
+		{
+		}
+
+		INTERCHANGEDISPATCHER_API virtual TSharedPtr<FJsonObject> GetActionDataObject() const override;
+		INTERCHANGEDISPATCHER_API virtual bool IsActionDataObjectValid(const FJsonObject& ActionDataObject) override;
+		
+		bool GetDoesConvertScene() const
+		{
+			//Code should not do query data if the data was not set before
+			ensure(bIsDataInitialize);
+			return bConvertScene;
+		}
+
+		bool GetDoesForceFrontXAxis() const
+		{
+			//Code should not do query data if the data was not set before
+			ensure(bIsDataInitialize);
+			return bForceFrontXAxis;
+		}
+
+		bool GetDoesConvertSceneUnit() const
+		{
+			//Code should not do query data if the data was not set before
+			ensure(bIsDataInitialize);
+			return bConvertSceneUnit;
+		}
+
+		bool GetDoesKeepFbxNamespace() const
+		{
+			//Code should not do query data if the data was not set before
+			ensure(bIsDataInitialize);
+			return bKeepFbxNamespace;
+		}
+
+		INTERCHANGEDISPATCHER_API static const TCHAR* TaskName;
+
+	private:
+		bool bConvertScene = true;
+		bool bForceFrontXAxis = false;
+		bool bConvertSceneUnit = true;
+		bool bKeepFbxNamespace = false;
+
+		static const TCHAR* ConvertSceneJsonKey;
+		static const TCHAR* ForceFrontXAxisJsonKey;
+		static const TCHAR* ConvertSceneUnitJsonKey;
+		static const TCHAR* KeepFbxNamespaceJsonKey;
+	};
+}//ns UE
diff --git a/Engine/Plugins/Interchange/Runtime/Source/Dispatcher/Public/InterchangeDispatcher.h b/Engine/Plugins/Interchange/Runtime/Source/Dispatcher/Public/InterchangeDispatcher.h
--- a/Engine/Plugins/Interchange/Runtime/Source/Dispatcher/Public/InterchangeDispatcher.h
+++ b/Engine/Plugins/Interchange/Runtime/Source/Dispatcher/Public/InterchangeDispatcher.h
@@ -18,8 +18,10 @@ namespace UE
 		class  FInterchangeDispatcher
 		{
 		public:
-			UE_API FInterchangeDispatcher(const FString& InResultFolder);
-			~FInterchangeDispatcher() { TerminateProcess(); }
+			// If WorkerCount == 0, the dispatcher will determine the number of worker processes to start
+			// based on main application's workload, threads and memory wise.
+			UE_API FInterchangeDispatcher(const FString& InResultFolder, int32 WorkerCount = 1);
+			virtual ~FInterchangeDispatcher() { TerminateProcess(); }
 
 			UE_API int32 AddTask(const FString& JsonDescription);
 			UE_API int32 AddTask(const FString& JsonDescription, FInterchangeDispatcherTaskCompleted TaskCompledDelegate);
@@ -29,11 +31,18 @@ namespace UE
 			UE_API void GetTaskState(int32 TaskIndex, ETaskState& TaskState, FString& JsonResult, TArray<FString>& JSonMessages);
 
 			UE_API void StartProcess();
-			UE_API void StopProcess(bool bBlockUntilTerminated);
-			UE_API void TerminateProcess();
+			UE_API virtual void StopProcess(bool bBlockUntilTerminated);
+			UE_API virtual void TerminateProcess();
 			UE_API void WaitAllTaskToCompleteExecution();
 			UE_API bool IsOver();
 
+			virtual const TCHAR* GetWorkerApplicationName()
+			{
+				return TEXT("InterchangeWorker");
+			}
+
+			const FString& GetWorkerApplicationPath();
+
 			void SetInterchangeWorkerFatalError(FString& ErrorMessage)
 			{
 				InterchangeWorkerFatalError = MoveTemp(ErrorMessage);
@@ -44,27 +53,30 @@ namespace UE
 				return InterchangeWorkerFatalError;
 			}
 
-			UE_API bool IsInterchangeWorkerRunning();
-			static UE_API bool IsInterchangeWorkerAvailable();
+			UE_API bool IsValid();
+		
+		protected:
+			UE_API void CloseHandlers();
+
+			/** Path where the result files are dump */
+			FString ResultFolder;
+			TArray<FTask> TaskPool;
+
 		private:
-			void SpawnHandler();
-			bool IsHandlerAlive();
-			void CloseHandler();
+			void SpawnHandlers();
 			void EmptyQueueTasks();
 
 			// Tasks
 			FCriticalSection TaskPoolCriticalSection;
-			TArray<FTask> TaskPool;
 			int32 NextTaskIndex;
 			int32 CompletedTaskCount;
 
-			/** Path where the result files are dump */
-			FString ResultFolder;
-
 			FString InterchangeWorkerFatalError;
 
 			// Workers
-			TUniquePtr<FInterchangeWorkerHandler> WorkerHandler = nullptr;
+			FString WorkerApplicationPath;
+			TArray<TUniquePtr<FInterchangeWorkerHandler>> WorkerHandlers;
+			TUniquePtr<FInterchangeWorkerHandler> WorkerHandler_Deprecated = nullptr;
 		};
 
 	} //ns Interchange
diff --git a/Engine/Plugins/Interchange/Runtime/Source/Dispatcher/Public/InterchangeDispatcherTask.h b/Engine/Plugins/Interchange/Runtime/Source/Dispatcher/Public/InterchangeDispatcherTask.h
--- a/Engine/Plugins/Interchange/Runtime/Source/Dispatcher/Public/InterchangeDispatcherTask.h
+++ b/Engine/Plugins/Interchange/Runtime/Source/Dispatcher/Public/InterchangeDispatcherTask.h
@@ -4,6 +4,8 @@
 
 #include "CoreMinimal.h"
 
+class FJsonObject;
+
 namespace UE
 {
 	namespace Interchange
@@ -12,6 +14,7 @@ namespace UE
 		struct FAnimationPayloadQuery;
 	}
 }
+
 DECLARE_DELEGATE_OneParam(FInterchangeDispatcherTaskCompleted, int32 TaskIndex);
 
 namespace UE
@@ -92,18 +95,12 @@ namespace UE
 				bIsDataInitialize = false;
 			}
 
+			virtual ~FJsonLoadSourceCmd() = default;
+
 			FJsonLoadSourceCmd(const FString& InTranslatorID
-				, const FString& InSourceFilename
-				, const bool InbConvertScene
-				, const bool InbForceFrontXAxis
-				, const bool InbConvertSceneUnit
-				, const bool InbKeepFbxNamespace)
+				, const FString& InSourceFilename)
 				: TranslatorID(InTranslatorID)
 				, SourceFilename(InSourceFilename)
-				, bConvertScene(InbConvertScene)
-				, bForceFrontXAxis(InbForceFrontXAxis)
-				, bConvertSceneUnit(InbConvertSceneUnit)
-				, bKeepFbxNamespace(InbKeepFbxNamespace)
 			{
 				bIsDataInitialize = true;
 			}
@@ -124,6 +121,12 @@ namespace UE
 			INTERCHANGEDISPATCHER_API virtual FString ToJson() const;
 			INTERCHANGEDISPATCHER_API virtual bool FromJson(const FString& JsonString);
 
+			INTERCHANGEDISPATCHER_API virtual TSharedPtr<FJsonObject> GetActionDataObject() const;
+			virtual bool IsActionDataObjectValid(const FJsonObject& ActionDataObject)
+			{
+				return true;
+			}
+
 			FString GetSourceFilename() const
 			{
 				//Code should not do query data if the data was not set before
@@ -137,58 +140,6 @@ namespace UE
 				return Key;
 			}
 
-			bool GetDoesConvertScene() const
-			{
-				//Code should not do query data if the data was not set before
-				ensure(bIsDataInitialize);
-				return bConvertScene;
-			}
-
-			static FString GetDoesConvertSceneJsonKey()
-			{
-				static const FString Key = TEXT("ConvertScene");
-				return Key;
-			}
-
-			bool GetDoesForceFrontXAxis() const
-			{
-				//Code should not do query data if the data was not set before
-				ensure(bIsDataInitialize);
-				return bForceFrontXAxis;
-			}
-
-			static FString GetDoesForceFrontXAxisJsonKey()
-			{
-				static const FString Key = TEXT("ForceFrontXAxis");
-				return Key;
-			}
-
-			bool GetDoesConvertSceneUnit() const
-			{
-				//Code should not do query data if the data was not set before
-				ensure(bIsDataInitialize);
-				return bConvertSceneUnit;
-			}
-
-			static FString GetDoesConvertSceneUnitJsonKey()
-			{
-				static const FString Key = TEXT("ConvertSceneUnit");
-				return Key;
-			}
-
-			bool GetDoesKeepFbxNamespace() const
-			{
-				//Code should not do query data if the data was not set before
-				ensure(bIsDataInitialize);
-				return bKeepFbxNamespace;
-			}
-
-			static FString GetDoesKeepFbxNamespaceJsonKey()
-			{
-				static const FString Key = TEXT("KeepFbxNamespace");
-				return Key;
-			}
-
 			/**
 			 * Use this class helper to create the cmd result json string and to read it
 			 */
@@ -218,10 +169,6 @@ namespace UE
 		private:
 			FString TranslatorID = FString();
 			FString SourceFilename = FString();
-			bool bConvertScene = true;
-			bool bForceFrontXAxis = false;
-			bool bConvertSceneUnit = true;
-			bool bKeepFbxNamespace = false;
 		};
 
 		class FJsonFetchPayloadCmd : public IJsonCmdBase
diff --git a/Engine/Plugins/Interchange/Runtime/Source/Dispatcher/Public/InterchangeWorkerHandler.h b/Engine/Plugins/Interchange/Runtime/Source/Dispatcher/Public/InterchangeWorkerHandler.h
--- a/Engine/Plugins/Interchange/Runtime/Source/Dispatcher/Public/InterchangeWorkerHandler.h
+++ b/Engine/Plugins/Interchange/Runtime/Source/Dispatcher/Public/InterchangeWorkerHandler.h
@@ -23,17 +23,19 @@ namespace UE
 
 		class FInterchangeDispatcher;
 
+		enum class EWorkerState : uint8
+		{
+			Uninitialized = 0,
+			Processing    = 1 << 0,   // send task and receive result
+			Idle          = 1 << 1,   // waiting for a new task
+			Closing       = 1 << 2,   // in the process of terminating
+			Terminated    = 1 << 3,   // worker does not run
+		};
+		ENUM_CLASS_FLAGS(EWorkerState);
+
 		//Handle a Worker by socket communication
 		class FInterchangeWorkerHandler
 		{
-			enum class EWorkerState
-			{
-				Uninitialized,
-				Processing, // send task and receive result
-				Closing, // in the process of terminating
-				Terminated, // aka. Not Alive
-			};
-
 			enum class EWorkerErrorState
 			{
 				Ok,
@@ -43,6 +45,7 @@ namespace UE
 				ConnectionLost_SendFailed,
 				WorkerProcess_CantCreate,
 				WorkerProcess_Lost,
+				WorkerProcess_Crashed,
 			};
 
 		public:
diff --git a/Engine/Plugins/Interchange/Runtime/Source/Import/Private/Fbx/InterchangeFbxTranslator.cpp b/Engine/Plugins/Interchange/Runtime/Source/Import/Private/Fbx/InterchangeFbxTranslator.cpp
--- a/Engine/Plugins/Interchange/Runtime/Source/Import/Private/Fbx/InterchangeFbxTranslator.cpp
+++ b/Engine/Plugins/Interchange/Runtime/Source/Import/Private/Fbx/InterchangeFbxTranslator.cpp
@@ -9,7 +9,7 @@
 #include "InterchangeAnalyticsHelper.h"
 #include "InterchangeCommonAnimationPayload.h"
 #include "InterchangeDispatcher.h"
-#include "InterchangeDispatcherTask.h"
+#include "FBX/InterchangeDispatcherFBXTasks.h"
 #include "InterchangeImportLog.h"
 #include "InterchangeManager.h"
 #include "InterchangeTranslatorHelper.h"
@@ -208,14 +208,14 @@ UInterchangeFbxTranslator::UInterchangeFbxTranslator()
 	if (GInterchangeFBXTranslatorUseWorker)
 	{
 		//Create the dispatcher
-		Dispatcher = MakeUnique<UE::Interchange::FInterchangeDispatcher>(ResultFolder);
+		Dispatcher = MakeUnique<UE::Interchange::FInterchangeDispatcher>(ResultFolder, 1);
 
 		if (ensure(Dispatcher.IsValid()))
 		{
 			Dispatcher->StartProcess();
 		}
 
-		if (Dispatcher->IsInterchangeWorkerRunning())
+		if (Dispatcher->IsValid())
 		{
 			bUseWorkerImport = true;
 		}
@@ -941,7 +941,8 @@ TArray<UE::Interchange::FAnimationPayloadData> UInterchangeFbxTranslator::GetAni
 
 FString UInterchangeFbxTranslator::CreateLoadFbxFileCommand(const FString& FbxFilePath, const bool bConvertScene, const bool bForceFrontXAxis, const bool bConvertSceneUnit, const bool bKeepFbxNamespace) const
 {
-	UE::Interchange::FJsonLoadSourceCmd LoadSourceCommand(TEXT("FBX"), FbxFilePath, bConvertScene, bForceFrontXAxis, bConvertSceneUnit, bKeepFbxNamespace);
+	using namespace UE::Interchange;
+	FJsonFBXLoadSourceCmd LoadSourceCommand(FJsonFBXLoadSourceCmd::TaskName, FbxFilePath, bConvertScene, bForceFrontXAxis, bConvertSceneUnit, bKeepFbxNamespace);
 	return LoadSourceCommand.ToJson();
 }
 
diff --git a/Engine/Source/Programs/InterchangeWorker/Private/InterchangeWorkerImpl.cpp b/Engine/Source/Programs/InterchangeWorker/Private/InterchangeWorkerImpl.cpp
--- a/Engine/Source/Programs/InterchangeWorker/Private/InterchangeWorkerImpl.cpp
+++ b/Engine/Source/Programs/InterchangeWorker/Private/InterchangeWorkerImpl.cpp
@@ -7,7 +7,7 @@
 #include "HAL/PlatformFileManager.h"
 #include "InterchangeCommands.h"
 #include "InterchangeDispatcherConfig.h"
-#include "InterchangeDispatcherTask.h"
+#include "FBX/InterchangeDispatcherFBXTasks.h"
 #include "InterchangeFbxParser.h"
 #include "Misc/FileHelper.h"
 #include "Misc/Paths.h"
@@ -221,7 +221,7 @@ void FInterchangeWorkerImpl::ProcessCommand(const TSharedPtr<UE::Interchange::IC
 	//Process the json and run the task
 	FString JSonResult;
 	TArray<FString> JSonMessages;
-	FJsonLoadSourceCmd LoadSourceCommand;
+	FJsonFBXLoadSourceCmd LoadSourceCommand;
 	FJsonFetchAnimationQueriesCmd FetchAnimationQueries;
 	FJsonFetchMeshPayloadCmd FetchMeshPayloadCommand;
 	FJsonFetchPayloadCmd FetchPayloadCommand;
@@ -229,7 +229,7 @@ void FInterchangeWorkerImpl::ProcessCommand(const TSharedPtr<UE::Interchange::IC
 	if (LoadSourceCommand.FromJson(JsonToProcess))
 	{
 		//Load file command
-		if (LoadSourceCommand.GetTranslatorID().Equals(TEXT("FBX"), ESearchCase::IgnoreCase))
+		if (LoadSourceCommand.GetTranslatorID().Equals(FJsonFBXLoadSourceCmd::TaskName))
 		{
 			//We want to load an FBX file
 			ProcessResult = LoadFbxFile(LoadSourceCommand, JSonResult, JSonMessages);
@@ -276,6 +276,12 @@ void FInterchangeWorkerImpl::ProcessCommand(const TSharedPtr<UE::Interchange::IC
 		CompletedTask.JSonResult = JSonResult;
 	}
 
+	UE_LOG(LogInterchangeWorker, Verbose, TEXT("CompletedTask.Result : %d\n %s"), (int32)ProcessResult , *CompletedTask.JSonResult);
+	int32 Index = 0;
+	for (const FString& Message : CompletedTask.JSonMessages)
+	{
+		UE_LOG(LogInterchangeWorker, Verbose, TEXT("CompletedTask.Message #%d : %s"), ++Index, *Message);
+	}
 	CommandIO.SendCommand(CompletedTask, Config::SendCommandTimeout_s);
 	UE_LOG(LogInterchangeWorker, Verbose, TEXT("End of Process %s"), *JsonToProcess);
 	
@@ -286,10 +292,11 @@ void FInterchangeWorkerImpl::ProcessCommand(const TSharedPtr<UE::Interchange::IC
 	}
 }
 
-ETaskState FInterchangeWorkerImpl::LoadFbxFile(const FJsonLoadSourceCmd& LoadSourceCommand, FString& OutJSonResult, TArray<FString>& OutJSonMessages)
+ETaskState FInterchangeWorkerImpl::LoadFbxFile(const FJsonFBXLoadSourceCmd& LoadSourceCommand, FString& OutJSonResult, TArray<FString>& OutJSonMessages)
 {
 	ETaskState ResultState = ETaskState::Unknown;
 	FString SourceFilename = LoadSourceCommand.GetSourceFilename();
+	UE_LOG(LogInterchangeWorker, Verbose, TEXT("LoadFbxFile : %s"), *SourceFilename);
 	FbxParser.Reset();
 	FbxParser.SetConvertSettings(LoadSourceCommand.GetDoesConvertScene(), LoadSourceCommand.GetDoesForceFrontXAxis(), LoadSourceCommand.GetDoesConvertSceneUnit(), LoadSourceCommand.GetDoesKeepFbxNamespace());
 	FbxParser.LoadFbxFile(SourceFilename, ResultFolder);
@@ -297,6 +304,7 @@ ETaskState FInterchangeWorkerImpl::LoadFbxFile(const FJsonLoadSourceCmd& LoadSou
 	ResultParser.SetResultFilename(FbxParser.GetResultFilepath());
 	OutJSonMessages = FbxParser.GetJsonLoadMessages();
 	OutJSonResult = ResultParser.ToJson();
+	UE_LOG(LogInterchangeWorker, Verbose, TEXT("ResultParser : %s"), *OutJSonResult);
 	ResultState = ETaskState::ProcessOk;
 	return ResultState;
 }
