Removed out-dated comment

#rnx

diff --git a/Engine/Plugins/Interchange/Runtime/Source/Import/Private/Mesh/InterchangeOBJTranslator.cpp b/Engine/Plugins/Interchange/Runtime/Source/Import/Private/Mesh/InterchangeOBJTranslator.cpp
--- a/Engine/Plugins/Interchange/Runtime/Source/Import/Private/Mesh/InterchangeOBJTranslator.cpp
+++ b/Engine/Plugins/Interchange/Runtime/Source/Import/Private/Mesh/InterchangeOBJTranslator.cpp
@@ -1216,19 +1216,6 @@ bool UInterchangeOBJTranslator::Translate(UInterchangeBaseNodeContainer& BaseNod
 
 	ObjDataPtr->ObjFilename = Filename;
 
-
-	// @TODO: API design
-	// The Translate() method is const, implying that the translator object should not be modified.
-	// However, in this particular case, the translator needs to build some data which is required when the payload is requested.
-	// In the case of FBX, this happens in an adjacent process.
-	// In our case, we need to put it somewhere, and we use a separate object to hold it; however this is not really in the spirit of
-	// 'const', we are effectively using the loophole that an owned object need not be const even if the owner is const.
-	//
-	// We need to think about the semantics of this a bit better.
-	// Is the translator *really* to be considered a const object?  Should it be allowed to hold state?
-	// Would it be better to shift the payload ownership to the nodes which the translator creates?
-	// Do we maybe want to decouple translators from payloads (which arguably makes more sense).
-
 	using namespace UE::Interchange;
 	using namespace ObjTranslatorUtils;
 
