[Interchange] Reducing SkeletalMesh import memory footprint

Same optimization than the one done on StaticMeshes in CL 44753863
While testing with a 3M triangles/one bone skeletal mesh FBX, memory was reduced by around 1.3Gb (around 15% of total import)
Memory delta usage went from 4.75Gb to 2.53Gb

Shelf #2: Balazs pointed out that SkeletalMeshes might still call GetMeshDescription after ImportAsset, which - contrary to StaticMeshes - is true. This was reducing the impact of the desired optimization.
I moved the clearing in the FinalizeObject (after the SkeletalMeshes build that was calling GetMeshDescriptions) and this improves the import memory delta without affecting the memory spike.

[RN] Reduced memory usage when importing Skeletal Meshes

diff --git a/Engine/Plugins/Interchange/Runtime/Source/Import/Private/Mesh/InterchangeSkeletalMeshFactory.cpp b/Engine/Plugins/Interchange/Runtime/Source/Import/Private/Mesh/InterchangeSkeletalMeshFactory.cpp
--- a/Engine/Plugins/Interchange/Runtime/Source/Import/Private/Mesh/InterchangeSkeletalMeshFactory.cpp
+++ b/Engine/Plugins/Interchange/Runtime/Source/Import/Private/Mesh/InterchangeSkeletalMeshFactory.cpp
@@ -1657,6 +1657,7 @@ UInterchangeFactoryBase::FImportAssetResult UInterchangeSkeletalMeshFactory::Imp
 			}
 		}
 		
+		PayloadsPerLodIndex.Remove(LodIndex);
 		
 		CurrentLodIndex++;
 	}
@@ -2397,6 +2398,11 @@ void UInterchangeSkeletalMeshFactory::FinalizeObject_GameThread(const FSetupObje
 		//Release the skeletal mesh async properties
 		LockEvent->Trigger();
 	}
+#if WITH_EDITORONLY_DATA
+	//Clear the skeletal mesh descriptions to free memory.
+	//Note that we cannot do this earlier due to the MeshDescriptions being used when building the skeletal mesh.
+	SkeletalMesh->ClearAllMeshDescriptions();
+#endif
 #endif //WITH_EDITOR
 }
 
diff --git a/Engine/Plugins/Interchange/Runtime/Source/Import/Private/Mesh/InterchangeStaticMeshFactory.cpp b/Engine/Plugins/Interchange/Runtime/Source/Import/Private/Mesh/InterchangeStaticMeshFactory.cpp
--- a/Engine/Plugins/Interchange/Runtime/Source/Import/Private/Mesh/InterchangeStaticMeshFactory.cpp
+++ b/Engine/Plugins/Interchange/Runtime/Source/Import/Private/Mesh/InterchangeStaticMeshFactory.cpp
@@ -936,16 +936,16 @@ UInterchangeFactoryBase::FImportAssetResult UInterchangeStaticMeshFactory::EndIm
 		}
 	}
 	FMeshBudgetProjectSettingsUtils::SetLodGroupForStaticMesh(StaticMesh);
-#endif
 
-#if WITH_EDITOR
 	if (bReimport)
 	{
 		UStaticMesh::RemoveUnusedMaterialSlots(StaticMesh);
 	}
 
+#if WITH_EDITORONLY_DATA
 	// Clear the mesh descriptions to free memory, we don't need them anymore since all the meshes have been committed at this point.
 	StaticMesh->ClearMeshDescriptions();
+#endif
 #endif
 
 	ImportAssetResult.ImportedObject = StaticMesh;
