[Interchange] : Add Support to import OBJ into level

diff --git a/Engine/Plugins/Interchange/Runtime/Source/Import/Private/Mesh/InterchangeOBJTranslator.cpp b/Engine/Plugins/Interchange/Runtime/Source/Import/Private/Mesh/InterchangeOBJTranslator.cpp
--- a/Engine/Plugins/Interchange/Runtime/Source/Import/Private/Mesh/InterchangeOBJTranslator.cpp
+++ b/Engine/Plugins/Interchange/Runtime/Source/Import/Private/Mesh/InterchangeOBJTranslator.cpp
@@ -7,6 +7,7 @@
 #include "InterchangeMaterialDefinitions.h"
 #include "InterchangeMeshDefinitions.h"
 #include "InterchangeMeshNode.h"
+#include "InterchangeSceneNode.h"
 #include "InterchangeMaterialInstanceNode.h"
 #include "InterchangeTexture2DNode.h"
 #include "InterchangeTextureNode.h"
@@ -33,6 +34,13 @@ static FAutoConsoleVariableRef CCvarInterchangeEnableOBJImport(
 	TEXT("Whether OBJ support is enabled."),
 	ECVF_Default);
 
+static bool GInterchangeEnableOBJLevelImport = false;
+static FAutoConsoleVariableRef CCvarInterchangeEnableOBJLevelImport(
+	TEXT("Interchange.FeatureFlags.Import.OBJ.ToLevel"),
+	GInterchangeEnableOBJLevelImport,
+	TEXT("Whether support for OBJ level import is enabled."),
+	ECVF_Default);
+
 
 #define INTERCHANGE_OBJ_TRANSLATOR_LOG_ONCE(Verbosity, Format, ...) \
 { \
@@ -1022,6 +1030,16 @@ namespace ObjTranslatorUtils
 		return Scalar >= 0.f;
 	};
 
+	static FString MakeSceneNodeUid(const FString& Name)
+	{
+		return TEXT("\\Scene\\") + (Name.IsEmpty() ? FString(TEXT("Null")) : Name);
+	}
+
+	static FString MakeActorNodeUid(const FString& Name)
+	{
+		return TEXT("\\Actor\\") + (Name.IsEmpty() ? FString(TEXT("Null")) : Name);
+	}
+
 	/**
 	 * Build and return a UID name for a mesh node.
 	 * @todo: move this to be a static method on the MeshNode class itself?
@@ -1204,6 +1222,11 @@ EInterchangeTranslatorAssetType UInterchangeOBJTranslator::GetSupportedAssetType
 	return EInterchangeTranslatorAssetType::Materials | EInterchangeTranslatorAssetType::Meshes;
 }
 
+EInterchangeTranslatorType UInterchangeOBJTranslator::GetTranslatorType() const
+{
+	return GInterchangeEnableOBJLevelImport ? EInterchangeTranslatorType::Scenes : EInterchangeTranslatorType::Assets;
+}
+
 bool UInterchangeOBJTranslator::Translate(UInterchangeBaseNodeContainer& BaseNodeContainer) const
 {
 	// Get filename from SourceData
@@ -1244,15 +1267,27 @@ bool UInterchangeOBJTranslator::Translate(UInterchangeBaseNodeContainer& BaseNod
 			return false;
 		}
 
+		// Add base scene node.
+		UInterchangeSceneNode* SceneNode = NewObject< UInterchangeSceneNode >(&BaseNodeContainer);
+		const FString SceneName = FPaths::GetBaseFilename(Filename);
+		const FString SceneNodeUid = MakeSceneNodeUid(SceneName);
+		BaseNodeContainer.SetupNode(SceneNode, SceneNodeUid, SceneName, EInterchangeNodeContainerType::TranslatedScene);
+
 		// Add mesh nodes to the container for each mesh group
 
 		for (const TPair<FString, FObjData::FGroupData>& Group : ObjDataPtr->Groups)
 		{
 			const FString& GroupName = Group.Key;
-			FString NodeUid = MakeMeshNodeUid(GroupName);
 
+			const FString ActorNodeUid = MakeActorNodeUid(GroupName);
+			const FString MeshNodeUid = MakeMeshNodeUid(GroupName);
+
+			UInterchangeSceneNode* InterchangeSceneNode = NewObject<UInterchangeSceneNode>(&BaseNodeContainer);
+			BaseNodeContainer.SetupNode(InterchangeSceneNode, ActorNodeUid, GroupName, EInterchangeNodeContainerType::TranslatedScene, SceneNodeUid);
+			InterchangeSceneNode->SetCustomAssetInstanceUid(MeshNodeUid);
+			
 			UInterchangeMeshNode* MeshNode = NewObject<UInterchangeMeshNode>(&BaseNodeContainer);
-			BaseNodeContainer.SetupNode(MeshNode, NodeUid, GroupName, EInterchangeNodeContainerType::TranslatedAsset);
+			BaseNodeContainer.SetupNode(MeshNode, MeshNodeUid, GroupName, EInterchangeNodeContainerType::TranslatedAsset);
 
 			// The payload key is the group name.
 			// A MeshDescription is generated for each group of faces.
diff --git a/Engine/Plugins/Interchange/Runtime/Source/Import/Public/Mesh/InterchangeOBJTranslator.h b/Engine/Plugins/Interchange/Runtime/Source/Import/Public/Mesh/InterchangeOBJTranslator.h
--- a/Engine/Plugins/Interchange/Runtime/Source/Import/Public/Mesh/InterchangeOBJTranslator.h
+++ b/Engine/Plugins/Interchange/Runtime/Source/Import/Public/Mesh/InterchangeOBJTranslator.h
@@ -31,6 +31,7 @@ public:
 
 	UE_API virtual TArray<FString> GetSupportedFormats() const override;
 	UE_API virtual EInterchangeTranslatorAssetType GetSupportedAssetTypes() const override;
+	UE_API virtual EInterchangeTranslatorType GetTranslatorType() const override;
 
 	/**
 	 * Translate the associated source data into a node hold by the specified nodes container.
