Interchange FBX:
- Fix bump map texture used as a normalmap instead of a bump map

diff --git a/Engine/Plugins/Interchange/Runtime/Source/Parsers/Fbx/Private/FbxMaterial.cpp b/Engine/Plugins/Interchange/Runtime/Source/Parsers/Fbx/Private/FbxMaterial.cpp
--- a/Engine/Plugins/Interchange/Runtime/Source/Parsers/Fbx/Private/FbxMaterial.cpp
+++ b/Engine/Plugins/Interchange/Runtime/Source/Parsers/Fbx/Private/FbxMaterial.cpp
@@ -277,7 +277,12 @@ namespace UE
 					{
 						const float FactorBump = GetFactor(FbxSurfaceMaterial::sBumpFactor);
 						DefaultValue.Set<float>(0.f);
-						bHasInput |= ConnectInput(Phong::Parameters::Normal, FbxSurfaceMaterial::sBump, FactorBump, DefaultValue, false);
+						bool bHasBump = ConnectInput(Phong::Parameters::Normal, FbxSurfaceMaterial::sBump, FactorBump, DefaultValue, false);
+						if (bHasBump)
+						{
+							MaterialInstanceNode->AddStaticSwitchParameterValue(TEXT("bHasBump"), true);
+						}
+						bHasInput |= bHasBump;
 					}
 				}
 
