[In Progress] Interchange CAD Importer plugin
- Rewrote the FInterchangeCADDispatcher class to usethe FInterchangeDispatcher one
- Modified the FInterchangeCADTranslator class to adjust to the rewritten  FInterchangeCADDispatcher class
- Fixed issue reported by BugHawk

#rnx

diff --git a/Engine/Plugins/Interchange/Runtime/Source/Dispatcher/Private/InterchangeWorkerHandler.cpp b/Engine/Plugins/Interchange/Runtime/Source/Dispatcher/Private/InterchangeWorkerHandler.cpp
--- a/Engine/Plugins/Interchange/Runtime/Source/Dispatcher/Private/InterchangeWorkerHandler.cpp
+++ b/Engine/Plugins/Interchange/Runtime/Source/Dispatcher/Private/InterchangeWorkerHandler.cpp
@@ -358,7 +358,7 @@ namespace UE
 
 		bool FInterchangeWorkerHandler::IsAlive() const
 		{
-			return !EnumHasAnyFlags(EWorkerState(WorkerState), EWorkerState::Terminated | EWorkerState::Closing);
+			return !EnumHasAnyFlags(EWorkerState(WorkerState), EWorkerState::Terminated);
 		}
 
 		void FInterchangeWorkerHandler::ProcessCommand(ICommand& Command)
