InterchangeAssets:
- Enable BaseMaterial plugin for the MF_DitherMask

Interchange:
- Add methods to disconnect inputs from a shader node

MaterialX:
- Add a material function to apply a mask dithering on surface materials that have both transmission and subsurface inputs
- Add a custom nodedef for sparse volume textures

USD:
- Convert geompropvalue nodes coming from volume materials to sparse volume nodes
- Assign the SVT asset to the right SVT Sample material expression

#virtualized

diff --git a/Engine/Plugins/Interchange/Assets/InterchangeAssets.uplugin b/Engine/Plugins/Interchange/Assets/InterchangeAssets.uplugin
--- a/Engine/Plugins/Interchange/Assets/InterchangeAssets.uplugin
+++ b/Engine/Plugins/Interchange/Assets/InterchangeAssets.uplugin
@@ -21,5 +21,11 @@
 			"Type": "Runtime",
 			"LoadingPhase": "PostConfigInit"
 		}
+	],
+	"Plugins": [
+		{
+			"Name": "BaseMaterial",
+			"Enabled": true
+		}
 	]
 }
\ No newline at end of file
diff --git a/Engine/Plugins/Interchange/Extensions/OpenUSD/Source/Import/Private/InterchangeUSDPipeline.cpp b/Engine/Plugins/Interchange/Extensions/OpenUSD/Source/Import/Private/InterchangeUSDPipeline.cpp
--- a/Engine/Plugins/Interchange/Extensions/OpenUSD/Source/Import/Private/InterchangeUSDPipeline.cpp
+++ b/Engine/Plugins/Interchange/Extensions/OpenUSD/Source/Import/Private/InterchangeUSDPipeline.cpp
@@ -22,12 +22,14 @@
 #include "DynamicMeshToMeshDescription.h"
 #include "Engine/StaticMesh.h"
 #include "Materials/Material.h"
+#include "Materials/MaterialExpressionSparseVolumeTextureSample.h"
 #include "Materials/MaterialExpressionTextureBase.h"
 #include "Materials/MaterialExpressionTextureSample.h"
 #include "MeshDescriptionToDynamicMesh.h"
 #include "Parameterization/PatchBasedMeshUVGenerator.h"
 #include "Sampling/MeshMapBaker.h"
 #include "Sampling/MeshPropertyMapEvaluator.h"
+#include "SparseVolumeTexture/SparseVolumeTexture.h"
 #include "StaticMeshAttributes.h"
 
 #include "UnrealUSDWrapper.h"
@@ -303,6 +305,17 @@ namespace UE::InterchangeUsdPipeline::Private
 				}
 			}
 		}
+
+		int32 NumberPrimvar = 0;
+		VolumeNode->GetInt32Attribute(Primvar::Number, NumberPrimvar);
+		FactoryNode->AddInt32Attribute(Primvar::Number, NumberPrimvar);
+		for (int32 Index = 0; Index < NumberPrimvar; ++Index)
+		{
+			FString PrimvarName;
+			FString Attribute = Primvar::ShaderNodeSparseVolumeTextureSample + FString::FromInt(Index);
+			VolumeNode->GetStringAttribute(Attribute, PrimvarName);
+			FactoryNode->AddStringAttribute(Attribute, PrimvarName);
+		}
 	}
 
 	void ProcessHeterogeneousVolumeSceneNode(const UInterchangeSceneNode* SceneNode, UInterchangeBaseNodeContainer& NodeContainer)
@@ -1274,9 +1287,7 @@ namespace UE::InterchangeUsdPipeline::Private
 
 			FString ShaderNodeTextureSampleUID;
 			FactoryNode->GetStringAttribute(USD::Primvar::ShaderNodeTextureSample + FString::FromInt(IndexPrimvar), ShaderNodeTextureSampleUID);
-			if (UInterchangeMaterialExpressionFactoryNode* TextureSampleExpressionFactoryNode = Cast<UInterchangeMaterialExpressionFactoryNode>(
-					NodeContainer->GetFactoryNode(ShaderNodeTextureSampleUID)
-				))
+			if (UInterchangeMaterialExpressionFactoryNode* TextureSampleExpressionFactoryNode = Cast<UInterchangeMaterialExpressionFactoryNode>(NodeContainer->GetFactoryNode(ShaderNodeTextureSampleUID)))
 			{
 				if (FString ExpressionPath;
 					TextureSampleExpressionFactoryNode->GetStringAttribute(Factory::Expression::Path.ToString(), ExpressionPath))
@@ -1429,6 +1440,70 @@ void UInterchangeUsdPipeline::ExecutePostFactoryPipeline(
 			Baker.Bake();
 		}
 	}
+
+	// Assign the right SVT asset to the corresponding SVT material expression
+	if (USparseVolumeTexture* SparseVolumeAsset = Cast<USparseVolumeTexture>(CreatedAsset))
+	{
+		using namespace UE::Interchange::USD;
+		using namespace UE::Interchange::Materials;
+
+		TArray<UMaterialExpression*> MaterialExpressions;
+		TSet<UMaterial*> Materials;
+		TSet<FString> VolumePrimvars; // if one of the names is the same as the attribute in the ShaderNode then we'll assign the SVT asset
+
+		int32 Number = 0;
+		FactoryNode->GetInt32Attribute(Primvar::Number, Number);
+		VolumePrimvars.Reserve(Number);
+		for (int32 Index = 0; Index < Number; ++Index)
+		{
+			FString SVTSampleName;
+			FactoryNode->GetStringAttribute(Primvar::ShaderNodeSparseVolumeTextureSample + FString::FromInt(Index), SVTSampleName);
+			VolumePrimvars.Emplace(SVTSampleName);
+		}
+
+		using UMaterialExpressionSVTSample = UMaterialExpressionSparseVolumeTextureSampleParameter;
+
+		BaseNodeContainer->BreakableIterateNodesOfType<UInterchangeShaderNode>(
+			[BaseNodeContainer, SparseVolumeAsset, &MaterialExpressions, &VolumePrimvars](const FString&, UInterchangeShaderNode* ShaderNode)
+			{
+				if (FString SparseVolumeName; ShaderNode->GetStringAttribute(Primvar::ShaderNodeSparseVolumeTextureSample, SparseVolumeName))
+				{
+					if (VolumePrimvars.Contains(SparseVolumeName))
+					{
+						if (UInterchangeMaterialExpressionFactoryNode* MaterialExpressionFactoryNode =
+							Cast<UInterchangeMaterialExpressionFactoryNode>(BaseNodeContainer->GetFactoryNode(UInterchangeFactoryBaseNode::BuildFactoryNodeUid(ShaderNode->GetUniqueID()))))
+						{
+							FString ExpressionPath;
+							MaterialExpressionFactoryNode->GetStringAttribute(Factory::Expression::Path.ToString(), ExpressionPath);
+
+							if (UMaterialExpressionSVTSample* SparseVolumeSample =
+								Cast<UMaterialExpressionSVTSample>(FSoftObjectPath{ ExpressionPath }.TryLoad()))
+							{
+								SparseVolumeSample->SparseVolumeTexture = SparseVolumeAsset;
+								MaterialExpressions.Emplace(SparseVolumeSample);
+								return true;
+							}
+						}
+					}
+				}
+				return false;
+			});
+
+		// Update materials and material expressions rendering
+		for (UMaterialExpression* Expression : MaterialExpressions)
+		{
+			Materials.FindOrAdd(Expression->Material);
+			Expression->PostEditChange();
+			Expression->Modify();
+		}
+
+		for (UMaterial* Material : Materials)
+		{
+			Material->PostEditChange();
+			Material->ForceRecompileForRendering();
+			Material->MarkPackageDirty();
+		}
+	}
 #endif	  // USE_USD_SDK && WITH_EDITORONLY_DATA
 }
 
diff --git a/Engine/Plugins/Interchange/Extensions/OpenUSD/Source/Import/Private/InterchangeUsdTranslator.cpp b/Engine/Plugins/Interchange/Extensions/OpenUSD/Source/Import/Private/InterchangeUsdTranslator.cpp
--- a/Engine/Plugins/Interchange/Extensions/OpenUSD/Source/Import/Private/InterchangeUsdTranslator.cpp
+++ b/Engine/Plugins/Interchange/Extensions/OpenUSD/Source/Import/Private/InterchangeUsdTranslator.cpp
@@ -3163,7 +3163,14 @@ namespace UE::InterchangeUsdTranslator::Private
 					);
 
 					if (VolumeNode)
-					{
+					{						
+						// The USD shader graph, in case of a volume material, reference the OpenVDBAsset as a primvar
+						// In the MaterialX Translator we create a SparseVolumeTextureSampleParameter that references this 'primvar' name, the SparseVolumeTexture asset
+						VolumeNode->AddInt32Attribute(Primvar::Number, VolumePrimInfo.VolumeFieldNames.Num());
+						for (int32 IndexPrimvar = 0; IndexPrimvar < VolumePrimInfo.VolumeFieldNames.Num(); ++IndexPrimvar)
+						{
+							VolumeNode->AddStringAttribute(Primvar::ShaderNodeSparseVolumeTextureSample + FString::FromInt(IndexPrimvar), VolumePrimInfo.VolumeFieldNames[IndexPrimvar]);
+						}
 						AnimationIDToVolumeNode.Add(AnimationID, VolumeNode);
 					}
 				}
diff --git a/Engine/Plugins/Interchange/Runtime/Source/Common/Public/MaterialX/InterchangeMaterialXDefinitions.h b/Engine/Plugins/Interchange/Runtime/Source/Common/Public/MaterialX/InterchangeMaterialXDefinitions.h
--- a/Engine/Plugins/Interchange/Runtime/Source/Common/Public/MaterialX/InterchangeMaterialXDefinitions.h
+++ b/Engine/Plugins/Interchange/Runtime/Source/Common/Public/MaterialX/InterchangeMaterialXDefinitions.h
@@ -306,6 +306,7 @@ namespace UE
 			constexpr const TCHAR* BreakOutFloat2Components = TEXT("/Engine/Functions/Engine_MaterialFunctions02/Utility/BreakOutFloat2Components.BreakOutFloat2Components");
 			constexpr const TCHAR* BreakOutFloat3Components = TEXT("/Engine/Functions/Engine_MaterialFunctions02/Utility/BreakOutFloat3Components.BreakOutFloat3Components");
 			constexpr const TCHAR* BreakOutFloat4Components = TEXT("/Engine/Functions/Engine_MaterialFunctions02/Utility/BreakOutFloat4Components.BreakOutFloat4Components");
+			constexpr const TCHAR* DitherMask= TEXT("/InterchangeAssets/Substrate/MF_DitherMask.MF_DitherMask");
 			constexpr const TCHAR* HeightToNormalSmooth = TEXT("/Engine/Functions/Engine_MaterialFunctions03/Procedurals/HeightToNormalSmooth.HeightToNormalSmooth");
 			constexpr const TCHAR* MxACEScgToRec709 = TEXT("/InterchangeAssets/Functions/MX_ACEScgToRec709.MX_ACEScgToRec709");
 			constexpr const TCHAR* MxAdobeRGBToRec709 = TEXT("/InterchangeAssets/Functions/MX_AdobeRGBToRec709.MX_AdobeRGBToRec709");
@@ -671,6 +672,7 @@ MATERIALX_NAMESPACE_BEGIN
 		static constexpr const char* ParentName = "UE:ParentName";
 		static constexpr const char* UniqueName = "UE:UniqueName";
 		static constexpr const char* GeomPropImage = "UE:GeomPropImage";
+		static constexpr const char* GeomPropSparseVolume = "UE:GeomPropSparseVolume";
 	}
 
 	namespace Category
@@ -871,6 +873,8 @@ MATERIALX_NAMESPACE_BEGIN
 		static constexpr const char* ViewDirection = "viewdirection";
 		// Volume shaders
 		static constexpr const char* Volume = "volume";
+		// Unreal Engine nodes defined in Engine\Binaries\ThirdParty\MaterialX\libraries\Interchange\Interchange_defs.mtlx
+		static constexpr const char* SparseVolume = "sparse_volume";
 	}
 
 	namespace NodeDefinition
diff --git a/Engine/Plugins/Interchange/Runtime/Source/Common/Public/Usd/InterchangeUsdDefinitions.h b/Engine/Plugins/Interchange/Runtime/Source/Common/Public/Usd/InterchangeUsdDefinitions.h
--- a/Engine/Plugins/Interchange/Runtime/Source/Common/Public/Usd/InterchangeUsdDefinitions.h
+++ b/Engine/Plugins/Interchange/Runtime/Source/Common/Public/Usd/InterchangeUsdDefinitions.h
@@ -78,6 +78,8 @@ namespace UE::Interchange::USD
 		const FString TangentSpace = TEXT("USD_PrimvarTangentSpace");
 		// Attribute that informs about the UID of a ShaderNode of type TextureSample
 		const FString ShaderNodeTextureSample = TEXT("USD_ShaderNodeTextureSample");
+		// Attribute that informs about the UID of a ShaderNode of type SparseVolumeTextureSample
+		const FString ShaderNodeSparseVolumeTextureSample = TEXT("USD_ShaderNodeSparseVolumeTextureSample");
 		// Attribute that informs how we should handle primvars on MeshDescriptions
 		const FString Import = TEXT("USD_Import_Primvars");
 	}
diff --git a/Engine/Plugins/Interchange/Runtime/Source/Import/Private/MaterialX/MaterialXUtils/MaterialXManager.cpp b/Engine/Plugins/Interchange/Runtime/Source/Import/Private/MaterialX/MaterialXUtils/MaterialXManager.cpp
--- a/Engine/Plugins/Interchange/Runtime/Source/Import/Private/MaterialX/MaterialXUtils/MaterialXManager.cpp
+++ b/Engine/Plugins/Interchange/Runtime/Source/Import/Private/MaterialX/MaterialXUtils/MaterialXManager.cpp
@@ -357,6 +357,7 @@ FMaterialXManager::FMaterialXManager()
 	MaterialFunctionPackage(UE::MaterialFunctions::Path::MxOr);
 	MaterialFunctionPackage(UE::MaterialFunctions::Path::MxNot);
 	MaterialFunctionPackage(UE::MaterialFunctions::Path::MxXor);
+	MaterialFunctionPackage(UE::MaterialFunctions::Path::DitherMask);
 
 	if (bIsSubstrateEnabled)
 	{
diff --git a/Engine/Plugins/Interchange/Runtime/Source/Import/Private/MaterialX/MaterialXUtils/MaterialXOpenPBRSurfaceShader.cpp b/Engine/Plugins/Interchange/Runtime/Source/Import/Private/MaterialX/MaterialXUtils/MaterialXOpenPBRSurfaceShader.cpp
--- a/Engine/Plugins/Interchange/Runtime/Source/Import/Private/MaterialX/MaterialXUtils/MaterialXOpenPBRSurfaceShader.cpp
+++ b/Engine/Plugins/Interchange/Runtime/Source/Import/Private/MaterialX/MaterialXUtils/MaterialXOpenPBRSurfaceShader.cpp
@@ -23,6 +23,8 @@ TSharedRef<FMaterialXBase> FMaterialXOpenPBRSurfaceShader::MakeInstance(UInterch
 UInterchangeBaseNode* FMaterialXOpenPBRSurfaceShader::Translate(MaterialX::NodePtr OpenPBRSurfaceNode)
 {
 	using namespace UE::Interchange::Materials;
+	using namespace UE::Interchange::MaterialX;
+	using namespace UE::Interchange::Materials::Standard::Nodes;
 
 	this->SurfaceShaderNode = OpenPBRSurfaceNode;
 
@@ -39,19 +41,127 @@ UInterchangeBaseNode* FMaterialXOpenPBRSurfaceShader::Translate(MaterialX::NodeP
 
 	if(UInterchangeShaderPortsAPI::HasInput(OpenPBRSurfaceShaderNode, OpenPBRSurface::Parameters::TransmissionWeight))
 	{
-		OpenPBRSurfaceShaderNode->AddInt32Attribute(UE::Interchange::MaterialX::Attributes::EnumType, UE::Interchange::MaterialX::IndexSurfaceShaders);
-		OpenPBRSurfaceShaderNode->AddInt32Attribute(UE::Interchange::MaterialX::Attributes::EnumValue, int32(EInterchangeMaterialXShaders::OpenPBRSurfaceTransmission));
+		OpenPBRSurfaceShaderNode->AddInt32Attribute(Attributes::EnumType, IndexSurfaceShaders);
+		OpenPBRSurfaceShaderNode->AddInt32Attribute(Attributes::EnumValue, int32(EInterchangeMaterialXShaders::OpenPBRSurfaceTransmission));
 		ShaderGraphNode->SetCustomBlendMode(EBlendMode::BLEND_TranslucentColoredTransmittance);
 	}
 
 	if(bIsSubstrateEnabled)
 	{
-		UInterchangeShaderPortsAPI::ConnectOuputToInputByName(ShaderGraphNode, SubstrateMaterial::Parameters::FrontMaterial.ToString(), OpenPBRSurfaceShaderNode->GetUniqueID(), OpenPBRSurface::SubstrateMaterial::Outputs::FrontMaterial.ToString());
-
-		if(UInterchangeShaderPortsAPI::HasInput(OpenPBRSurfaceShaderNode, OpenPBRSurface::Parameters::GeometryOpacity))
+		auto CreateShaderNodeInternal = [this]<class T>(mx::NodePtr OpenPBRSurfaceNode, const FString & PathOrShaderType, const FString & DisplayLabel)
+		{
+			T* Node;
+
+			const FString NodeUID = UInterchangeShaderNode::MakeNodeUid(GetUniqueName(OpenPBRSurfaceNode) + DisplayLabel, FStringView{});
+
+			if (Node = const_cast<T*>(Cast<T>(NodeContainer.GetNode(NodeUID))); !Node)
+			{
+				Node = NewObject<T>(&NodeContainer);
+				NodeContainer.SetupNode(Node, NodeUID, DisplayLabel, EInterchangeNodeContainerType::TranslatedAsset);
+				if constexpr (std::is_same_v<T, UInterchangeFunctionCallShaderNode>)
+				{
+					Node->SetCustomMaterialFunction(PathOrShaderType);
+				}
+				else
+				{
+					Node->SetCustomShaderType(PathOrShaderType);
+				}
+			}
+
+			return Node;
+		};
+
+		if (UInterchangeShaderPortsAPI::HasInput(OpenPBRSurfaceShaderNode, OpenPBRSurface::Parameters::TransmissionWeight) &&
+			UInterchangeShaderPortsAPI::HasInput(OpenPBRSurfaceShaderNode, OpenPBRSurface::Parameters::SubsurfaceWeight) &&
+			UInterchangeShaderPortsAPI::HasInput(OpenPBRSurfaceShaderNode, OpenPBRSurface::Parameters::BaseColor))
 		{
-			UInterchangeShaderPortsAPI::ConnectOuputToInputByName(ShaderGraphNode, SubstrateMaterial::Parameters::OpacityMask.ToString(), OpenPBRSurfaceShaderNode->GetUniqueID(), OpenPBRSurface::SubstrateMaterial::Outputs::OpacityMask.ToString());
+			// 1. Set the Material as Opaque with Mask Mode
+			OpenPBRSurfaceShaderNode->AddInt32Attribute(Attributes::EnumType, IndexSurfaceShaders);
+			OpenPBRSurfaceShaderNode->AddInt32Attribute(Attributes::EnumValue, int32(EInterchangeMaterialXShaders::OpenPBRSurface));
 			ShaderGraphNode->SetCustomBlendMode(EBlendMode::BLEND_Masked);
+
+			//2. Disconnect the transmission inputs
+			FString TransmissionNodeUid, TransmissionOutputName;
+			UInterchangeShaderPortsAPI::DisconnectInputFromOutputNode(OpenPBRSurfaceShaderNode, OpenPBRSurface::Parameters::TransmissionWeight.ToString(), TransmissionNodeUid, TransmissionOutputName);
+			UInterchangeShaderPortsAPI::DisconnectInput(OpenPBRSurfaceShaderNode, OpenPBRSurface::Parameters::TransmissionColor.ToString());
+			UInterchangeShaderPortsAPI::DisconnectInput(OpenPBRSurfaceShaderNode, OpenPBRSurface::Parameters::TransmissionDepth.ToString());
+			UInterchangeShaderPortsAPI::DisconnectInput(OpenPBRSurfaceShaderNode, OpenPBRSurface::Parameters::TransmissionDispersionAbbeNumber.ToString());
+			UInterchangeShaderPortsAPI::DisconnectInput(OpenPBRSurfaceShaderNode, OpenPBRSurface::Parameters::TransmissionDispersionScale.ToString());
+			UInterchangeShaderPortsAPI::DisconnectInput(OpenPBRSurfaceShaderNode, OpenPBRSurface::Parameters::TransmissionScatterAnisotropy.ToString());
+			UInterchangeShaderPortsAPI::DisconnectInput(OpenPBRSurfaceShaderNode, OpenPBRSurface::Parameters::TransmissionScatter.ToString());
+
+			//3. Create DitherMask node
+			UInterchangeShaderNode* DitherMaskNode = CreateShaderNodeInternal.operator()<UInterchangeFunctionCallShaderNode>(OpenPBRSurfaceNode, UE::MaterialFunctions::Path::DitherMask, TEXT("DitherMaskNode"));
+
+			//4. Connect the surface shader and the transmission to the MF_DitherMask
+			if (const UInterchangeShaderNode* TransmissionNode = Cast<UInterchangeShaderNode>(NodeContainer.GetNode(TransmissionNodeUid)))
+			{
+				UInterchangeShaderPortsAPI::ConnectOuputToInputByName(DitherMaskNode, TEXT("BSDF"), OpenPBRSurfaceShaderNode->GetUniqueID(), OpenPBRSurface::SubstrateMaterial::Outputs::FrontMaterial.ToString());
+				UInterchangeShaderPortsAPI::ConnectOuputToInputByName(DitherMaskNode, TEXT("transmission"), TransmissionNode->GetUniqueID(), TransmissionOutputName);
+			}
+
+			UInterchangeShaderPortsAPI::ConnectOuputToInputByName(ShaderGraphNode, SubstrateMaterial::Parameters::FrontMaterial.ToString(), DitherMaskNode->GetUniqueID(), TEXT("Dither"));
+			UInterchangeShaderPortsAPI::ConnectOuputToInputByName(ShaderGraphNode, SubstrateMaterial::Parameters::OpacityMask.ToString(), DitherMaskNode->GetUniqueID(), SubstrateMaterial::Parameters::OpacityMask.ToString());
+		}
+		else if (UInterchangeShaderPortsAPI::HasInput(OpenPBRSurfaceShaderNode, OpenPBRSurface::Parameters::TransmissionWeight) &&
+				 UInterchangeShaderPortsAPI::HasInput(OpenPBRSurfaceShaderNode, OpenPBRSurface::Parameters::SubsurfaceWeight))
+		{
+			ShaderGraphNode->SetCustomBlendMode(EBlendMode::BLEND_TranslucentColoredTransmittance);
+			auto ReplaceInputConnection = [this, OpenPBRSurfaceShaderNode](const FString& InputToRemove, const FString& InputName)
+				{
+					UInterchangeShaderPortsAPI::DisconnectInput(OpenPBRSurfaceShaderNode, InputToRemove);
+
+					FString NodeUid, NodeOutputName;
+					UInterchangeShaderPortsAPI::GetInputConnection(OpenPBRSurfaceShaderNode, InputName, NodeUid, NodeOutputName);
+					if (const UInterchangeShaderNode* ShaderNode = Cast<UInterchangeShaderNode>(NodeContainer.GetNode(NodeUid)))
+					{
+						UInterchangeShaderPortsAPI::ConnectOuputToInputByName(OpenPBRSurfaceShaderNode, InputToRemove, ShaderNode->GetUniqueID(), NodeOutputName);
+					}
+				};
+
+			//1. subsurface_color input now replaces the base_color
+			if (UInterchangeShaderPortsAPI::HasInput(OpenPBRSurfaceShaderNode, OpenPBRSurface::Parameters::SubsurfaceColor))
+			{
+				ReplaceInputConnection(OpenPBRSurface::Parameters::BaseColor.ToString(), OpenPBRSurface::Parameters::SubsurfaceColor.ToString());
+			}
+
+			//2. replace base by subsurface
+			ReplaceInputConnection(OpenPBRSurface::Parameters::BaseWeight.ToString(), OpenPBRSurface::Parameters::SubsurfaceWeight.ToString());
+
+			//3. scale the subsurface_radius by 600 (empirical value that gives better results, need further testing though)
+			FString SubsurfaceRadiusNodeUid, SubsurfaceRadiusNodeOutputName;
+			if (UInterchangeShaderPortsAPI::GetInputConnection(OpenPBRSurfaceShaderNode, OpenPBRSurface::Parameters::SubsurfaceRadius.ToString(), SubsurfaceRadiusNodeUid, SubsurfaceRadiusNodeOutputName))
+			{
+				UInterchangeShaderNode* ScaleSubsurfaceRadiusNode = CreateShaderNodeInternal.operator()<UInterchangeShaderNode>(OpenPBRSurfaceNode, Multiply::Name.ToString(), TEXT("ScaleSubsurfaceRadius"));
+				ScaleSubsurfaceRadiusNode->AddFloatAttribute(UInterchangeShaderPortsAPI::MakeInputValueKey(Multiply::Inputs::B.ToString()), 600);
+				UInterchangeShaderPortsAPI::ConnectOuputToInputByName(ScaleSubsurfaceRadiusNode, Multiply::Inputs::A.ToString(), SubsurfaceRadiusNodeUid, SubsurfaceRadiusNodeOutputName);
+
+				// 4. Multiply the scaled SubsurfaceRadius by the TransmissionColor if any
+				FString TransmissionColorNodeUid, TransmissionColorNodeOutputName;
+				if (UInterchangeShaderPortsAPI::DisconnectInputFromOutputNode(OpenPBRSurfaceShaderNode, OpenPBRSurface::Parameters::TransmissionColor.ToString(), TransmissionColorNodeUid, TransmissionColorNodeOutputName))
+				{
+					UInterchangeShaderNode* SubsurfaceTransmissionColorNode = CreateShaderNodeInternal.operator()<UInterchangeShaderNode>(OpenPBRSurfaceNode, Multiply::Name.ToString(), TEXT("SubsurfaceRadiusTransmissionColorMult"));
+					UInterchangeShaderPortsAPI::ConnectOuputToInputByName(SubsurfaceTransmissionColorNode, Multiply::Inputs::A.ToString(), TransmissionColorNodeUid, TransmissionColorNodeOutputName);
+					UInterchangeShaderPortsAPI::ConnectDefaultOuputToInput(SubsurfaceTransmissionColorNode, Multiply::Inputs::B.ToString(), ScaleSubsurfaceRadiusNode->GetUniqueID());
+					UInterchangeShaderPortsAPI::ConnectDefaultOuputToInput(OpenPBRSurfaceShaderNode, OpenPBRSurface::Parameters::TransmissionColor.ToString(), SubsurfaceTransmissionColorNode->GetUniqueID());
+				}
+				else
+				{
+					UInterchangeShaderPortsAPI::ConnectDefaultOuputToInput(OpenPBRSurfaceShaderNode, OpenPBRSurface::Parameters::TransmissionColor.ToString(), ScaleSubsurfaceRadiusNode->GetUniqueID());
+				}
+			}
+
+			UInterchangeShaderPortsAPI::ConnectOuputToInputByName(ShaderGraphNode, SubstrateMaterial::Parameters::FrontMaterial.ToString(), OpenPBRSurfaceShaderNode->GetUniqueID(), OpenPBRSurface::SubstrateMaterial::Outputs::FrontMaterial.ToString());
+		}
+		else
+		{
+			UInterchangeShaderPortsAPI::ConnectOuputToInputByName(ShaderGraphNode, SubstrateMaterial::Parameters::FrontMaterial.ToString(), OpenPBRSurfaceShaderNode->GetUniqueID(), OpenPBRSurface::SubstrateMaterial::Outputs::FrontMaterial.ToString());
+
+			if (UInterchangeShaderPortsAPI::HasInput(OpenPBRSurfaceShaderNode, OpenPBRSurface::Parameters::GeometryOpacity))
+			{
+				UInterchangeShaderPortsAPI::ConnectOuputToInputByName(ShaderGraphNode, SubstrateMaterial::Parameters::OpacityMask.ToString(), OpenPBRSurfaceShaderNode->GetUniqueID(), OpenPBRSurface::SubstrateMaterial::Outputs::OpacityMask.ToString());
+				ShaderGraphNode->SetCustomBlendMode(EBlendMode::BLEND_Masked);
+			}
 		}
 	}
 	else
diff --git a/Engine/Plugins/Interchange/Runtime/Source/Import/Private/MaterialX/MaterialXUtils/MaterialXStandardSurfaceShader.cpp b/Engine/Plugins/Interchange/Runtime/Source/Import/Private/MaterialX/MaterialXUtils/MaterialXStandardSurfaceShader.cpp
--- a/Engine/Plugins/Interchange/Runtime/Source/Import/Private/MaterialX/MaterialXUtils/MaterialXStandardSurfaceShader.cpp
+++ b/Engine/Plugins/Interchange/Runtime/Source/Import/Private/MaterialX/MaterialXUtils/MaterialXStandardSurfaceShader.cpp
@@ -24,6 +24,8 @@ UInterchangeBaseNode* FMaterialXStandardSurfaceShader::Translate(mx::NodePtr Sta
 	this->SurfaceShaderNode = StandardSurfaceNode;
 	
 	using namespace UE::Interchange::Materials;
+	using namespace UE::Interchange::MaterialX;
+	using namespace UE::Interchange::Materials::Standard::Nodes;
 
 	UInterchangeShaderNode* StandardSurfaceShaderNode = FMaterialXSurfaceShaderAbstract::Translate(EInterchangeMaterialXShaders::StandardSurface);
 
@@ -38,22 +40,125 @@ UInterchangeBaseNode* FMaterialXStandardSurfaceShader::Translate(mx::NodePtr Sta
 
 	if(UInterchangeShaderPortsAPI::HasInput(StandardSurfaceShaderNode, StandardSurface::Parameters::Transmission))
 	{
-		StandardSurfaceShaderNode->AddInt32Attribute(UE::Interchange::MaterialX::Attributes::EnumType, UE::Interchange::MaterialX::IndexSurfaceShaders);
-		StandardSurfaceShaderNode->AddInt32Attribute(UE::Interchange::MaterialX::Attributes::EnumValue, int32(EInterchangeMaterialXShaders::StandardSurfaceTransmission));
+		StandardSurfaceShaderNode->AddInt32Attribute(Attributes::EnumType, IndexSurfaceShaders);
+		StandardSurfaceShaderNode->AddInt32Attribute(Attributes::EnumValue, int32(EInterchangeMaterialXShaders::StandardSurfaceTransmission));
 	}
 
 	// Outputs
 	if(bIsSubstrateEnabled)
 	{
-		if(UInterchangeShaderPortsAPI::HasInput(StandardSurfaceShaderNode, StandardSurface::Parameters::Transmission))
+		auto CreateShaderNodeInternal = [this]<class T>(mx::NodePtr StandardSurfaceNode, const FString& PathOrShaderType, const FString& DisplayLabel)
+		{
+			T* Node;
+
+			const FString NodeUID = UInterchangeShaderNode::MakeNodeUid(GetUniqueName(StandardSurfaceNode) + DisplayLabel, FStringView{});
+
+			if (Node = const_cast<T*>(Cast<T>(NodeContainer.GetNode(NodeUID))); !Node)
+			{
+				Node = NewObject<T>(&NodeContainer);
+				NodeContainer.SetupNode(Node, NodeUID, DisplayLabel, EInterchangeNodeContainerType::TranslatedAsset);
+				if constexpr (std::is_same_v<T, UInterchangeFunctionCallShaderNode>)
+				{
+					Node->SetCustomMaterialFunction(PathOrShaderType);
+				}
+				else
+				{
+					Node->SetCustomShaderType(PathOrShaderType);
+				}
+			}
+
+			return Node;
+		};
+
+		if (UInterchangeShaderPortsAPI::HasInput(StandardSurfaceShaderNode, StandardSurface::Parameters::Transmission) &&
+			UInterchangeShaderPortsAPI::HasInput(StandardSurfaceShaderNode, StandardSurface::Parameters::Subsurface) &&
+			UInterchangeShaderPortsAPI::HasInput(StandardSurfaceShaderNode, StandardSurface::Parameters::BaseColor))
+		{
+			// 1. Set the Material as Opaque with Mask Mode
+			StandardSurfaceShaderNode->AddInt32Attribute(Attributes::EnumType, IndexSurfaceShaders);
+			StandardSurfaceShaderNode->AddInt32Attribute(Attributes::EnumValue, int32(EInterchangeMaterialXShaders::StandardSurface));
+			ShaderGraphNode->SetCustomBlendMode(EBlendMode::BLEND_Masked);
+
+			//2. Disconnect the transmission inputs
+			FString TransmissionNodeUid, TransmissionOutputName;
+			UInterchangeShaderPortsAPI::DisconnectInputFromOutputNode(StandardSurfaceShaderNode, StandardSurface::Parameters::Transmission.ToString(), TransmissionNodeUid, TransmissionOutputName);
+			UInterchangeShaderPortsAPI::DisconnectInput(StandardSurfaceShaderNode, StandardSurface::Parameters::TransmissionColor.ToString());
+			UInterchangeShaderPortsAPI::DisconnectInput(StandardSurfaceShaderNode, StandardSurface::Parameters::TransmissionDepth.ToString());
+			UInterchangeShaderPortsAPI::DisconnectInput(StandardSurfaceShaderNode, StandardSurface::Parameters::TransmissionDispersion.ToString());
+			UInterchangeShaderPortsAPI::DisconnectInput(StandardSurfaceShaderNode, StandardSurface::Parameters::TransmissionExtraRoughness.ToString());
+			UInterchangeShaderPortsAPI::DisconnectInput(StandardSurfaceShaderNode, StandardSurface::Parameters::TransmissionScatter.ToString());
+			UInterchangeShaderPortsAPI::DisconnectInput(StandardSurfaceShaderNode, StandardSurface::Parameters::TransmissionScatterAnisotropy.ToString());
+
+			//3. Create DitherMask node
+			UInterchangeShaderNode* DitherMaskNode = CreateShaderNodeInternal.operator()<UInterchangeFunctionCallShaderNode>(StandardSurfaceNode, UE::MaterialFunctions::Path::DitherMask, TEXT("DitherMaskNode"));
+
+			//4. Connect the surface shader and the transmission to the MF_DitherMask
+			if (const UInterchangeShaderNode* TransmissionNode = Cast<UInterchangeShaderNode>(NodeContainer.GetNode(TransmissionNodeUid)))
+			{
+				UInterchangeShaderPortsAPI::ConnectOuputToInputByName(DitherMaskNode, TEXT("BSDF"), StandardSurfaceShaderNode->GetUniqueID(), StandardSurface::SubstrateMaterial::Outputs::Opaque.ToString());
+				UInterchangeShaderPortsAPI::ConnectOuputToInputByName(DitherMaskNode, TEXT("transmission"), TransmissionNode->GetUniqueID(), TransmissionOutputName);
+			}
+
+			UInterchangeShaderPortsAPI::ConnectOuputToInputByName(ShaderGraphNode, SubstrateMaterial::Parameters::FrontMaterial.ToString(), DitherMaskNode->GetUniqueID(), TEXT("Dither"));
+			UInterchangeShaderPortsAPI::ConnectOuputToInputByName(ShaderGraphNode, SubstrateMaterial::Parameters::OpacityMask.ToString(), DitherMaskNode->GetUniqueID(), SubstrateMaterial::Parameters::OpacityMask.ToString());
+		}
+		else if (UInterchangeShaderPortsAPI::HasInput(StandardSurfaceShaderNode, StandardSurface::Parameters::Transmission) &&
+				 UInterchangeShaderPortsAPI::HasInput(StandardSurfaceShaderNode, StandardSurface::Parameters::Subsurface))
+		{
+			ShaderGraphNode->SetCustomBlendMode(EBlendMode::BLEND_TranslucentColoredTransmittance);
+			auto ReplaceInputConnection = [this, StandardSurfaceShaderNode](const FString& InputToRemove, const FString& InputName)
+			{
+				UInterchangeShaderPortsAPI::DisconnectInput(StandardSurfaceShaderNode, InputToRemove);
+
+				FString NodeUid, NodeOutputName;
+				UInterchangeShaderPortsAPI::GetInputConnection(StandardSurfaceShaderNode, InputName, NodeUid, NodeOutputName);
+				if (const UInterchangeShaderNode* ShaderNode = Cast<UInterchangeShaderNode>(NodeContainer.GetNode(NodeUid)))
+				{
+					UInterchangeShaderPortsAPI::ConnectOuputToInputByName(StandardSurfaceShaderNode, InputToRemove, ShaderNode->GetUniqueID(), NodeOutputName);
+				}
+			};
+
+			//1. subsurface_color input now replaces the base_color
+			if(UInterchangeShaderPortsAPI::HasInput(StandardSurfaceShaderNode, StandardSurface::Parameters::SubsurfaceColor))
+			{
+				ReplaceInputConnection(StandardSurface::Parameters::BaseColor.ToString(), StandardSurface::Parameters::SubsurfaceColor.ToString());
+			}
+
+			//2. replace base by subsurface
+			ReplaceInputConnection(StandardSurface::Parameters::Base.ToString(), StandardSurface::Parameters::Subsurface.ToString());
+
+			//3. scale the subsurface_radius by 600 (empirical value that gives better results, need further testing though)
+			FString SubsurfaceRadiusNodeUid, SubsurfaceRadiusNodeOutputName;
+			if (UInterchangeShaderPortsAPI::GetInputConnection(StandardSurfaceShaderNode, StandardSurface::Parameters::SubsurfaceRadius.ToString(), SubsurfaceRadiusNodeUid, SubsurfaceRadiusNodeOutputName))
+			{
+				UInterchangeShaderNode* ScaleSubsurfaceRadiusNode = CreateShaderNodeInternal.operator()<UInterchangeShaderNode>(StandardSurfaceNode, Multiply::Name.ToString(), TEXT("ScaleSubsurfaceRadius"));
+				ScaleSubsurfaceRadiusNode->AddFloatAttribute(UInterchangeShaderPortsAPI::MakeInputValueKey(Multiply::Inputs::B.ToString()), 600);
+				UInterchangeShaderPortsAPI::ConnectOuputToInputByName(ScaleSubsurfaceRadiusNode, Multiply::Inputs::A.ToString(), SubsurfaceRadiusNodeUid, SubsurfaceRadiusNodeOutputName);
+
+				// 4. Multiply the scaled SubsurfaceRadius by the TransmissionColor if any
+				FString TransmissionColorNodeUid, TransmissionColorNodeOutputName;
+				if(UInterchangeShaderPortsAPI::DisconnectInputFromOutputNode(StandardSurfaceShaderNode, StandardSurface::Parameters::TransmissionColor.ToString(), TransmissionColorNodeUid, TransmissionColorNodeOutputName))
+				{
+					UInterchangeShaderNode* SubsurfaceTransmissionColorNode = CreateShaderNodeInternal.operator()<UInterchangeShaderNode>(StandardSurfaceNode, Multiply::Name.ToString(), TEXT("SubsurfaceRadiusTransmissionColorMult"));
+					UInterchangeShaderPortsAPI::ConnectOuputToInputByName(SubsurfaceTransmissionColorNode, Multiply::Inputs::A.ToString(), TransmissionColorNodeUid, TransmissionColorNodeOutputName);
+					UInterchangeShaderPortsAPI::ConnectDefaultOuputToInput(SubsurfaceTransmissionColorNode, Multiply::Inputs::B.ToString(), ScaleSubsurfaceRadiusNode->GetUniqueID());
+					UInterchangeShaderPortsAPI::ConnectDefaultOuputToInput(StandardSurfaceShaderNode, StandardSurface::Parameters::TransmissionColor.ToString(), SubsurfaceTransmissionColorNode->GetUniqueID());
+				}
+				else
+				{
+					UInterchangeShaderPortsAPI::ConnectDefaultOuputToInput(StandardSurfaceShaderNode, StandardSurface::Parameters::TransmissionColor.ToString(), ScaleSubsurfaceRadiusNode->GetUniqueID());
+				}
+			}
+
+			UInterchangeShaderPortsAPI::ConnectOuputToInputByName(ShaderGraphNode, SubstrateMaterial::Parameters::FrontMaterial.ToString(), StandardSurfaceShaderNode->GetUniqueID(), StandardSurface::SubstrateMaterial::Outputs::Translucent.ToString());
+		}
+		else if(UInterchangeShaderPortsAPI::HasInput(StandardSurfaceShaderNode, StandardSurface::Parameters::Transmission))
 		{
 			ShaderGraphNode->SetCustomBlendMode(EBlendMode::BLEND_TranslucentColoredTransmittance);
-			UInterchangeShaderPortsAPI::ConnectOuputToInputByName(ShaderGraphNode, PBRMR::Parameters::Opacity.ToString(), StandardSurfaceShaderNode->GetUniqueID(), PBRMR::Parameters::Opacity.ToString());
 			UInterchangeShaderPortsAPI::ConnectOuputToInputByName(ShaderGraphNode, SubstrateMaterial::Parameters::FrontMaterial.ToString(), StandardSurfaceShaderNode->GetUniqueID(), StandardSurface::SubstrateMaterial::Outputs::Translucent.ToString());
 		}
 		else
 		{
-
 			UInterchangeShaderPortsAPI::ConnectOuputToInputByName(ShaderGraphNode, SubstrateMaterial::Parameters::FrontMaterial.ToString(), StandardSurfaceShaderNode->GetUniqueID(), StandardSurface::SubstrateMaterial::Outputs::Opaque.ToString());
 			if(UInterchangeShaderPortsAPI::HasInput(StandardSurfaceShaderNode, StandardSurface::Parameters::Opacity))
 			{
diff --git a/Engine/Plugins/Interchange/Runtime/Source/Import/Private/MaterialX/MaterialXUtils/MaterialXSurfaceShaderAbstract.cpp b/Engine/Plugins/Interchange/Runtime/Source/Import/Private/MaterialX/MaterialXUtils/MaterialXSurfaceShaderAbstract.cpp
--- a/Engine/Plugins/Interchange/Runtime/Source/Import/Private/MaterialX/MaterialXUtils/MaterialXSurfaceShaderAbstract.cpp
+++ b/Engine/Plugins/Interchange/Runtime/Source/Import/Private/MaterialX/MaterialXUtils/MaterialXSurfaceShaderAbstract.cpp
@@ -3,13 +3,13 @@
 #if WITH_EDITOR
 #include "MaterialX/MaterialXUtils/MaterialXSurfaceShaderAbstract.h"
 
-#include "InterchangeImportLog.h"
 #include "InterchangeTexture2DNode.h"
 #include "InterchangeTextureBlurNode.h"
 #include "Materials/MaterialExpressionTransform.h"
 #include "Materials/MaterialExpressionVectorNoise.h"
 #include "MaterialX/MaterialXUtils/MaterialXManager.h"
 #include "MaterialX/MaterialExpressions/MaterialExpressionTextureSampleParameterBlur.h"
+#include "Usd/InterchangeUsdDefinitions.h"
 
 #define LOCTEXT_NAMESPACE "MaterialXSurfaceShaderAbstract"
 
@@ -1376,6 +1376,38 @@ void FMaterialXSurfaceShaderAbstract::ConnectNotInputToOutput(const FConnectNode
 	ConnectLogicalInputToOutput(Connect, UE::MaterialFunctions::Path::MxNot);
 }
 
+void FMaterialXSurfaceShaderAbstract::ConnectSparseVolumeInputToOutput(const FConnectNode& Connect)
+{
+	using namespace UE::Interchange::Materials::Standard::Nodes;
+	// the sparse_volume node is not standard, it's only here for USD, just like <geompropvalue> baked from primvars coming from a surface material
+	// this one should only be created from a volume material in USD
+	// we'll set the SparseVolumeTexture in the PostProcessFactory
+	int32 OutputIndex = Connect.OutputName == "outA" ? 0 : 1;
+	UInterchangeShaderNode* TextureShaderNode = CreateShaderNode(Connect.UpstreamNode, Connect.UpstreamNode->getName().c_str(), TEXT("SparseVolumeTextureSampleParameter"));
+	if (Connect.UpstreamNode->hasAttribute(mx::Attributes::GeomPropSparseVolume))
+	{
+		TextureShaderNode->AddStringAttribute(UE::Interchange::USD::Primvar::ShaderNodeSparseVolumeTextureSample, Connect.UpstreamNode->getAttribute(mx::Attributes::GeomPropSparseVolume).c_str());
+	}
+	const std::string& OutputTypeA = Connect.UpstreamNode->getNodeDef()->getOutput("outA")->getType();
+	
+	std::unordered_map<std::string, uint8> Masks{ 
+		{mx::Type::Vector2, 0b1100},
+		{mx::Type::Vector3, 0b1110},
+		{mx::Type::Color3, 0b1100},
+		{mx::Type::Float, 0b1000},
+	};
+
+	UInterchangeShaderNode* MaskNode = TextureShaderNode;
+	if (OutputTypeA != mx::Type::Color4 && OutputTypeA != mx::Type::Vector4)
+	{
+		MaskNode = CreateMaskShaderNode(Masks[OutputTypeA], Connect.UpstreamNode, (Connect.UpstreamNode->getName() + "_Mask").c_str(), Connect.OutputName);
+		UInterchangeShaderPortsAPI::ConnectOuputToInputByIndex(MaskNode, Mask::Inputs::Input.ToString(), TextureShaderNode->GetUniqueID(), OutputIndex);
+		OutputIndex = 0;
+	}
+
+	UInterchangeShaderPortsAPI::ConnectOuputToInputByIndex(Connect.ParentShaderNode, Connect.InputChannelName, MaskNode->GetUniqueID(), OutputIndex);
+}
+
 UInterchangeShaderNode* FMaterialXSurfaceShaderAbstract::CreateMaskShaderNode(uint8 RGBA, mx::ElementPtr Element, const FString& NodeName, const FString& OutputName)
 {
 	bool bR = (0b1000 & RGBA) >> 3;
@@ -1616,6 +1648,7 @@ void FMaterialXSurfaceShaderAbstract::RegisterConnectNodeOutputToInputDelegates(
 	MatchingConnectNodeDelegates.Add(mx::Category::Or,				FOnConnectNodeOutputToInput::CreateSP(this, &FMaterialXSurfaceShaderAbstract::ConnectOrInputToOutput));
 	MatchingConnectNodeDelegates.Add(mx::Category::Xor,				FOnConnectNodeOutputToInput::CreateSP(this, &FMaterialXSurfaceShaderAbstract::ConnectXorInputToOutput));
 	MatchingConnectNodeDelegates.Add(mx::Category::Not,				FOnConnectNodeOutputToInput::CreateSP(this, &FMaterialXSurfaceShaderAbstract::ConnectNotInputToOutput));
+	MatchingConnectNodeDelegates.Add(mx::Category::SparseVolume,	FOnConnectNodeOutputToInput::CreateSP(this, &FMaterialXSurfaceShaderAbstract::ConnectSparseVolumeInputToOutput));
 }
 
 void FMaterialXSurfaceShaderAbstract::SetMatchingInputsNames(MaterialX::NodePtr Node) const
diff --git a/Engine/Plugins/Interchange/Runtime/Source/Import/Public/MaterialX/MaterialXUtils/MaterialXSurfaceShaderAbstract.h b/Engine/Plugins/Interchange/Runtime/Source/Import/Public/MaterialX/MaterialXUtils/MaterialXSurfaceShaderAbstract.h
--- a/Engine/Plugins/Interchange/Runtime/Source/Import/Public/MaterialX/MaterialXUtils/MaterialXSurfaceShaderAbstract.h
+++ b/Engine/Plugins/Interchange/Runtime/Source/Import/Public/MaterialX/MaterialXUtils/MaterialXSurfaceShaderAbstract.h
@@ -336,6 +336,9 @@ protected:
 	/** <not> */
 	void ConnectNotInputToOutput(const FConnectNode& Connect);
 
+	/** <sparse_volume> is define in Engine/Binaries/ThirdParty/MaterialX/libraries/Interchange */
+	void ConnectSparseVolumeInputToOutput(const FConnectNode& Connect);
+
 	/** End Connect MaterialX nodes*/
 
 	/**
diff --git a/Engine/Plugins/Interchange/Runtime/Source/Nodes/Private/InterchangeShaderGraphNode.cpp b/Engine/Plugins/Interchange/Runtime/Source/Nodes/Private/InterchangeShaderGraphNode.cpp
--- a/Engine/Plugins/Interchange/Runtime/Source/Nodes/Private/InterchangeShaderGraphNode.cpp
+++ b/Engine/Plugins/Interchange/Runtime/Source/Nodes/Private/InterchangeShaderGraphNode.cpp
@@ -116,6 +116,20 @@ bool UInterchangeShaderPortsAPI::ConnectDefaultOuputToInput(UInterchangeBaseNode
 	return InterchangeNode->AddStringAttribute(MakeInputConnectionKey(InputName), ExpressionUid);
 }
 
+bool UInterchangeShaderPortsAPI::DisconnectInput(UInterchangeBaseNode* InterchangeNode, const FString& InputName)
+{
+	return InterchangeNode->RemoveAttribute(MakeInputConnectionKey(InputName));
+}
+
+bool UInterchangeShaderPortsAPI::DisconnectInputFromOutputNode(UInterchangeBaseNode* InterchangeNode, const FString& InputName, FString& OutExpressionUid, FString& OutputName)
+{
+	if (GetInputConnection(InterchangeNode, InputName, OutExpressionUid, OutputName))
+	{
+		return DisconnectInput(InterchangeNode, InputName);
+	}
+	return false;
+}
+
 bool UInterchangeShaderPortsAPI::ConnectOuputToInputByName(UInterchangeBaseNode* InterchangeNode, const FString& InputName, const FString& ExpressionUid, const FString& OutputName)
 {
 	if (OutputName.IsEmpty())
diff --git a/Engine/Plugins/Interchange/Runtime/Source/Nodes/Public/InterchangeShaderGraphNode.h b/Engine/Plugins/Interchange/Runtime/Source/Nodes/Public/InterchangeShaderGraphNode.h
--- a/Engine/Plugins/Interchange/Runtime/Source/Nodes/Public/InterchangeShaderGraphNode.h
+++ b/Engine/Plugins/Interchange/Runtime/Source/Nodes/Public/InterchangeShaderGraphNode.h
@@ -86,6 +86,26 @@ public:
 	UFUNCTION(BlueprintCallable, Category = "Interchange | Node | Material")
 	static UE_API bool ConnectDefaultOuputToInput(UInterchangeBaseNode* InterchangeNode, const FString& InputName, const FString& ExpressionUid);
  
+	/**
+	 * Remove an input connection attribute.
+	 * @param InterchangeNode	The node to disconnect the input from.
+	 * @param InputName			The name of the input to disconnect.
+	 * @return					True if the input connection was successfully removed from the node.
+	 */
+	UFUNCTION(BlueprintCallable, Category = "Interchange | Node | Material")
+	static UE_API bool DisconnectInput(UInterchangeBaseNode* InterchangeNode, const FString& InputName);
+ 
+	/**
+	 * Remove an input connection attribute by also retrieving the node unique id and the output name connected to the given input, if any.
+	 * @param InterchangeNode	The node to disconnect the input from.
+	 * @param InputName			The name of the input to disconnect.
+	 * @param OutExpressionUid  The Uid of the node that was disconnected from the input.
+	 * @param OutputName        The name of node output.
+	 * @return					True if the input connection was successfully removed from the node.
+	 */
+	UFUNCTION(BlueprintCallable, Category = "Interchange | Node | Material")
+	static UE_API bool DisconnectInputFromOutputNode(UInterchangeBaseNode* InterchangeNode, const FString& InputName, FString& OutExpressionUid, FString& OutputName);
+
 	/**
 	 * Adds an input connection attribute.
 	 * @param InterchangeNode	The node to create the input on.
diff --git a/Engine/Plugins/Interchange/Runtime/Source/Pipelines/Private/InterchangeGenericMaterialPipeline.cpp b/Engine/Plugins/Interchange/Runtime/Source/Pipelines/Private/InterchangeGenericMaterialPipeline.cpp
--- a/Engine/Plugins/Interchange/Runtime/Source/Pipelines/Private/InterchangeGenericMaterialPipeline.cpp
+++ b/Engine/Plugins/Interchange/Runtime/Source/Pipelines/Private/InterchangeGenericMaterialPipeline.cpp
@@ -22,6 +22,7 @@
 #include "InterchangeTextureCubeNode.h"
 #include "InterchangeTextureFactoryNode.h"
 #include "InterchangeTextureNode.h"
+#include "InterchangeVolumeTextureNode.h"
 #include "Materials/Material.h"
 #include "Materials/MaterialExpressionAdd.h"
 #include "Materials/MaterialExpressionComponentMask.h"
@@ -38,6 +39,7 @@
 #include "Materials/MaterialExpressionRotator.h"
 #include "Materials/MaterialExpressionScalarParameter.h"
 #include "Materials/MaterialExpressionSine.h"
+#include "Materials/MaterialExpressionSparseVolumeTextureSample.h"
 #include "Materials/MaterialExpressionStaticBoolParameter.h"
 #include "Materials/MaterialExpressionSubstrate.h"
 #include "Materials/MaterialExpressionSwitch.h"
@@ -2228,25 +2230,20 @@ void UInterchangeGenericMaterialPipeline::HandleTextureSampleNode(const UInterch
 
 	if (const UInterchangeTextureNode* TextureNode = Cast<const UInterchangeTextureNode>(BaseNodeContainer->GetNode(TextureUid)))
 	{
-		if (TextureNode->IsA<UInterchangeTextureCubeNode>())
+		TMap<UClass*, FString> TextureNodeSampleExpression
 		{
-			ExpressionClassName = UMaterialExpressionTextureSampleParameterCube::StaticClass()->GetName();
-		}
-		else if (TextureNode->IsA<UInterchangeTexture2DArrayNode>())
-		{
-			ExpressionClassName = UMaterialExpressionTextureSampleParameter2DArray::StaticClass()->GetName();
-		}
-		else if(TextureNode->IsA<UInterchangeTextureBlurNode>())
-		{
-			ExpressionClassName = UMaterialExpressionMaterialXTextureSampleParameterBlur::StaticClass()->GetName();
-		}
-		else if (TextureNode->IsA<UInterchangeTexture2DNode>())
-		{
-			ExpressionClassName = UMaterialExpressionTextureSampleParameter2D::StaticClass()->GetName();
-		}
-		else
+			{UInterchangeTextureCubeNode::StaticClass(), UMaterialExpressionTextureSampleParameterCube::StaticClass()->GetName()},
+			{UInterchangeTexture2DArrayNode::StaticClass(), UMaterialExpressionTextureSampleParameter2DArray::StaticClass()->GetName()},
+			{UInterchangeTextureBlurNode::StaticClass(), UMaterialExpressionMaterialXTextureSampleParameterBlur::StaticClass()->GetName()},
+			{UInterchangeTexture2DNode::StaticClass(), UMaterialExpressionTextureSampleParameter2D::StaticClass()->GetName()},
+			{UInterchangeTextureNode::StaticClass(), UMaterialExpressionTextureSampleParameter2D::StaticClass()->GetName()},
+			{UInterchangeVolumeTextureNode::StaticClass(), UMaterialExpressionSparseVolumeTextureSampleParameter::StaticClass()->GetName()},
+		};
+
+		ExpressionClassName = UMaterialExpressionTextureSampleParameter2D::StaticClass()->GetName();
+		if(FString* ExpressionFound = TextureNodeSampleExpression.Find(TextureNode->GetClass()))
 		{
-			ExpressionClassName = UMaterialExpressionTextureSampleParameter2D::StaticClass()->GetName();
+			ExpressionClassName = *ExpressionFound;
 		}
 
 		HandleTextureNode(TextureNode, MaterialFactoryNode, TextureSampleFactoryNode, ExpressionClassName, bIsAParameter);
diff --git a/Engine/Plugins/Runtime/USDCore/Source/USDUtilities/Private/USDMaterialXShaderGraph.cpp b/Engine/Plugins/Runtime/USDCore/Source/USDUtilities/Private/USDMaterialXShaderGraph.cpp
--- a/Engine/Plugins/Runtime/USDCore/Source/USDUtilities/Private/USDMaterialXShaderGraph.cpp
+++ b/Engine/Plugins/Runtime/USDCore/Source/USDUtilities/Private/USDMaterialXShaderGraph.cpp
@@ -202,6 +202,7 @@ void FUsdMaterialXShaderGraph::CreateMaterial(const pxr::UsdShadeMaterial& UsdSh
 	// if that's the case let's just suffix the name of the volume material to avoid collision during the creation of the node
 	if (VolumeSource)
 	{
+		TGuardValue<bool> GuardSparseVolume{ bConvertGeomPropValueToSparseVolume, true };
 		if (mx::NodePtr VolumeShader = CreateShaderGraph(VolumeSource, Document))
 		{
 			mx::NodeDefPtr VolumeMaterialNodeDef = Document->getNodeDef("ND_volumematerial");
@@ -277,9 +278,16 @@ bool FUsdMaterialXShaderGraph::ComputeShaderGraph(mx::NodePtr Node, const pxr::U
 			}
 
 			Input->setConnectedNode(ConnectedNode);
+			std::string OutputName = Source.sourceName.GetString();
+
+			// we can only assume at this point that the output for a sparse volume is the 'Attributes A'
+			if (ConnectedNode->getCategory() == "sparse_volume")
+			{
+				OutputName = "outA";
+			}
 
 			// no need to add the default output
-			if (const std::string& OutputName = Source.sourceName.GetString(); OutputName != "out")
+			if (OutputName != "out")
 			{
 				Input->setOutputString(OutputName);
 			}
@@ -465,6 +473,14 @@ mx::NodePtr FUsdMaterialXShaderGraph::GetNode(const pxr::UsdShadeConnectionSourc
 		{
 			ConnectedNode->setTypedAttribute("UE:GeomPropImage", bIsGeomPropValue);
 		}
+		else if (ConnectedNode->getCategory() == "sparse_volume")
+		{
+			// the name of the geompropvalue (which is the name of a volume primvar), in that case the name that the SparseVolumeAsset will be named at
+			// we'll use that name later in the ExecutePostFactoryPipeline to assign the asset to the SparseVolumeTextureSample
+			TUsdStore<std::string> GeomPropName;
+			pxr::UsdShadeShader{ Source.source.GetPrim() }.GetInput(pxr::TfToken{ "geomprop" }).Get<std::string>(&(*GeomPropName));
+			ConnectedNode->setAttribute("UE:GeomPropSparseVolume", *GeomPropName);
+		}
 		bIsNodeAlreadyExist = false;
 	}
 
@@ -537,13 +553,22 @@ mx::NodeDefPtr FUsdMaterialXShaderGraph::GetNodeDef(const pxr::UsdShadeConnectio
 		}
 		else
 		{
-			bIsGeomPropValue = true;
-			// Here we just replace the replace nodedef of a <geompropvalue> by the <image> one
+			if (bConvertGeomPropValueToSparseVolume)
+			{
+				// we convert the geompropvalue to a sparse_volume node
+				// at this point we can only assume that the output is attached to the 'Attributes A' of the SparseVolumeTexture
+				NodeName = "sparse_volume_A";
+			}
+			else
+			{
+				bIsGeomPropValue = true;
+			}
+			// we replace the nodedef of a <geompropvalue> by either an <image> or a <sparse_volume> 
 			// e.g: ND_geompropvalue_vector3
 			//		ND_image_vector3
-			// for integers/bool, <image> has no such node definitions, we just consider it as floats
-			// e.g: ND_geompropvalue_integer
-			//		ND_image_float
+			// or   ND_sparse_volume_A_vector3
+			// for integers/bool let's just convert them to float
+
 			if (std::size_t pos = NodeDefString.rfind("integer"); pos != std::string::npos)
 			{
 				NodeDefString.replace(pos, sizeof("integer") - 1, "float");
diff --git a/Engine/Plugins/Runtime/USDCore/Source/USDUtilities/Public/USDMaterialXShaderGraph.h b/Engine/Plugins/Runtime/USDCore/Source/USDUtilities/Public/USDMaterialXShaderGraph.h
--- a/Engine/Plugins/Runtime/USDCore/Source/USDUtilities/Public/USDMaterialXShaderGraph.h
+++ b/Engine/Plugins/Runtime/USDCore/Source/USDUtilities/Public/USDMaterialXShaderGraph.h
@@ -99,6 +99,8 @@ private:
 
 	bool bTangentSpaceInput = false;
 
+	bool bConvertGeomPropValueToSparseVolume = false;
+
 private:
 
 	// complex template to just retrieve the right type, basically a mapping between Usd and Mtlx for Vector types, otherwise we just
