Interchange:
- remove  use of a TMap instead of IsA function

diff --git a/Engine/Plugins/Interchange/Runtime/Source/Pipelines/Private/InterchangeGenericMaterialPipeline.cpp b/Engine/Plugins/Interchange/Runtime/Source/Pipelines/Private/InterchangeGenericMaterialPipeline.cpp
--- a/Engine/Plugins/Interchange/Runtime/Source/Pipelines/Private/InterchangeGenericMaterialPipeline.cpp
+++ b/Engine/Plugins/Interchange/Runtime/Source/Pipelines/Private/InterchangeGenericMaterialPipeline.cpp
@@ -2230,20 +2230,29 @@ void UInterchangeGenericMaterialPipeline::HandleTextureSampleNode(const UInterch
 
 	if (const UInterchangeTextureNode* TextureNode = Cast<const UInterchangeTextureNode>(BaseNodeContainer->GetNode(TextureUid)))
 	{
-		TMap<UClass*, FString> TextureNodeSampleExpression
-		{
-			{UInterchangeTextureCubeNode::StaticClass(), UMaterialExpressionTextureSampleParameterCube::StaticClass()->GetName()},
-			{UInterchangeTexture2DArrayNode::StaticClass(), UMaterialExpressionTextureSampleParameter2DArray::StaticClass()->GetName()},
-			{UInterchangeTextureBlurNode::StaticClass(), UMaterialExpressionMaterialXTextureSampleParameterBlur::StaticClass()->GetName()},
-			{UInterchangeTexture2DNode::StaticClass(), UMaterialExpressionTextureSampleParameter2D::StaticClass()->GetName()},
-			{UInterchangeTextureNode::StaticClass(), UMaterialExpressionTextureSampleParameter2D::StaticClass()->GetName()},
-			{UInterchangeVolumeTextureNode::StaticClass(), UMaterialExpressionSparseVolumeTextureSampleParameter::StaticClass()->GetName()},
-		};
-
-		ExpressionClassName = UMaterialExpressionTextureSampleParameter2D::StaticClass()->GetName();
-		if(FString* ExpressionFound = TextureNodeSampleExpression.Find(TextureNode->GetClass()))
+		if (TextureNode->IsA<UInterchangeTextureCubeNode>())
+		{
+			ExpressionClassName = UMaterialExpressionTextureSampleParameterCube::StaticClass()->GetName();
+		}
+		else if (TextureNode->IsA<UInterchangeTexture2DArrayNode>())
+		{
+			ExpressionClassName = UMaterialExpressionTextureSampleParameter2DArray::StaticClass()->GetName();
+		}
+		else if (TextureNode->IsA<UInterchangeTextureBlurNode>())
+		{
+			ExpressionClassName = UMaterialExpressionMaterialXTextureSampleParameterBlur::StaticClass()->GetName();
+		}
+		else if (TextureNode->IsA<UInterchangeTexture2DNode>())
+		{
+			ExpressionClassName = UMaterialExpressionTextureSampleParameter2D::StaticClass()->GetName();
+		}
+		else if (TextureNode->IsA<UInterchangeVolumeTextureNode>())
+		{
+			ExpressionClassName = UMaterialExpressionSparseVolumeTextureSampleParameter::StaticClass()->GetName();
+		}
+		else
 		{
-			ExpressionClassName = *ExpressionFound;
+			ExpressionClassName = UMaterialExpressionTextureSampleParameter2D::StaticClass()->GetName();
 		}
 
 		HandleTextureNode(TextureNode, MaterialFactoryNode, TextureSampleFactoryNode, ExpressionClassName, bIsAParameter);
