[Interchange] Audio Improvements

- Added a new category for audio in both Translator as well as Factory.
- Fixed Reaper Translator file path issues
- Added ability to reimport sounds
- Made the Audio Import, Reimport use the Interchange Payload system
- Moved InterchangeAudio/Core to Interchange plugin in respective modules
- Added the CVar to make the Audio Import via Interchange optional feature.

TODO:
- Add support for CuePoints and Loop Regions
- Add the remaining code from the Sound Factory to the Interchange Sound Wave Factory.

Tests:
- Have added 2 tests related to import of Audio Assets from the AudioQA project adapted to the rest of the Interchange Python Automation Tests.
- Duplicated the assets from AudioQA

#rnx

#virtualized

diff --git a/Engine/Plugins/Interchange/Runtime/Source/FactoryNodes/Public/InterchangeAudioSoundWaveFactoryNode.h b/Engine/Plugins/Interchange/Runtime/Source/FactoryNodes/Public/InterchangeAudioSoundWaveFactoryNode.h
new file mode 100644
--- /dev/null
+++ b/Engine/Plugins/Interchange/Runtime/Source/FactoryNodes/Public/InterchangeAudioSoundWaveFactoryNode.h
@@ -0,0 +1,25 @@
+// Copyright Epic Games, Inc. All Rights Reserved.
+
+#pragma once
+
+#include "Nodes/InterchangeFactoryBaseNode.h"
+#include "Sound/SoundWave.h"
+
+#include "InterchangeAudioSoundWaveFactoryNode.generated.h"
+
+UCLASS(BlueprintType)
+class INTERCHANGEFACTORYNODES_API UInterchangeAudioSoundWaveFactoryNode : public UInterchangeFactoryBaseNode
+{
+	GENERATED_BODY()
+
+public:
+	virtual FString GetTypeName() const override
+	{
+		return TEXT("SoundWaveFactoryNode");
+	}
+
+	virtual UClass* GetObjectClass() const override
+	{
+		return USoundWave::StaticClass();
+	}
+};
diff --git a/Engine/Plugins/Interchange/Runtime/Source/Import/InterchangeImport.Build.cs b/Engine/Plugins/Interchange/Runtime/Source/Import/InterchangeImport.Build.cs
--- a/Engine/Plugins/Interchange/Runtime/Source/Import/InterchangeImport.Build.cs
+++ b/Engine/Plugins/Interchange/Runtime/Source/Import/InterchangeImport.Build.cs
@@ -32,6 +32,7 @@ namespace UnrealBuildTool.Rules
 				new string[]
 				{
 					"AssetRegistry",
+					"AudioMixer",
 					"CinematicCamera",
 					"ClothingSystemRuntimeCommon",
 					"GeometryCache",
diff --git a/Engine/Plugins/Interchange/Runtime/Source/Import/Private/Audio/InterchangeAudioPayloadData.cpp b/Engine/Plugins/Interchange/Runtime/Source/Import/Private/Audio/InterchangeAudioPayloadData.cpp
new file mode 100644
--- /dev/null
+++ b/Engine/Plugins/Interchange/Runtime/Source/Import/Private/Audio/InterchangeAudioPayloadData.cpp
@@ -0,0 +1,132 @@
+// Copyright Epic Games, Inc. All Rights Reserved.
+#include "Audio/InterchangeAudioPayloadData.h"
+#include "Misc/Paths.h"
+#include "Misc/FileHelper.h"
+#include "SoundFileIO/SoundFileIO.h"
+
+#include "InterchangeImportLog.h"
+
+namespace UE::Interchange::Private
+{
+	bool ConvertAudioFile(const TArray<uint8>& RawWaveData, FWaveModInfo& WaveInfo, TArray<uint8>& ConvertedRawWaveData, FString& ErrorMessage, bool bEnsureEqualNumSamples = true)
+	{
+		// This duplication will make sure that the buffer is not having
+		// a mismatch in length and capacity. This is a must for Sound File IO operations
+		// Otherwise a buffer overflow will be observed.
+		TArray<uint8> RawAudioData;
+		RawAudioData.Reserve(RawWaveData.Num());
+		RawAudioData.AddUninitialized(RawWaveData.Num());
+		FMemory::Memcpy(RawAudioData.GetData(), RawWaveData.GetData(), RawWaveData.Num());
+		
+		const uint32 OrigNumSamples = Audio::SoundFileUtils::GetNumSamples(RawAudioData);
+
+		if (Audio::SoundFileUtils::ConvertAudioToWav(RawAudioData, ConvertedRawWaveData))
+		{
+			WaveInfo = FWaveModInfo();
+			if (!WaveInfo.ReadWaveInfo(ConvertedRawWaveData.GetData(), ConvertedRawWaveData.Num(), &ErrorMessage))
+			{
+				return false;
+			}
+		}
+
+		if (bEnsureEqualNumSamples)
+		{
+			const uint32 ConvertedNumSamples = WaveInfo.GetNumSamples();
+			if (!ensure(ConvertedNumSamples == OrigNumSamples))
+			{
+				ErrorMessage = "The converted wave file does not have the same number of samples as the original file.";
+				return false;
+			}
+		}
+
+		return true;
+	}
+}
+
+TOptional<UE::Interchange::FInterchangeAudioPayloadData> FInterchangeAudioPayloadDataUtils::GetAudioPayloadFromSourceFileKey(const FString& PayloadSourceFileKey)
+{
+	using namespace UE::Interchange;
+
+	if (!FPaths::FileExists(*PayloadSourceFileKey))
+	{
+		UE_LOG(LogInterchangeImport, Error, TEXT("[Audio] File does not exist: [%s]"), *PayloadSourceFileKey);
+		return TOptional<FInterchangeAudioPayloadData>();
+	}
+
+	TArray<uint8> ReadBuffer;
+	if (!FFileHelper::LoadFileToArray(ReadBuffer, *PayloadSourceFileKey))
+	{	
+		UE_LOG(LogInterchangeImport, Error, TEXT("[Audio] Failed to load file: [%s]"), *PayloadSourceFileKey);
+		return TOptional<FInterchangeAudioPayloadData>();
+	}
+
+	const uint8* Buffer = &ReadBuffer[0];
+	const uint8* BufferEnd = Buffer + ReadBuffer.Num() - 1;
+	
+	
+
+	const uint64 Size = BufferEnd - Buffer;
+	if (!IntFitsIn<int32>(Size))
+	{
+		UE_LOG(LogInterchangeImport, Error, TEXT("[Audio] Audio file is too large: [%s]"), *PayloadSourceFileKey);
+		return TOptional<FInterchangeAudioPayloadData>();
+	}
+
+	FInterchangeAudioPayloadData AudioPayloadData;
+	FWaveModInfo WaveInfo;
+	FString ErrorMessage;
+
+	const bool bFileNeedsConversion = !FPaths::GetExtension(*(PayloadSourceFileKey.ToLower())).Equals(TEXT(".wav"));
+	if (!WaveInfo.ReadWaveInfo(ReadBuffer.GetData(), ReadBuffer.Num(), &ErrorMessage))
+	{
+		if (!bFileNeedsConversion)
+		{
+			UE_LOG(LogInterchangeImport, Error, TEXT("[Audio] Failed to retrieve WaveFileInfo for '%s' : %s"), *PayloadSourceFileKey, *ErrorMessage);
+			return TOptional<FInterchangeAudioPayloadData>();
+		}
+	}
+
+	if (bFileNeedsConversion || *WaveInfo.pBitsPerSample != 16 || !WaveInfo.IsFormatSupported())
+	{
+		TArray<uint8> ConvertedRawWaveData;
+#if WITH_SNDFILE_IO
+		const bool bEnsureEqualNumSamples = !bFileNeedsConversion;
+		if (!Private::ConvertAudioFile(ReadBuffer, WaveInfo, ConvertedRawWaveData, ErrorMessage, bEnsureEqualNumSamples))
+		{
+			UE_LOG(LogInterchangeImport, Error, TEXT("[Audio] Failed to convert '%s' audio file: %s"), *PayloadSourceFileKey, *ErrorMessage);
+			return TOptional<FInterchangeAudioPayloadData>();
+		}
+		AudioPayloadData.Buffer = MoveTemp(ConvertedRawWaveData);
+#else
+		WaveInfo.ReportImportFailure();
+		return TOptional<FInterchangeAudioPayloadData>();
+#endif
+	}
+	else
+	{
+		AudioPayloadData.Buffer = MoveTemp(ReadBuffer);
+	}
+
+	WaveInfo = FWaveModInfo();
+	if (!WaveInfo.ReadWaveInfo(AudioPayloadData.Buffer.GetData(), AudioPayloadData.Buffer.Num(), &ErrorMessage))
+	{
+		WaveInfo.ReportImportFailure();
+		return TOptional<FInterchangeAudioPayloadData>();
+	}
+	
+	check(*WaveInfo.pBitsPerSample == 16);
+	AudioPayloadData.FactoryInfo = FSoundWaveFactoryInfo();
+	AudioPayloadData.FactoryInfo.ChannelCount = (int32)*WaveInfo.pChannels;
+	check(AudioPayloadData.FactoryInfo.ChannelCount > 0);
+
+	AudioPayloadData.FactoryInfo.SizeOfSample = (*WaveInfo.pBitsPerSample) / 8;
+	AudioPayloadData.FactoryInfo.NumSamples = WaveInfo.SampleDataSize / AudioPayloadData.FactoryInfo.SizeOfSample;
+	AudioPayloadData.FactoryInfo.SamplesPerSec = (int32)*WaveInfo.pSamplesPerSec;
+	AudioPayloadData.FactoryInfo.SampleDataOffset = (int16*)WaveInfo.SampleDataStart - (int16*)AudioPayloadData.Buffer.GetData();
+	ensure(AudioPayloadData.FactoryInfo.SampleDataOffset > 0);
+
+	AudioPayloadData.FactoryInfo.NumFrames = AudioPayloadData.FactoryInfo.NumSamples / AudioPayloadData.FactoryInfo.ChannelCount;
+
+	AudioPayloadData.WaveModInfo = WaveInfo;
+	return TOptional<FInterchangeAudioPayloadData>(AudioPayloadData);
+}
\ No newline at end of file
diff --git a/Engine/Plugins/Interchange/Runtime/Source/Import/Private/Audio/InterchangeAudioSoundWaveFactory.cpp b/Engine/Plugins/Interchange/Runtime/Source/Import/Private/Audio/InterchangeAudioSoundWaveFactory.cpp
new file mode 100644
--- /dev/null
+++ b/Engine/Plugins/Interchange/Runtime/Source/Import/Private/Audio/InterchangeAudioSoundWaveFactory.cpp
@@ -0,0 +1,629 @@
+﻿// Copyright Epic Games, Inc. All Rights Reserved.
+
+#include "Audio/InterchangeAudioSoundWaveFactory.h"
+#include "Engine/Engine.h"
+#include "AudioDeviceManager.h"
+#include "Components/AudioComponent.h"
+#include "InterchangeAssetImportData.h"
+#include "Audio/InterchangeAudioPayloadInterface.h"
+#include "InterchangeAudioSoundWaveFactoryNode.h"
+#include "InterchangeAudioSoundWaveNode.h"
+#include "InterchangeImportCommon.h"
+#include "InterchangeTranslatorBase.h"
+#include "Sound/SoundWave.h"
+#include "SoundFileIO/SoundFileIO.h"
+
+#include UE_INLINE_GENERATED_CPP_BY_NAME(InterchangeAudioSoundWaveFactory)
+
+#define LOCTEXT_NAMESPACE "InterchangeAudioSoundWaveFactory"
+
+UClass* UInterchangeAudioSoundWaveFactory::GetFactoryClass() const
+{
+	return USoundWave::StaticClass();
+}
+
+UInterchangeFactoryBase::FImportAssetResult UInterchangeAudioSoundWaveFactory::BeginImportAsset_GameThread(const FImportAssetObjectParams& Arguments)
+{
+	TRACE_CPUPROFILER_EVENT_SCOPE(UInterchangeAudioSoundWaveFactory::BeginImportAsset_GameThread);
+
+	FImportAssetResult ImportAssetResult;
+
+	if (!Arguments.AssetNode || !Arguments.AssetNode->GetObjectClass()->IsChildOf(GetFactoryClass()))
+	{
+		return ImportAssetResult;
+	}
+
+	const UInterchangeAudioSoundWaveFactoryNode* SoundWaveFactoryNode = Cast<UInterchangeAudioSoundWaveFactoryNode>(Arguments.AssetNode);
+	if (!SoundWaveFactoryNode)
+	{
+		const FText Message = FText::Format(LOCTEXT("InterchangeAudioSoundWaveFactory_ImportSoundWave_FactoryNodeFailedCast", "Failed to cast AssetNode to UInterchangeAudioSoundWaveFactoryNode. AssetNode class: `{0}`")
+			, FText::FromString(Arguments.AssetNode->GetObjectClass()->GetName()));
+		LogError(Arguments, ImportAssetResult, Message);
+		return ImportAssetResult;
+	}
+
+	UObject* ExistingAsset = Arguments.ReimportObject;
+	if (!ExistingAsset)
+	{
+		FSoftObjectPath ReferenceObject;
+		if (SoundWaveFactoryNode->GetCustomReferenceObject(ReferenceObject))
+		{
+			ExistingAsset = ReferenceObject.TryLoad();
+		}
+	}
+
+	if (!ensure(SoundWaveFactoryNode->GetTargetNodeCount() == 1))
+	{
+		const FText Message = FText::Format(LOCTEXT("InterchangeAudioSoundWaveFactory_ImportSoundWave_UnexpectedNumTargets"
+			, "Expected one target node, but found {0} target nodes associated with SoundWave factory node `{1}`")
+			, SoundWaveFactoryNode->GetTargetNodeCount()
+			, FText::FromString(SoundWaveFactoryNode->GetDisplayLabel()));
+		LogError(Arguments, ImportAssetResult, Message);
+		return ImportAssetResult;
+	}
+
+	USoundWave* SoundWaveAsset = nullptr;
+	if (ExistingAsset)
+	{
+		SoundWaveAsset = Cast<USoundWave>(ExistingAsset);
+		if (!SoundWaveAsset)
+		{
+			const FText Message = FText::Format(LOCTEXT("InterchangeAudioSoundWaveFactory_ImportSoundWave_InvalidExistingAsset", "Existing asset `{0}` must be a USoundWave asset.")
+				, FText::FromString(ExistingAsset->GetName()));
+			LogError(Arguments, ImportAssetResult, Message);
+			return ImportAssetResult;
+		}
+
+		StopComponentsUsingImportedSound(SoundWaveAsset);
+	}
+	else
+	{
+		SoundWaveAsset = NewObject<USoundWave>(Arguments.Parent, GetFactoryClass(), *Arguments.AssetName, RF_Public | RF_Standalone);
+	}
+
+	if (!SoundWaveAsset)
+	{
+		const FText Message = FText::Format(LOCTEXT("InterchangeAudioSoundWaveFactory_ImportSoundWave_CouldNotCreateSoundWaveAsset", "Could not create SoundWave asset for factory node `{0}`")
+			, FText::FromString(SoundWaveFactoryNode->GetDisplayLabel()));
+		LogError(Arguments, ImportAssetResult, Message);
+		return ImportAssetResult;
+	}
+
+	ImportAssetResult.ImportedObject = SoundWaveAsset;
+	return ImportAssetResult;
+}
+
+UInterchangeFactoryBase::FImportAssetResult UInterchangeAudioSoundWaveFactory::ImportAsset_Async(const FImportAssetObjectParams& Arguments)
+{
+	using namespace UE::Interchange;
+	FImportAssetResult ImportAssetResult;
+
+	if (!Arguments.AssetNode || !Arguments.AssetNode->GetObjectClass()->IsChildOf(GetFactoryClass()))
+	{
+		return ImportAssetResult;
+	}
+
+	UObject* SoundWaveObject = FFactoryCommon::AsyncFindObject(Arguments.AssetNode, GetFactoryClass(), Arguments.Parent, Arguments.AssetName);
+	const bool bIsReimport = Arguments.ReimportObject && SoundWaveObject;
+	if (!SoundWaveObject)
+	{
+		LogError(Arguments, ImportAssetResult, 
+			FText::Format(
+				LOCTEXT("UInterchangeAudioSoundWaveFactory_ImportSoundWave_NoExistingAsset", "Could not import the SoundWaveAsset '{0}' because the asset does not exist."),
+				FText::FromString(Arguments.AssetName)
+			)
+		);
+		return ImportAssetResult;
+	}
+
+	USoundWave* SoundWave = Cast<USoundWave>(SoundWaveObject);
+	if (!ensure(SoundWave))
+	{
+		LogError(Arguments, ImportAssetResult, LOCTEXT("InterchangeAudioSoundWaveFactory_ImportSoundWave_FailedToCast", "Could not cast to SoundWave."));
+		return ImportAssetResult;
+	}
+	
+	const UInterchangeAudioSoundWaveFactoryNode* SoundWaveFactoryNode = Cast<UInterchangeAudioSoundWaveFactoryNode>(Arguments.AssetNode);
+	if (!ensure(SoundWaveFactoryNode))
+	{
+		const FText Message = FText::Format(LOCTEXT("InterchangeAudioSoundWaveFactory_ImportSoundWave_FactoryNodeFailedCast", "Failed to cast AssetNode to UInterchangeAudioSoundWaveFactoryNode. AssetNode class: `{0}`")
+			, FText::FromString(Arguments.AssetNode->GetObjectClass()->GetName()));
+		LogError(Arguments, ImportAssetResult, Message);
+		return ImportAssetResult;
+	}
+
+	TArray<FString> TargetNodeUids;
+	SoundWaveFactoryNode->GetTargetNodeUids(TargetNodeUids);
+	const UInterchangeBaseNode* TargetNode = Arguments.NodeContainer->GetNode(TargetNodeUids[0]);
+	if (!ensure(TargetNode))
+	{
+		const FText Message = LOCTEXT("InterchangeAudioSoundWaveFactory_ImportSoundWave_InvalidTargetNode", "Target node was null.");
+		LogError(Arguments, ImportAssetResult, Message);
+		return ImportAssetResult;
+	}
+
+	const UInterchangeAudioSoundWaveNode* SoundWaveNode = Cast<UInterchangeAudioSoundWaveNode>(TargetNode);
+	if (!SoundWaveNode)
+	{
+		const FText Message = FText::Format(LOCTEXT("InterchangeAudioSoundWaveFactory_ImportSoundWave_NullSoundWaveNode", "Failed to cast target node `{0}` to a sound wave node.")
+			, FText::FromString(TargetNode->GetDisplayLabel()));
+		LogError(Arguments, ImportAssetResult, Message);
+		return ImportAssetResult;
+	}
+
+	TOptional<FString> SourceFilePayloadKey = SoundWaveNode->GetPayloadKey();
+	if (!SourceFilePayloadKey.IsSet())
+	{
+		const FText Message = FText::Format(LOCTEXT("InterchangeAudioSoundWaveFactory_ImportSoundWave_SoundWaveNodeInvalidPayload", "Sound wave node `{0}` doesn't have a valid Payload.")
+			, FText::FromString(SoundWaveNode->GetDisplayLabel()));
+		LogError(Arguments, ImportAssetResult, Message);
+		return ImportAssetResult;
+	}
+
+	const IInterchangeAudioPayloadInterface* AudioTranslator = Cast<IInterchangeAudioPayloadInterface>(Arguments.Translator);
+	if (!AudioTranslator)
+	{
+		LogError(Arguments, ImportAssetResult, LOCTEXT("InterchangeAudioSoundWaveFactory_ImportSoundWave_InvalidTranslator", "Translator used doesn't implement IInterchangeAudioPayloadInterface."));
+		return ImportAssetResult;
+	}
+
+	TOptional<UE::Interchange::FInterchangeAudioPayloadData> AudioPayloadData = AudioTranslator->GetAudioPayloadData(SourceFilePayloadKey.GetValue());
+	if (!AudioPayloadData.IsSet())
+	{
+		LogError(Arguments, ImportAssetResult, FText::Format(LOCTEXT("InterchangeAudioSoundWaveFactory_ImportSoundWave_FailedToGetAudioPayloadData", "Failed to retrieve audio payload from file [{0}]. See output log for more details."), FText::FromString(*SourceFilePayloadKey.GetValue())));
+		return ImportAssetResult;
+	}
+
+	SetUpSoundWaveFromPayload(Arguments, SoundWave, AudioPayloadData.GetValue(), ImportAssetResult);
+
+	ImportAssetResult.ImportedObject = SoundWave;
+	return ImportAssetResult;
+}
+
+UInterchangeFactoryBase::FImportAssetResult UInterchangeAudioSoundWaveFactory::EndImportAsset_GameThread(const FImportAssetObjectParams& Arguments)
+{
+	FImportAssetResult ImportAssetResult;
+
+	if (!Arguments.AssetNode || !Arguments.AssetNode->GetObjectClass()->IsChildOf(GetFactoryClass()))
+	{
+		return ImportAssetResult;
+	}
+
+	UObject* ExistingAsset = nullptr;
+	FSoftObjectPath ReferenceObject;
+	if (Arguments.AssetNode->GetCustomReferenceObject(ReferenceObject))
+	{
+		ExistingAsset = ReferenceObject.TryLoad();
+	}
+
+	//We need a valid imported object
+	if (!ExistingAsset)
+	{
+		return ImportAssetResult;
+	}
+
+	if (USoundWave* SoundWaveAsset = Cast<USoundWave>(ExistingAsset))
+	{
+		constexpr bool bFreeResources = true;
+		SoundWaveAsset->InvalidateCompressedData(bFreeResources);
+
+		if (SoundWaveAsset->IsStreaming(nullptr))
+		{
+			SoundWaveAsset->LoadZerothChunk();
+		}
+
+#if WITH_EDITOR
+		SoundWaveAsset->PostImport();
+#endif
+		ImportAssetResult.ImportedObject = SoundWaveAsset;
+	}
+
+	return ImportAssetResult;
+}
+
+void UInterchangeAudioSoundWaveFactory::SetupObject_GameThread(const FSetupObjectParams& Arguments)
+{
+	TRACE_CPUPROFILER_EVENT_SCOPE(UInterchangeAudioSoundWaveFactory::SetupObject_GameThread);
+
+	check(IsInGameThread());
+
+	USoundWave* SoundWaveAsset = Cast<USoundWave>(Arguments.ImportedObject);
+	if (!SoundWaveAsset)
+	{
+		return;
+	}
+
+#if WITH_EDITOR
+	SoundWaveAsset->PreEditChange(nullptr);
+
+	UInterchangeFactoryBaseNode* SoundWaveFactoryNode = Arguments.FactoryNode;
+	if (!SoundWaveFactoryNode)
+	{
+		return;
+	}
+
+	if (Arguments.bIsReimport)
+	{
+		const UInterchangeAssetImportData* InterchangeAssetImportData = Cast<UInterchangeAssetImportData>(SoundWaveAsset->AssetImportData);
+		const UInterchangeFactoryBaseNode* PreviousNode = nullptr;
+		if (InterchangeAssetImportData)
+		{
+			PreviousNode = InterchangeAssetImportData->GetStoredFactoryNode(InterchangeAssetImportData->NodeUniqueID);
+		}
+
+		UInterchangeFactoryBaseNode* CurrentNode = NewObject<UInterchangeFactoryBaseNode>(GetTransientPackage(), UInterchangeAudioSoundWaveFactoryNode::StaticClass());
+		ensure(CurrentNode);
+		UInterchangeBaseNode::CopyStorage(SoundWaveFactoryNode, CurrentNode);
+		CurrentNode->FillAllCustomAttributeFromObject(SoundWaveAsset);
+
+		UE::Interchange::FFactoryCommon::ApplyReimportStrategyToAsset(SoundWaveAsset, PreviousNode, CurrentNode, SoundWaveFactoryNode);
+	}
+	else
+	{
+		SoundWaveFactoryNode->ApplyAllCustomAttributeToObject(SoundWaveAsset);
+	}
+#endif
+
+#if WITH_EDITORONLY_DATA
+	if (ensure(Arguments.SourceData))
+	{
+		UE::Interchange::FFactoryCommon::FSetImportAssetDataParameters SetImportAssetDataParameters(
+			SoundWaveAsset,
+			SoundWaveAsset->AssetImportData,
+			Arguments.SourceData,
+			Arguments.NodeUniqueID,
+			Arguments.NodeContainer,
+			Arguments.OriginalPipelines,
+			Arguments.Translator);
+
+		SoundWaveAsset->AssetImportData = UE::Interchange::FFactoryCommon::SetImportAssetData(SetImportAssetDataParameters);
+	}
+#endif
+}
+
+void UInterchangeAudioSoundWaveFactory::FinalizeObject_GameThread(const FSetupObjectParams& Arguments)
+{
+	TRACE_CPUPROFILER_EVENT_SCOPE(UInterchangeAudioSoundWaveFactory::FinalizeObject_GameThread);
+
+	check(IsInGameThread());
+
+	if (!Arguments.ImportedObject)
+	{
+		return;
+	}
+
+	if (Arguments.bIsReimport)
+	{
+		for (UAudioComponent* AudioComponent : ComponentsToRestart)
+		{
+			AudioComponent->Play();
+		}
+	}
+
+#if WITH_EDITOR
+	if (USoundWave* ImportedSoundWaveAsset = Cast<USoundWave>(Arguments.ImportedObject))
+	{
+		ImportedSoundWaveAsset->SetRedrawThumbnail(true);
+	}
+#endif
+
+}
+
+bool UInterchangeAudioSoundWaveFactory::GetSourceFilenames(const UObject* Object, TArray<FString>& OutSourceFilenames) const
+{
+#if WITH_EDITORONLY_DATA
+	if (const USoundWave* SoundWave = Cast<USoundWave>(Object))
+	{
+		return UE::Interchange::FFactoryCommon::GetSourceFilenames(SoundWave->AssetImportData.Get(), OutSourceFilenames);
+	}
+#endif
+
+	return false;
+}
+
+bool UInterchangeAudioSoundWaveFactory::SetSourceFilename(const UObject* Object, const FString& SourceFilename, int32 SourceIndex) const
+{
+#if WITH_EDITORONLY_DATA
+	if (const USoundWave* SoundWave = Cast<USoundWave>(Object))
+	{
+		return UE::Interchange::FFactoryCommon::SetSourceFilename(SoundWave->AssetImportData.Get(), SourceFilename, SourceIndex);
+	}
+#endif
+
+	return false;
+}
+
+void UInterchangeAudioSoundWaveFactory::BackupSourceData(const UObject* Object) const
+{
+#if WITH_EDITORONLY_DATA
+	if (const USoundWave* SoundWave = Cast<USoundWave>(Object))
+	{
+		UE::Interchange::FFactoryCommon::BackupSourceData(SoundWave->AssetImportData.Get());
+	}
+#endif
+}
+
+void UInterchangeAudioSoundWaveFactory::ReinstateSourceData(const UObject* Object) const
+{
+#if WITH_EDITORONLY_DATA
+	if (const USoundWave* SoundWave = Cast<USoundWave>(Object))
+	{
+		UE::Interchange::FFactoryCommon::ReinstateSourceData(SoundWave->AssetImportData.Get());
+	}
+#endif
+}
+
+void UInterchangeAudioSoundWaveFactory::ClearBackupSourceData(const UObject* Object) const
+{
+#if WITH_EDITORONLY_DATA
+	if (const USoundWave* SoundWave = Cast<USoundWave>(Object))
+	{
+		UE::Interchange::FFactoryCommon::ClearBackupSourceData(SoundWave->AssetImportData.Get());
+	}
+#endif
+}
+
+void UInterchangeAudioSoundWaveFactory::SetUpSoundWaveFromPayload(const FImportAssetObjectParams& Arguments, USoundWave*& SoundWaveAsset, const UE::Interchange::FInterchangeAudioPayloadData& AudioPayloadData, FImportAssetResult& ImportAssetResult)
+{
+	if (!SoundWaveAsset)
+	{
+		const FText Message = LOCTEXT("InterchangeAudioSoundWaveFactory_ImportWaveData_NullAsset", "Sound wave asset was null.");
+		LogError(Arguments, ImportAssetResult, Message);
+		return;
+	}
+
+	if (AudioPayloadData.FactoryInfo.ChannelCount > 2)
+	{
+		if (!SetUpMultichannelSoundWave(Arguments, SoundWaveAsset, AudioPayloadData, ImportAssetResult))
+		{
+			const FText Message = LOCTEXT("InterchangeAudioSoundWaveFactory_ImportWaveData_SetUpMultiChanFailure", "Failed to set up multichannel sound wave asset.");
+			LogError(Arguments, ImportAssetResult, Message);
+			return;
+		}
+	}
+	else
+	{
+		SoundWaveAsset->bIsAmbisonics = false;
+
+#if WITH_EDITORONLY_DATA
+		SoundWaveAsset->ChannelOffsets.Reset();
+		SoundWaveAsset->ChannelSizes.Reset();
+		SoundWaveAsset->RawData.UpdatePayload(FSharedBuffer::Clone(AudioPayloadData.Buffer.GetData(), AudioPayloadData.Buffer.Num()));
+#endif
+	}
+
+	SoundWaveAsset->Duration = (float)AudioPayloadData.FactoryInfo.NumFrames / AudioPayloadData.FactoryInfo.SamplesPerSec;
+	SoundWaveAsset->SetImportedSampleRate(AudioPayloadData.FactoryInfo.SamplesPerSec);
+	SoundWaveAsset->SetSampleRate(AudioPayloadData.FactoryInfo.SamplesPerSec);
+	SoundWaveAsset->NumChannels = AudioPayloadData.FactoryInfo.ChannelCount;
+	SoundWaveAsset->TotalSamples = AudioPayloadData.FactoryInfo.SamplesPerSec * SoundWaveAsset->Duration;
+
+#if WITH_EDITORONLY_DATA
+	if (AudioPayloadData.WaveModInfo.TimecodeInfo)
+	{
+		SoundWaveAsset->SetTimecodeInfo(*AudioPayloadData.WaveModInfo.TimecodeInfo);
+	}
+	else
+	{
+		SoundWaveAsset->SetTimecodeInfo(FSoundWaveTimecodeInfo{});
+	}
+#endif
+
+}
+
+void UInterchangeAudioSoundWaveFactory::StopComponentsUsingImportedSound(USoundWave* SoundWaveAsset)
+{
+	if (!SoundWaveAsset)
+	{
+		return;
+	}
+
+	FAudioDeviceManager* AudioDeviceManager = GEngine->GetAudioDeviceManager();
+	if (!AudioDeviceManager)
+	{
+		return;
+	}
+
+	TArray<UAudioComponent*> AudioComponents = ComponentsToRestart;
+	AudioDeviceManager->StopSoundsUsingResource(SoundWaveAsset, &AudioComponents);
+
+	if (!SoundWaveAsset->GetResourceData())
+	{
+		const FName RuntimeFormat = SoundWaveAsset->GetRuntimeFormat();
+		SoundWaveAsset->InitAudioResource(RuntimeFormat);
+	}
+
+	if (ComponentsToRestart.Num() > 0)
+	{
+		for (UAudioComponent* AudioComponent : ComponentsToRestart)
+		{
+			AudioComponent->Stop();
+		}
+	}
+}
+
+bool UInterchangeAudioSoundWaveFactory::SetUpMultichannelSoundWave(const FImportAssetObjectParams& Arguments, USoundWave*& SoundWaveAsset, const UE::Interchange::FInterchangeAudioPayloadData& AudioPayloadData, FImportAssetResult& ImportAssetResult)
+{
+	using namespace UE::Interchange;
+
+	const FString RootName = Arguments.AssetName;
+	const FString AmbiXTag = RootName.Right(6).ToLower();
+	const FString FuMaTag = RootName.Right(5).ToLower();
+
+	const bool bIsAmbiX = AmbiXTag == TEXT("_ambix");
+	const bool bIsFuMa = FuMaTag == TEXT("_fuma");
+
+	const FSoundWaveFactoryInfo SoundWaveFactoryInfo = AudioPayloadData.FactoryInfo;
+
+	if ((bIsAmbiX || bIsFuMa) && SoundWaveFactoryInfo.ChannelCount != 4)
+	{
+		const FText Message = LOCTEXT("InterchangeAudioSoundWaveFactory_SetUpMultichan_InvalidChanCountAmbi", "Tried to import ambisonics format file but the file does not contain exactly four channels.");
+		LogError(Arguments, ImportAssetResult, Message);
+		return false;
+	}
+
+	const int16* SampleDataStart = (int16*)AudioPayloadData.Buffer.GetData() + SoundWaveFactoryInfo.SampleDataOffset;
+	// Starting or ending on a non-zero sample creates a pop unless you have very specific use cases
+	// (doesn't apply to ambi)
+	if (!bIsFuMa && !bIsAmbiX)
+	{
+		for (int32 Chan = 0; Chan < SoundWaveFactoryInfo.ChannelCount; ++Chan)
+		{
+			const int16* FirstFrame = SampleDataStart;
+			const int16* LastFrame = FirstFrame + (SoundWaveFactoryInfo.NumFrames - 1) * SoundWaveFactoryInfo.ChannelCount;
+
+			FirstFrame += Chan;
+			LastFrame += Chan;
+
+			if (FirstFrame[0] != 0)
+			{
+				//Warn->Logf(ELogVerbosity::Warning, TEXT("Channel %d starts with a non zero value (%d) - will likely pop and not loop well: '%s'"), Chan, *FirstFrame, *Name.ToString());
+			}
+
+			if (LastFrame[0] != 0)
+			{
+				//Warn->Logf(ELogVerbosity::Warning, TEXT("Channel %d ends with a non zero value (%d) - will likely pop and not loop well: '%s'"), Chan, *FirstFrame, *Name.ToString());
+			}
+		}
+	}
+
+#if WITH_EDITORONLY_DATA
+	SoundWaveAsset->ChannelOffsets.Empty(SPEAKER_Count);
+	SoundWaveAsset->ChannelOffsets.AddZeroed(SPEAKER_Count);
+
+	SoundWaveAsset->ChannelSizes.Empty(SPEAKER_Count);
+	SoundWaveAsset->ChannelSizes.AddZeroed(SPEAKER_Count);
+#endif
+
+	TArray<int32> ChannelIndices;
+	if (SoundWaveFactoryInfo.ChannelCount == 4)
+	{
+		ChannelIndices = {
+			SPEAKER_FrontLeft,
+			SPEAKER_FrontRight,
+			SPEAKER_LeftSurround,
+			SPEAKER_RightSurround
+		};
+	}
+	else if (SoundWaveFactoryInfo.ChannelCount == 6)
+	{
+		ChannelIndices = {
+			SPEAKER_FrontLeft,
+			SPEAKER_FrontRight,
+			SPEAKER_FrontCenter,
+			SPEAKER_LowFrequency,
+			SPEAKER_LeftSurround,
+			SPEAKER_RightSurround
+		};
+	}
+	else if (SoundWaveFactoryInfo.ChannelCount == 8)
+	{
+		ChannelIndices = {
+			SPEAKER_FrontLeft,
+			SPEAKER_FrontRight,
+			SPEAKER_FrontCenter,
+			SPEAKER_LowFrequency,
+			SPEAKER_LeftSurround,
+			SPEAKER_RightSurround,
+			SPEAKER_LeftBack,
+			SPEAKER_RightBack
+		};
+	}
+	else
+	{
+		const FText Message = LOCTEXT("InterchangeAudioSoundWaveFactory_SetUpMultichan_UnknownChanConfig", "Unable to import multichannel file with unsupported channel configuration.");
+		LogError(Arguments, ImportAssetResult, Message);
+		return false;
+	}
+
+	int32 TotalSize = 0;
+	TArray<uint8> RawChannelWaveData[SPEAKER_Count];
+	TArray<int16> DeinterleavedAudioScratchBuffer;
+	const int16* SampleDataBuffer = SampleDataStart;
+
+	check(SoundWaveFactoryInfo.ChannelCount == ChannelIndices.Num());
+	for (int32 Chan = 0; Chan < SoundWaveFactoryInfo.ChannelCount; ++Chan)
+	{
+		DeinterleavedAudioScratchBuffer.Empty();
+		for (int32 Frame = 0; Frame < SoundWaveFactoryInfo.NumFrames; ++Frame)
+		{
+			const int32 SampleIndex = Frame * SoundWaveFactoryInfo.ChannelCount + Chan;
+			DeinterleavedAudioScratchBuffer.Add(SampleDataBuffer[SampleIndex]);
+		}
+
+		SerializeWaveFile(RawChannelWaveData[Chan]
+			, (uint8*)DeinterleavedAudioScratchBuffer.GetData()
+			, SoundWaveFactoryInfo.NumFrames * sizeof(int16)
+			, 1
+			, SoundWaveFactoryInfo.SamplesPerSec);
+#if WITH_EDITORONLY_DATA
+		SoundWaveAsset->ChannelOffsets[ChannelIndices[Chan]] = TotalSize;
+		const int32 ChannelSize = RawChannelWaveData[Chan].Num();
+		SoundWaveAsset->ChannelSizes[ChannelIndices[Chan]] = ChannelSize;
+		TotalSize += ChannelSize;
+#endif
+	}
+
+	FUniqueBuffer EditableBuffer = FUniqueBuffer::Alloc(TotalSize);
+	uint8* LockedData = (uint8*)EditableBuffer.GetData();
+	int32 RawDataOffset = 0;
+
+	if (bIsAmbiX || bIsFuMa)
+	{
+		check(SoundWaveFactoryInfo.ChannelCount == 4)
+		SoundWaveAsset->bIsAmbisonics = true;
+	}
+	if (bIsFuMa)
+	{
+		int32 FuMaChannelIndices[4] = { 0, 2, 3, 1 };
+		for (const int32 ChannelIndex : FuMaChannelIndices)
+		{
+			const int32 ChannelSize = RawChannelWaveData[ChannelIndex].Num();
+			FMemory::Memcpy(LockedData + RawDataOffset, RawChannelWaveData[ChannelIndex].GetData(), ChannelSize);
+			RawDataOffset += ChannelSize;
+		}
+	}
+	else
+	{
+		for (int32 Chan = 0; Chan < SoundWaveFactoryInfo.ChannelCount; ++Chan)
+		{
+			const int32 ChannelSize = RawChannelWaveData[Chan].Num();
+			FMemory::Memcpy(LockedData + RawDataOffset, RawChannelWaveData[Chan].GetData(), ChannelSize);
+			RawDataOffset += ChannelSize;
+		}
+	}
+#if WITH_EDITORONLY_DATA
+	SoundWaveAsset->RawData.UpdatePayload(EditableBuffer.MoveToShared());
+#endif
+	return true;
+}
+
+void UInterchangeAudioSoundWaveFactory::LogError(const FImportAssetObjectParams& Arguments, FImportAssetResult& ImportAssetResult, const FText& ErrorText)
+{
+	UInterchangeResultError_Generic* ErrorMessage = AddMessage<UInterchangeResultError_Generic>();
+	if (ensure(ErrorMessage && Arguments.Translator))
+	{
+		ErrorMessage->AssetType = GetFactoryClass();
+
+		if (const UInterchangeSourceData* SourceData = Arguments.Translator->GetSourceData())
+		{
+			const FString Filename = SourceData->GetFilename();
+			ErrorMessage->SourceAssetName = Arguments.AssetName;
+			ErrorMessage->DestinationAssetName = Arguments.AssetName;
+			ErrorMessage->InterchangeKey = FPaths::GetBaseFilename(Filename);
+		}
+		else
+		{
+			ErrorMessage->InterchangeKey = TEXT("Undefined");
+		}
+
+		ImportAssetResult = FImportAssetResult{
+			.bIsFactorySkipAsset = true,
+			.ImportedObject = nullptr
+		};
+
+		ErrorMessage->Text = ErrorText;
+	}
+}
+
+#undef LOCTEXT_NAMESPACE
diff --git a/Engine/Plugins/Interchange/Runtime/Source/Import/Private/Audio/InterchangeGenericAudioTranslator.cpp b/Engine/Plugins/Interchange/Runtime/Source/Import/Private/Audio/InterchangeGenericAudioTranslator.cpp
new file mode 100644
--- /dev/null
+++ b/Engine/Plugins/Interchange/Runtime/Source/Import/Private/Audio/InterchangeGenericAudioTranslator.cpp
@@ -0,0 +1,106 @@
+// Copyright Epic Games, Inc. All Rights Reserved.
+#include "Audio/InterchangeGenericAudioTranslator.h"
+#include "InterchangeAudioSoundWaveNode.h"
+#include "HAL/IConsoleManager.h"
+
+#define LOCTEXT_NAMESPACE "InterchangeGenericAudioTranslator"
+
+static bool GInterchangeEnableAudio = false;
+static FAutoConsoleVariableRef CCvarInterchangeEnableSubstrateSupport(
+	TEXT("Interchange.FeatureFlags.Import.Audio"),
+	GInterchangeEnableAudio,
+	TEXT("Enable or disable support of Audio with Interchange. Disabled by default."),
+	ECVF_Default);
+
+TArray<FString> UInterchangeGenericAudioTranslator::GetSupportedFormats() const
+{
+	TArray<FString> Formats;
+
+#if WITH_EDITOR
+	if (GInterchangeEnableAudio || GIsAutomationTesting)
+	{
+#if WITH_SNDFILE_IO
+		Formats.Add(TEXT("aif;Audio Interchange File"));
+		Formats.Add(TEXT("aiff;Audio Interchange File Format"));
+		Formats.Add(TEXT("ogg;OGG Vorbis bitstream format "));
+		Formats.Add(TEXT("flac;Free Lossless Audio Codec"));
+		Formats.Add(TEXT("opus;OGG OPUS bitstream format"));
+		Formats.Add(TEXT("mp3;MPEG Layer 3 Audio"));
+#endif // WITH_SNDFILE_IO
+
+		Formats.Add(TEXT("wav;Wave Audio File"));
+	}
+#endif
+	return Formats;
+}
+
+EInterchangeTranslatorAssetType UInterchangeGenericAudioTranslator::GetSupportedAssetTypes() const
+{
+	return EInterchangeTranslatorAssetType::Sounds;
+}
+
+bool UInterchangeGenericAudioTranslator::Translate(UInterchangeBaseNodeContainer& BaseNodeContainer) const
+{
+	FString Filename = SourceData->GetFilename();
+	FPaths::NormalizeFilename(Filename);
+	if (!FPaths::FileExists(Filename))
+	{
+		LogError(
+			FText::Format(
+				LOCTEXT("InterchangeGenericAudioTranslator_Translate_FileNotFound", "Filepath {0} doesn't exist"), 
+				FText::FromString(Filename)
+			)
+		);
+
+		return false;
+	}
+
+	UInterchangeAudioSoundWaveNode* SoundWaveNode = UInterchangeAudioSoundWaveNode::Create(BaseNodeContainer, FPaths::GetBaseFilename(Filename));
+	SoundWaveNode->SetPayloadKey(Filename);
+
+	return true;
+}
+
+void UInterchangeGenericAudioTranslator::LogError(FText&& ErrorText) const
+{
+	if (UInterchangeResultError_Generic* Result = AddMessage<UInterchangeResultError_Generic>())
+	{
+		Result->AssetType = GetClass();
+
+		if (SourceData)
+		{
+			const FString Filename = SourceData->GetFilename();
+			Result->SourceAssetName = Filename;
+			Result->InterchangeKey = FPaths::GetBaseFilename(Filename);
+		}
+		else
+		{
+			Result->InterchangeKey = TEXT("Undefined");
+		}
+
+		Result->Text = MoveTemp(ErrorText);
+	}
+}
+
+void UInterchangeGenericAudioTranslator::LogWarning(FText&& WarningText) const
+{
+	if (UInterchangeResultWarning_Generic* Result = AddMessage<UInterchangeResultWarning_Generic>())
+	{
+		Result->AssetType = GetClass();
+
+		if (SourceData)
+		{
+			const FString Filename = SourceData->GetFilename();
+			Result->SourceAssetName = Filename;
+			Result->InterchangeKey = FPaths::GetBaseFilename(Filename);
+		}
+		else
+		{
+			Result->InterchangeKey = TEXT("Undefined");
+		}
+
+		Result->Text = MoveTemp(WarningText);
+	}
+}
+
+#undef LOCTEXT_NAMESPACE
\ No newline at end of file
diff --git a/Engine/Plugins/Interchange/Runtime/Source/Import/Public/Audio/InterchangeAudioPayloadData.h b/Engine/Plugins/Interchange/Runtime/Source/Import/Public/Audio/InterchangeAudioPayloadData.h
new file mode 100644
--- /dev/null
+++ b/Engine/Plugins/Interchange/Runtime/Source/Import/Public/Audio/InterchangeAudioPayloadData.h
@@ -0,0 +1,35 @@
+// Copyright Epic Games, Inc. All Rights Reserved.
+#pragma once
+#include "Containers/Array.h"
+#include "Misc/Optional.h"
+#include "Audio.h"
+
+#define UE_API INTERCHANGEIMPORT_API
+
+namespace UE::Interchange
+{
+	struct FSoundWaveFactoryInfo
+	{
+		int32 ChannelCount = 0;
+		int32 SizeOfSample = 0;
+		int32 NumSamples = 0;
+		int32 NumFrames = 0;
+		int32 SamplesPerSec = 0;
+		int32 SampleDataOffset = -1;
+	};
+
+	struct FInterchangeAudioPayloadData
+	{
+		FSoundWaveFactoryInfo FactoryInfo;
+		FWaveModInfo WaveModInfo;
+		TArray<uint8> Buffer;
+	};
+}
+
+struct FInterchangeAudioPayloadDataUtils
+{
+public:
+	static UE_API TOptional<UE::Interchange::FInterchangeAudioPayloadData> GetAudioPayloadFromSourceFileKey(const FString& PayloadSourceFileKey);
+};
+
+#undef UE_API
\ No newline at end of file
diff --git a/Engine/Plugins/Interchange/Runtime/Source/Import/Public/Audio/InterchangeAudioPayloadInterface.h b/Engine/Plugins/Interchange/Runtime/Source/Import/Public/Audio/InterchangeAudioPayloadInterface.h
new file mode 100644
--- /dev/null
+++ b/Engine/Plugins/Interchange/Runtime/Source/Import/Public/Audio/InterchangeAudioPayloadInterface.h
@@ -0,0 +1,40 @@
+// Copyright Epic Games, Inc. All Rights Reserved.
+#pragma once
+
+#include "Misc/Optional.h"
+#include "UObject/Object.h"
+#include "UObject/ObjectMacros.h"
+#include "InterchangeAudioPayloadData.h"
+
+#include "InterchangeAudioPayloadInterface.generated.h"
+
+#define UE_API INTERCHANGEIMPORT_API
+
+UINTERFACE(MinimalAPI)
+class UInterchangeAudioPayloadInterface : public UInterface
+{
+	GENERATED_BODY()
+};
+
+/**
+ * Audio Payload interface. Derive from it if your payload can import Audio
+ */
+class IInterchangeAudioPayloadInterface
+{
+	GENERATED_BODY()
+public:
+	/**
+	 * Once the translation is done, the import process need a way to retrieve payload data.
+	 * This payload will be use by the factories to create the asset.
+	 *
+	 * @param PayloadKey - The key to retrieve the a particular payload contain into the specified source data.
+	 * @return a PayloadData containing the import sound wave data, this maybe converted data if the format is other than '.wav'. The TOptional will not be set if there is an error.
+	 */
+	virtual TOptional< UE::Interchange::FInterchangeAudioPayloadData > GetAudioPayloadData(const FString& PayloadSourceFileKey) const 
+	{
+		return FInterchangeAudioPayloadDataUtils::GetAudioPayloadFromSourceFileKey(PayloadSourceFileKey);
+	}
+};
+
+
+#undef UE_API
\ No newline at end of file
diff --git a/Engine/Plugins/Interchange/Runtime/Source/Import/Public/Audio/InterchangeAudioSoundWaveFactory.h b/Engine/Plugins/Interchange/Runtime/Source/Import/Public/Audio/InterchangeAudioSoundWaveFactory.h
new file mode 100644
--- /dev/null
+++ b/Engine/Plugins/Interchange/Runtime/Source/Import/Public/Audio/InterchangeAudioSoundWaveFactory.h
@@ -0,0 +1,50 @@
+﻿// Copyright Epic Games, Inc. All Rights Reserved.
+
+#pragma once
+
+#include "InterchangeFactoryBase.h"
+#include "InterchangeAudioPayloadData.h"
+#include "InterchangeAudioSoundWaveFactory.generated.h"
+
+#define UE_API INTERCHANGEIMPORT_API
+
+class UAudioComponent;
+class USoundWave;
+class FWaveModInfo;
+
+UCLASS(BlueprintType)
+class UE_API UInterchangeAudioSoundWaveFactory : public UInterchangeFactoryBase
+{
+	GENERATED_BODY()
+
+public:
+	/** UInterchangeFactoryBase interface */
+	virtual UClass* GetFactoryClass() const override;
+	virtual EInterchangeFactoryAssetType GetFactoryAssetType() override { return EInterchangeFactoryAssetType::Sounds; }
+	virtual FImportAssetResult BeginImportAsset_GameThread(const FImportAssetObjectParams& Arguments) override;
+	virtual FImportAssetResult ImportAsset_Async(const FImportAssetObjectParams& Arguments) override;
+	virtual FImportAssetResult EndImportAsset_GameThread(const FImportAssetObjectParams& Arguments) override;
+
+	virtual void SetupObject_GameThread(const FSetupObjectParams& Arguments) override;
+	virtual void FinalizeObject_GameThread(const FSetupObjectParams& Arguments) override;
+
+	// To allow reimporting
+	virtual bool GetSourceFilenames(const UObject* Object, TArray<FString>& OutSourceFilenames) const override;
+	virtual bool SetSourceFilename(const UObject* Object, const FString& SourceFilename, int32 SourceIndex) const override;
+	virtual void BackupSourceData(const UObject* Object) const override;
+	virtual void ReinstateSourceData(const UObject* Object) const override;
+	virtual void ClearBackupSourceData(const UObject* Object) const override;
+
+private:
+	void SetUpSoundWaveFromPayload(const FImportAssetObjectParams& Arguments, USoundWave*& SoundWaveAsset, const UE::Interchange::FInterchangeAudioPayloadData& AudioPayloadData, FImportAssetResult& ImportAssetResult);
+	bool SetUpMultichannelSoundWave(const FImportAssetObjectParams& Arguments, USoundWave*& SoundWaveAsset, const UE::Interchange::FInterchangeAudioPayloadData& AudioPayloadData, FImportAssetResult& ImportAssetResult);
+	
+	void StopComponentsUsingImportedSound(USoundWave* SoundWaveAsset);
+	
+	void LogError(const FImportAssetObjectParams& Arguments, FImportAssetResult& ImportAssetResult, const FText& ErrorText);
+
+	UPROPERTY(Transient)
+	TArray<TObjectPtr<UAudioComponent>> ComponentsToRestart;
+};
+
+#undef UE_API
\ No newline at end of file
diff --git a/Engine/Plugins/Interchange/Runtime/Source/Import/Public/Audio/InterchangeGenericAudioTranslator.h b/Engine/Plugins/Interchange/Runtime/Source/Import/Public/Audio/InterchangeGenericAudioTranslator.h
new file mode 100644
--- /dev/null
+++ b/Engine/Plugins/Interchange/Runtime/Source/Import/Public/Audio/InterchangeGenericAudioTranslator.h
@@ -0,0 +1,29 @@
+// Copyright Epic Games, Inc. All Rights Reserved.
+#pragma once
+#include "InterchangeTranslatorBase.h"
+#include "Audio/InterchangeAudioPayloadInterface.h"
+
+#include "InterchangeGenericAudioTranslator.generated.h"
+
+#define UE_API INTERCHANGEIMPORT_API
+
+UCLASS(BlueprintType)
+class UE_API UInterchangeGenericAudioTranslator
+	: public UInterchangeTranslatorBase
+	, public IInterchangeAudioPayloadInterface
+{
+	GENERATED_BODY()
+
+public:
+	/** Begin UInterchangeTranslatorBase API */
+	virtual TArray<FString> GetSupportedFormats() const override;
+	virtual EInterchangeTranslatorAssetType GetSupportedAssetTypes() const override;
+	virtual bool Translate(UInterchangeBaseNodeContainer& BaseNodeContainer) const override;
+	/** End UInterchangeTranslatorBase API */
+
+private:
+	void LogError(FText&& ErrorText)  const;
+	void LogWarning(FText&& WarningText) const;
+};
+
+#undef UE_API
\ No newline at end of file
diff --git a/Engine/Plugins/Interchange/Runtime/Source/Nodes/Private/InterchangeAudioSoundWaveNode.cpp b/Engine/Plugins/Interchange/Runtime/Source/Nodes/Private/InterchangeAudioSoundWaveNode.cpp
new file mode 100644
--- /dev/null
+++ b/Engine/Plugins/Interchange/Runtime/Source/Nodes/Private/InterchangeAudioSoundWaveNode.cpp
@@ -0,0 +1,47 @@
+// Copyright Epic Games, Inc. All Rights Reserved.
+
+#include "InterchangeAudioSoundWaveNode.h"
+
+#include "Misc/Paths.h"
+#include "Nodes/InterchangeBaseNodeContainer.h"
+
+#include UE_INLINE_GENERATED_CPP_BY_NAME(InterchangeAudioSoundWaveNode)
+
+namespace UE::Interchange
+{
+	const FAttributeKey& FSoundWaveNodeStaticData::PayloadSourceFileKey()
+	{
+		static FAttributeKey AttributeKey(TEXT("__PayloadSourceFile__"));
+		return AttributeKey;
+	}
+}
+
+UInterchangeAudioSoundWaveNode* UInterchangeAudioSoundWaveNode::Create(UInterchangeBaseNodeContainer& NodeContainer, const FString& NodeName)
+{
+	UInterchangeAudioSoundWaveNode* SoundWaveNode = NewObject<UInterchangeAudioSoundWaveNode>();
+	check(SoundWaveNode);
+
+	const FString SoundWaveNodeUid = MakeNodeUid(NodeName);
+	NodeContainer.SetupNode(SoundWaveNode, SoundWaveNodeUid, NodeName, EInterchangeNodeContainerType::TranslatedAsset);
+	
+	return SoundWaveNode;
+}
+
+
+const TOptional<FString> UInterchangeAudioSoundWaveNode::GetPayloadKey() const
+{
+	using namespace UE::Interchange;
+	FString OutAttributeValue;
+	if (InterchangePrivateNodeBase::GetCustomAttribute(*Attributes, FSoundWaveNodeStaticData::PayloadSourceFileKey(), TEXT("UInterchangeAudioSoundWaveNode.GetPayloadKey"), OutAttributeValue))
+	{
+		return TOptional<FString>(MoveTemp(OutAttributeValue));
+	}
+
+	return TOptional<FString>();
+}
+
+void UInterchangeAudioSoundWaveNode::SetPayloadKey(const FString& PayloadKey)
+{
+	using namespace UE::Interchange;
+	InterchangePrivateNodeBase::SetCustomAttribute(*Attributes, FSoundWaveNodeStaticData::PayloadSourceFileKey(), TEXT("UInterchangeAudioSoundWaveNode.SetPayloadKey"), PayloadKey);
+}
\ No newline at end of file
diff --git a/Engine/Plugins/Interchange/Runtime/Source/Nodes/Public/InterchangeAudioSoundWaveNode.h b/Engine/Plugins/Interchange/Runtime/Source/Nodes/Public/InterchangeAudioSoundWaveNode.h
new file mode 100644
--- /dev/null
+++ b/Engine/Plugins/Interchange/Runtime/Source/Nodes/Public/InterchangeAudioSoundWaveNode.h
@@ -0,0 +1,45 @@
+// Copyright Epic Games, Inc. All Rights Reserved.
+
+#pragma once
+
+#include "Nodes/InterchangeBaseNode.h"
+
+#include "InterchangeAudioSoundWaveNode.generated.h"
+
+#define UE_API INTERCHANGENODES_API
+
+class UInterchangeBaseNodeContainer;
+
+namespace UE::Interchange
+{
+	struct FSoundWaveNodeStaticData : FBaseNodeStaticData
+	{
+		static UE_API const FAttributeKey& PayloadSourceFileKey();
+	};
+}
+
+UCLASS(BlueprintType)
+class UE_API UInterchangeAudioSoundWaveNode : public UInterchangeBaseNode
+{
+	GENERATED_BODY()
+
+public:
+	static UInterchangeAudioSoundWaveNode* Create(UInterchangeBaseNodeContainer& NodeContainer, const FString& NodeName);
+
+	virtual FString GetTypeName() const override
+	{
+		return TEXT("SoundWaveNode");
+	}
+
+	static FString MakeNodeUid(const FStringView NodeName)
+	{
+		return TEXT("SoundWaves") + FString(HierarchySeparator) + NodeName;
+	}
+
+	const TOptional<FString> GetPayloadKey() const;
+
+	void SetPayloadKey(const FString& PayloadKey);
+	
+};
+
+#undef UE_API
\ No newline at end of file
diff --git a/Engine/Plugins/Interchange/Runtime/Source/Pipelines/Private/InterchangeGenericAudioPipeline.cpp b/Engine/Plugins/Interchange/Runtime/Source/Pipelines/Private/InterchangeGenericAudioPipeline.cpp
new file mode 100644
--- /dev/null
+++ b/Engine/Plugins/Interchange/Runtime/Source/Pipelines/Private/InterchangeGenericAudioPipeline.cpp
@@ -0,0 +1,143 @@
+// Copyright Epic Games, Inc. All Rights Reserved.
+
+#include "InterchangeGenericAudioPipeline.h"
+
+#include "InterchangeAudioSoundWaveFactoryNode.h"
+#include "InterchangeAudioSoundWaveNode.h"
+
+#include "InterchangePipelineHelper.h"
+#include "Nodes/InterchangeSourceNode.h"
+
+#include UE_INLINE_GENERATED_CPP_BY_NAME(InterchangeGenericAudioPipeline)
+
+#define LOCTEXT_NAMESPACE "InterchangeGenericAudioPipeline"
+
+FString UInterchangeGenericAudioPipeline::GetPipelineCategory(UClass* AssetClass)
+{
+	return TEXT("Audio");
+}
+
+void UInterchangeGenericAudioPipeline::ExecutePipeline(UInterchangeBaseNodeContainer* InBaseNodeContainer, const TArray<UInterchangeSourceData*>& InSourceDatas, const FString& ContentBasePath)
+{
+	if (!InBaseNodeContainer)
+	{
+		UInterchangeResultError_Generic* ErrorMessage = AddMessage<UInterchangeResultError_Generic>();
+		check(ErrorMessage);
+		ErrorMessage->Text = LOCTEXT("InterchangeAudio_ExecutePipeline_NullContainer", "Cannot execute pre-import pipeline because InBaseNodeContainer is null.");
+		return;
+	}
+
+	BaseNodeContainer = InBaseNodeContainer;
+
+	BaseNodeContainer->IterateNodes([this](const FString& NodeUid, UInterchangeBaseNode* Node)
+	{
+		if (Node->GetNodeContainerType() == EInterchangeNodeContainerType::TranslatedAsset)
+		{
+			if (UInterchangeAudioSoundWaveNode* SoundWaveNode = Cast<UInterchangeAudioSoundWaveNode>(Node))
+			{
+				SoundWaveNodes.Add(SoundWaveNode);
+			}
+		}
+	});
+
+	if (bImportSounds)
+	{
+		for (const UInterchangeAudioSoundWaveNode* SoundWaveNode : SoundWaveNodes)
+		{
+			CreateSoundWaveFactoryNode(SoundWaveNode);
+		}
+	}
+}
+
+UInterchangeAudioSoundWaveFactoryNode* UInterchangeGenericAudioPipeline::CreateSoundWaveFactoryNode(const UInterchangeAudioSoundWaveNode* SoundWaveNode)
+{
+	if (!SoundWaveNode)
+	{
+		return nullptr;
+	}
+
+	const FString DisplayLabel = SoundWaveNode->GetDisplayLabel();
+	const FString NodeUid = UInterchangeFactoryBaseNode::BuildFactoryNodeUid(SoundWaveNode->GetUniqueID());
+
+	UInterchangeAudioSoundWaveFactoryNode* SoundWaveFactoryNode = NewObject<UInterchangeAudioSoundWaveFactoryNode>(BaseNodeContainer);
+	if (!ensure(SoundWaveFactoryNode))
+	{
+		return nullptr;
+	}
+
+	BaseNodeContainer->SetupNode(SoundWaveFactoryNode, NodeUid, DisplayLabel, EInterchangeNodeContainerType::FactoryData);
+	UInterchangeSourceNode* SourceNode = UInterchangeSourceNode::FindOrCreateUniqueInstance(BaseNodeContainer);
+	UE::Interchange::PipelineHelper::FillSubPathFromSourceNode(SoundWaveFactoryNode, SourceNode);
+
+	SoundWaveFactoryNodes.Add(SoundWaveFactoryNode);
+
+	SoundWaveFactoryNode->AddTargetNodeUid(SoundWaveNode->GetUniqueID());
+	SoundWaveNode->AddTargetNodeUid(SoundWaveFactoryNode->GetUniqueID());
+
+	return SoundWaveFactoryNode;
+}
+
+void UInterchangeGenericAudioPipeline::AdjustSettingsForContext(const FInterchangePipelineContextParams& ContextParams)
+{
+	Super::AdjustSettingsForContext(ContextParams);
+#if WITH_EDITOR
+	TArray<FString> HideCategories;
+	const bool bIsObjectASound = ContextParams.ReimportAsset ? ContextParams.ReimportAsset.IsA(USoundWave::StaticClass()) : false;
+	if ((!bIsObjectASound && ContextParams.ContextType == EInterchangePipelineContext::AssetReimport)
+		|| ContextParams.ContextType == EInterchangePipelineContext::AssetCustomLODImport
+		|| ContextParams.ContextType == EInterchangePipelineContext::AssetCustomLODReimport
+		|| ContextParams.ContextType == EInterchangePipelineContext::AssetAlternateSkinningImport
+		|| ContextParams.ContextType == EInterchangePipelineContext::AssetAlternateSkinningReimport
+		|| ContextParams.ContextType == EInterchangePipelineContext::AssetCustomMorphTargetImport
+		|| ContextParams.ContextType == EInterchangePipelineContext::AssetCustomMorphTargetReImport)
+	{
+		bImportSounds = false;
+		HideCategories.Add(UInterchangeGenericAudioPipeline::GetPipelineCategory(nullptr));
+	}
+
+	if (UInterchangePipelineBase* OuterMostPipeline = GetMostPipelineOuter())
+	{
+		for (const FString& HideCategoryName : HideCategories)
+		{
+			HidePropertiesOfCategory(OuterMostPipeline, this, HideCategoryName);
+		}
+	}
+#endif //WITH_EDITOR
+}
+
+#if WITH_EDITOR
+
+bool UInterchangeGenericAudioPipeline::IsPropertyChangeNeedRefresh(const FPropertyChangedEvent& PropertyChangedEvent) const
+{
+	if (PropertyChangedEvent.GetPropertyName() == GET_MEMBER_NAME_CHECKED(UInterchangeGenericAudioPipeline, bImportSounds))
+	{
+		return true;
+	}
+	return Super::IsPropertyChangeNeedRefresh(PropertyChangedEvent);
+}
+
+void UInterchangeGenericAudioPipeline::FilterPropertiesFromTranslatedData(UInterchangeBaseNodeContainer* InBaseNodeContainer)
+{
+	Super::FilterPropertiesFromTranslatedData(InBaseNodeContainer);
+
+	
+	TArray<FString> TempSoundWaveNodes;
+	InBaseNodeContainer->GetNodes(UInterchangeAudioSoundWaveNode::StaticClass(), TempSoundWaveNodes);
+	if (TempSoundWaveNodes.Num() == 0)
+	{
+		//Filter out all Audio properties
+		if (UInterchangePipelineBase* OuterMostPipeline = GetMostPipelineOuter())
+		{
+			HidePropertiesOfCategory(OuterMostPipeline, this, UInterchangeGenericAudioPipeline::GetPipelineCategory(nullptr));
+		}
+	}
+}
+
+void UInterchangeGenericAudioPipeline::GetSupportAssetClasses(TArray<UClass*>& PipelineSupportAssetClasses) const
+{
+	PipelineSupportAssetClasses.Add(USoundWave::StaticClass());
+}
+
+#endif //WITH_EDITOR
+
+#undef LOCTEXT_NAMESPACE
diff --git a/Engine/Plugins/Interchange/Runtime/Source/Pipelines/Public/InterchangeGenericAudioPipeline.h b/Engine/Plugins/Interchange/Runtime/Source/Pipelines/Public/InterchangeGenericAudioPipeline.h
new file mode 100644
--- /dev/null
+++ b/Engine/Plugins/Interchange/Runtime/Source/Pipelines/Public/InterchangeGenericAudioPipeline.h
@@ -0,0 +1,70 @@
+// Copyright Epic Games, Inc. All Rights Reserved.
+#pragma once
+
+#include "InterchangePipelineBase.h"
+#include "InterchangeAudioSoundWaveFactoryNode.h"
+#include "InterchangeAudioSoundWaveNode.h"
+
+#include "InterchangeGenericAudioPipeline.generated.h"
+
+/** Basic pipeline for importing sound wave assets. */
+UCLASS(BlueprintType, EditInlineNew)
+class INTERCHANGEPIPELINES_API UInterchangeGenericAudioPipeline : public UInterchangePipelineBase
+{
+	GENERATED_BODY()
+
+public:
+	static FString GetPipelineCategory(UClass* AssetClass);
+
+	/** The name of the pipeline that will be display in the import dialog. */
+	UPROPERTY(EditAnywhere, BlueprintReadWrite, Category = "Audio", meta = (StandAlonePipelineProperty = "True", PipelineInternalEditionData = "True"))
+	FString PipelineDisplayName;
+
+	/** If enabled, imports all sounds found in the source. */
+	UPROPERTY(EditAnywhere, BlueprintReadWrite, Category = "Audio")
+	bool bImportSounds = true;
+
+	virtual bool IsScripted() override
+	{
+		return false;
+	}
+
+public:
+	virtual void AdjustSettingsForContext(const FInterchangePipelineContextParams& ContextParams) override;
+	virtual void ExecutePipeline(UInterchangeBaseNodeContainer* InBaseNodeContainer, const TArray<UInterchangeSourceData*>& InSourceDatas, const FString& ContentBasePath) override;
+
+#if WITH_EDITOR
+	virtual bool IsPropertyChangeNeedRefresh(const FPropertyChangedEvent& PropertyChangedEvent) const override;
+	
+	virtual void FilterPropertiesFromTranslatedData(UInterchangeBaseNodeContainer* InBaseNodeContainer) override;
+
+	virtual void GetSupportAssetClasses(TArray<UClass*>& PipelineSupportAssetClasses) const override;
+#endif //WITH_EDITOR
+
+protected:
+	virtual bool CanExecuteOnAnyThread(EInterchangePipelineTask PipelineTask) override
+	{
+		return false;
+	}
+
+	/**
+	 * Creates a SoundWaveFactoryNode for a given SoundWaveNode, if one doesn't already exist.
+	 * Initializes the factory node, sets the custom directory path, and sets the target nodes
+	 * on both the factory and sound wave nodes.
+	 */
+	UFUNCTION(BlueprintCallable, Category = "Interchange")
+	UInterchangeAudioSoundWaveFactoryNode* CreateSoundWaveFactoryNode(const UInterchangeAudioSoundWaveNode* SoundWaveNode);
+
+private:
+	UPROPERTY(Transient)
+	TObjectPtr<UInterchangeBaseNodeContainer> BaseNodeContainer;
+
+	UPROPERTY(Transient)
+	TArray<const TObjectPtr<UInterchangeSourceData>> SourceDatas;
+
+	UPROPERTY(Transient)
+	TArray<TObjectPtr<UInterchangeAudioSoundWaveNode>> SoundWaveNodes;
+
+	UPROPERTY(Transient)
+	TArray<TObjectPtr<UInterchangeAudioSoundWaveFactoryNode>> SoundWaveFactoryNodes;
+};
\ No newline at end of file
diff --git a/Engine/Source/Runtime/Interchange/Core/Public/InterchangeFactoryBase.h b/Engine/Source/Runtime/Interchange/Core/Public/InterchangeFactoryBase.h
--- a/Engine/Source/Runtime/Interchange/Core/Public/InterchangeFactoryBase.h
+++ b/Engine/Source/Runtime/Interchange/Core/Public/InterchangeFactoryBase.h
@@ -28,7 +28,8 @@ enum class EInterchangeFactoryAssetType : uint8
 	Meshes,
 	Animations,
 	Physics,
-	Grooms
+	Grooms,
+	Sounds
 };
 
 /**
diff --git a/Engine/Source/Runtime/Interchange/Core/Public/InterchangeTranslatorBase.h b/Engine/Source/Runtime/Interchange/Core/Public/InterchangeTranslatorBase.h
--- a/Engine/Source/Runtime/Interchange/Core/Public/InterchangeTranslatorBase.h
+++ b/Engine/Source/Runtime/Interchange/Core/Public/InterchangeTranslatorBase.h
@@ -58,6 +58,7 @@ enum class EInterchangeTranslatorAssetType : uint8
 	Materials = 0x1 << 1,
 	Meshes = 0x1 << 2,
 	Animations = 0x1 << 3,
+	Sounds = 0x1 << 4,
 };
 ENUM_CLASS_FLAGS(EInterchangeTranslatorAssetType);
 
