[Interchange Audio] Fix console var static variable name.

#rnx

diff --git a/Engine/Plugins/Interchange/Runtime/Source/Import/Private/Audio/InterchangeGenericAudioTranslator.cpp b/Engine/Plugins/Interchange/Runtime/Source/Import/Private/Audio/InterchangeGenericAudioTranslator.cpp
--- a/Engine/Plugins/Interchange/Runtime/Source/Import/Private/Audio/InterchangeGenericAudioTranslator.cpp
+++ b/Engine/Plugins/Interchange/Runtime/Source/Import/Private/Audio/InterchangeGenericAudioTranslator.cpp
@@ -6,7 +6,7 @@
 #define LOCTEXT_NAMESPACE "InterchangeGenericAudioTranslator"
 
 static bool GInterchangeEnableAudio = false;
-static FAutoConsoleVariableRef CCvarInterchangeEnableSubstrateSupport(
+static FAutoConsoleVariableRef CCvarInterchangeEnableAudioSupport(
 	TEXT("Interchange.FeatureFlags.Import.Audio"),
 	GInterchangeEnableAudio,
 	TEXT("Enable or disable support of Audio with Interchange. Disabled by default."),
