[UBT][CppCompileWarnings] Engine compliance for wundef

diff --git a/Engine/Plugins/Interchange/Runtime/Source/Import/Private/Audio/InterchangeAudioPayloadData.cpp b/Engine/Plugins/Interchange/Runtime/Source/Import/Private/Audio/InterchangeAudioPayloadData.cpp
--- a/Engine/Plugins/Interchange/Runtime/Source/Import/Private/Audio/InterchangeAudioPayloadData.cpp
+++ b/Engine/Plugins/Interchange/Runtime/Source/Import/Private/Audio/InterchangeAudioPayloadData.cpp
@@ -89,7 +89,7 @@ TOptional<UE::Interchange::FInterchangeAudioPayloadData> FInterchangeAudioPayloa
 	if (bFileNeedsConversion || *WaveInfo.pBitsPerSample != 16 || !WaveInfo.IsFormatSupported())
 	{
 		TArray<uint8> ConvertedRawWaveData;
-#if WITH_SNDFILE_IO
+#if defined(WITH_SNDFILE_IO) && WITH_SNDFILE_IO
 		const bool bEnsureEqualNumSamples = !bFileNeedsConversion;
 		if (!Private::ConvertAudioFile(ReadBuffer, WaveInfo, ConvertedRawWaveData, ErrorMessage, bEnsureEqualNumSamples))
 		{
