Interchange MaterialX:
- Replace geompropvalue named "st" as texcoord node

diff --git a/Engine/Plugins/Interchange/Runtime/Source/Common/Public/MaterialX/InterchangeMaterialXDefinitions.h b/Engine/Plugins/Interchange/Runtime/Source/Common/Public/MaterialX/InterchangeMaterialXDefinitions.h
--- a/Engine/Plugins/Interchange/Runtime/Source/Common/Public/MaterialX/InterchangeMaterialXDefinitions.h
+++ b/Engine/Plugins/Interchange/Runtime/Source/Common/Public/MaterialX/InterchangeMaterialXDefinitions.h
@@ -794,6 +794,7 @@ MATERIALX_NAMESPACE_BEGIN
 		// Geometric nodes
 		static constexpr const char* Bitangent = "bitangent";
 		static constexpr const char* GeomColor = "geomcolor";
+		static constexpr const char* GeomPropValue = "geompropvalue";
 		static constexpr const char* Normal = "normal";
 		static constexpr const char* Position = "position";
 		static constexpr const char* Tangent = "tangent";
diff --git a/Engine/Plugins/Interchange/Runtime/Source/Import/Private/MaterialX/MaterialXUtils/MaterialXBase.cpp b/Engine/Plugins/Interchange/Runtime/Source/Import/Private/MaterialX/MaterialXUtils/MaterialXBase.cpp
--- a/Engine/Plugins/Interchange/Runtime/Source/Import/Private/MaterialX/MaterialXUtils/MaterialXBase.cpp
+++ b/Engine/Plugins/Interchange/Runtime/Source/Import/Private/MaterialX/MaterialXUtils/MaterialXBase.cpp
@@ -22,6 +22,7 @@ void FMaterialXBase::UpdateDocumentRecursively(MaterialX::GraphElementPtr Graph)
 	//For example the node tiledimage is just a succession of the nodes image, multiply, add, etc.
 	Graph->flattenSubgraphs(mx::EMPTY_STRING, [](mx::NodePtr Node){return FMaterialXManager::GetInstance().FilterNodeGraph(Node);} );
 
+	ConvertGeomPropValueNodesToTexCoord(Graph);
 	ConvertNeutralNodesToDot(Graph);
 
 	mx::DocumentPtr Document = Graph->getDocument();
@@ -154,6 +155,27 @@ void FMaterialXBase::ConvertNeutralNodesToDot(MaterialX::ElementPtr NodeGraph)
 	}
 }
 
+void FMaterialXBase::ConvertGeomPropValueNodesToTexCoord(MaterialX::GraphElementPtr Graph)
+{
+	const std::vector<mx::NodePtr> Nodes = Graph->getChildrenOfType<mx::Node>(mx::Category::GeomPropValue);
+	for (mx::NodePtr Node : Nodes)
+	{
+		if (Node->getType() != mx::Type::Vector2)
+		{
+			continue;
+		}
+
+		mx::InputPtr Input = Node->getInput("geomprop");
+		if (!(Input && Input->getValueString() == "st"))
+		{
+			continue;
+		}
+
+		// the default UV set is 0 (input "index" value=0)
+		Graph->changeChildCategory(Node, mx::Category::TexCoord);
+	}
+}
+
 MaterialX::NodePtr FMaterialXBase::CreateNode(MaterialX::ElementPtr NodeGraph, const char* NodeName, const char* Category, const TArray<FInputToCopy>& InputsToCopy, const TArray<FInputToCreate>& InputsToCreate)
 {
 	std::string UniqueNodeName{ NodeName + std::string{"_"} + Category };
diff --git a/Engine/Plugins/Interchange/Runtime/Source/Import/Public/MaterialX/MaterialXUtils/MaterialXBase.h b/Engine/Plugins/Interchange/Runtime/Source/Import/Public/MaterialX/MaterialXUtils/MaterialXBase.h
--- a/Engine/Plugins/Interchange/Runtime/Source/Import/Public/MaterialX/MaterialXUtils/MaterialXBase.h
+++ b/Engine/Plugins/Interchange/Runtime/Source/Import/Public/MaterialX/MaterialXUtils/MaterialXBase.h
@@ -134,6 +134,9 @@ protected:
 	 */
 	static void ConvertNeutralNodesToDot(MaterialX::ElementPtr NodeGraph);
 
+	/** Convert <geompropvalue> nodes that are vector2 with the exact input geomprop named "st" as <texcoord> nodes*/
+	static void ConvertGeomPropValueNodesToTexCoord(MaterialX::GraphElementPtr Graph);
+
 	/**
 	 * Helper function that create a MaterialX node, useful to insert a node between 2 nodes especially when an operation cannot be done in one pass
 	 * If the node already exists in the graph, this one is returned
