MaterialX: Remove unused code.

diff --git a/Engine/Plugins/Interchange/Runtime/Source/Import/Public/MaterialX/MaterialXUtils/MaterialXBase.h b/Engine/Plugins/Interchange/Runtime/Source/Import/Public/MaterialX/MaterialXUtils/MaterialXBase.h
--- a/Engine/Plugins/Interchange/Runtime/Source/Import/Public/MaterialX/MaterialXUtils/MaterialXBase.h
+++ b/Engine/Plugins/Interchange/Runtime/Source/Import/Public/MaterialX/MaterialXUtils/MaterialXBase.h
@@ -68,30 +68,6 @@ public:
 	/** Set the translator, only used for logging messages*/
 	void SetTranslator(const UInterchangeTranslatorBase* InTranslator);
 
-	/** Log a message either in form of a message passed to the translator or as a console log */
-	template<EMaterialXLog Verbosity, typename... Args>
-    void LogMessage(const FText& Text, Args... Arguments)
-	{
-		//if(Translator)
-		{
-			//if constexpr (Verbosity == EMaterialXLog::Error)
-			//{
-				UInterchangeResultError_Generic* Message = Translator->AddMessage<UInterchangeResultError_Generic>();
-				Message->Text = Text;
-			//}
-			//else
-			//{
-			//	UInterchangeResultWarning_Generic* Message = Translator->AddMessage<UInterchangeResultWarning_Generic>();
-			//	Message->Text = Text;
-			//}
-		}
-		//else
-		//{
-		//	//const TCHAR* Fmt = *(Text.ToString());
-		//	//UE_LOG(LogInterchangeImport, Warning, Fmt TEXT("%s"), Forward<Args>(Arguments)...);
-		//}
-	}
-
 protected:
 
 	/**
