MaterialX:
- use inline instead of static

diff --git a/Engine/Plugins/Interchange/Runtime/Source/Common/Public/MaterialX/InterchangeMaterialXDefinitions.h b/Engine/Plugins/Interchange/Runtime/Source/Common/Public/MaterialX/InterchangeMaterialXDefinitions.h
--- a/Engine/Plugins/Interchange/Runtime/Source/Common/Public/MaterialX/InterchangeMaterialXDefinitions.h
+++ b/Engine/Plugins/Interchange/Runtime/Source/Common/Public/MaterialX/InterchangeMaterialXDefinitions.h
@@ -127,15 +127,15 @@ namespace UE
 	{
 		namespace MaterialX
 		{
-			static constexpr uint8 IndexSurfaceShaders = 0;
-			static constexpr uint8 IndexBSDF = 1;
-			static constexpr uint8 IndexEDF = 2;
-			static constexpr uint8 IndexVDF = 3;
+			inline constexpr uint8 IndexSurfaceShaders = 0;
+			inline constexpr uint8 IndexBSDF = 1;
+			inline constexpr uint8 IndexEDF = 2;
+			inline constexpr uint8 IndexVDF = 3;
 
 			namespace Attributes
 			{
-				constexpr const TCHAR* EnumType = TEXT("MaterialXEnumType");
-				constexpr const TCHAR* EnumValue = TEXT("MaterialXEnumValue");
+				inline constexpr const TCHAR* EnumType = TEXT("MaterialXEnumType");
+				inline constexpr const TCHAR* EnumValue = TEXT("MaterialXEnumValue");
 			}
 		}
 	}
@@ -144,158 +144,158 @@ namespace UE
 	{
 		namespace Names
 		{
-			static constexpr const TCHAR* Abs = TEXT("Abs");
-			static constexpr const TCHAR* Add = TEXT("Add");
-			static constexpr const TCHAR* AppendVector = TEXT("AppendVector");
-			static constexpr const TCHAR* Arccosine = TEXT("Arccosine");
-			static constexpr const TCHAR* Arcsine = TEXT("Arcsine");
-			static constexpr const TCHAR* Arctangent2 = TEXT("Arctangent2");
-			static constexpr const TCHAR* BlackBody = TEXT("BlackBody");
-			static constexpr const TCHAR* CameraVectorWS = TEXT("CameraVectorWS");
-			static constexpr const TCHAR* Ceil = TEXT("Ceil");
-			static constexpr const TCHAR* Clamp = TEXT("Clamp");
-			static constexpr const TCHAR* Contrast = TEXT("MaterialXContrast");
-			static constexpr const TCHAR* Cosine = TEXT("Cosine");
-			static constexpr const TCHAR* Crossproduct = TEXT("Crossproduct");
-			static constexpr const TCHAR* Desaturation = TEXT("Desaturation");
-			static constexpr const TCHAR* Divide = TEXT("Divide");
-			static constexpr const TCHAR* Distance = TEXT("Distance");
-			static constexpr const TCHAR* Dotproduct = TEXT("Dotproduct");
-			static constexpr const TCHAR* Exponential = TEXT("Exponential");
-			static constexpr const TCHAR* Floor = TEXT("Floor");
-			static constexpr const TCHAR* Frac = TEXT("Frac");
-			static constexpr const TCHAR* HsvToRgb = TEXT("HsvToRgb");
-			static constexpr const TCHAR* If = TEXT("If");
-			static constexpr const TCHAR* Logarithm = TEXT("Logarithm");
-			static constexpr const TCHAR* Length = TEXT("Length");
-			static constexpr const TCHAR* Lerp = TEXT("Lerp");
-			static constexpr const TCHAR* LocalPosition = TEXT("LocalPosition");
-			static constexpr const TCHAR* MaterialXAppend3Vector = TEXT("MaterialXAppend3Vector");
-			static constexpr const TCHAR* MaterialXAppend4Vector = TEXT("MaterialXAppend4Vector");
-			static constexpr const TCHAR* MaterialXBurn = TEXT("MaterialXBurn");
-			static constexpr const TCHAR* MaterialXDifference = TEXT("MaterialXDifference");
-			static constexpr const TCHAR* MaterialXDisjointover = TEXT("MaterialXDisjointover");
-			static constexpr const TCHAR* MaterialXDodge = TEXT("MaterialXDodge");
-			static constexpr const TCHAR* MaterialXFractal3D = TEXT("MaterialXFractal3D");
-			static constexpr const TCHAR* MaterialXIn = TEXT("MaterialXIn");
-			static constexpr const TCHAR* MaterialXLuminance = TEXT("MaterialXLuminance");
-			static constexpr const TCHAR* MaterialXMask = TEXT("MaterialXMask");
-			static constexpr const TCHAR* MaterialXMatte = TEXT("MaterialXMatte");
-			static constexpr const TCHAR* MaterialXMinus = TEXT("MaterialXMinus");
-			static constexpr const TCHAR* MaterialXMod = TEXT("MaterialXMod");
-			static constexpr const TCHAR* MaterialXOut = TEXT("MaterialXOut");
-			static constexpr const TCHAR* MaterialXOver = TEXT("MaterialXOver");
-			static constexpr const TCHAR* MaterialXOverlay = TEXT("MaterialXOverlay");
-			static constexpr const TCHAR* MaterialXPlus = TEXT("MaterialXPlus");
-			static constexpr const TCHAR* MaterialXPremult = TEXT("MaterialXPremult");
-			static constexpr const TCHAR* MaterialXRamp4 = TEXT("MaterialXRamp4");
-			static constexpr const TCHAR* MaterialXRampLeftRight = TEXT("MaterialXRampLeftRight");
-			static constexpr const TCHAR* MaterialXRampTopBottom = TEXT("MaterialXRampTopBottom");
-			static constexpr const TCHAR* MaterialXRange = TEXT("MaterialXRange");
-			static constexpr const TCHAR* MaterialXRemap = TEXT("MaterialXRemap");
-			static constexpr const TCHAR* MaterialXScreen = TEXT("MaterialXScreen");
-			static constexpr const TCHAR* MaterialXSplitLeftRight = TEXT("MaterialXSplitLeftRight");
-			static constexpr const TCHAR* MaterialXSplitTopBottom = TEXT("MaterialXSplitTopBottom");
-			static constexpr const TCHAR* MaterialXUnpremult = TEXT("MaterialXUnpremult");
-			static constexpr const TCHAR* Max = TEXT("Max");
-			static constexpr const TCHAR* Min = TEXT("Min");
-			static constexpr const TCHAR* Multiply = TEXT("Multiply");
-			static constexpr const TCHAR* Normalize = TEXT("Normalize");
-			static constexpr const TCHAR* Power = TEXT("Power");
-			static constexpr const TCHAR* RgbToHsv = TEXT("RgbToHsv");
-			static constexpr const TCHAR* RotateAboutAxis = TEXT("RotateAboutAxis");
-			static constexpr const TCHAR* Round = TEXT("Round");
-			static constexpr const TCHAR* Sign = TEXT("Sign");
-			static constexpr const TCHAR* Sine = TEXT("Sine");
-			static constexpr const TCHAR* SmoothStep = TEXT("SmoothStep");
-			static constexpr const TCHAR* SquareRoot = TEXT("SquareRoot");
-			static constexpr const TCHAR* SubstrateAdd = TEXT("SubstrateAdd");
-			static constexpr const TCHAR* SubstrateHorizontalMixing = TEXT("SubstrateHorizontalMixing");
-			static constexpr const TCHAR* SubstrateVerticalLayering = TEXT("SubstrateVerticalLayering");
-			static constexpr const TCHAR* SubstrateWeight = TEXT("SubstrateWeight");
-			static constexpr const TCHAR* Subtract = TEXT("Subtract");
-			static constexpr const TCHAR* Tangent = TEXT("Tangent");
-			static constexpr const TCHAR* Time = TEXT("Time");
-			static constexpr const TCHAR* Transform = TEXT("Transform");
-			static constexpr const TCHAR* TransformPosition = TEXT("TransformPosition");
-			static constexpr const TCHAR* VertexColor = TEXT("VertexColor");
-			static constexpr const TCHAR* VertexNormalWS = TEXT("VertexNormalWS");
-			static constexpr const TCHAR* VertexTangentWS = TEXT("VertexTangentWS");
-			static constexpr const TCHAR* WorldPosition = TEXT("WorldPosition");
+			inline constexpr const TCHAR* Abs = TEXT("Abs");
+			inline constexpr const TCHAR* Add = TEXT("Add");
+			inline constexpr const TCHAR* AppendVector = TEXT("AppendVector");
+			inline constexpr const TCHAR* Arccosine = TEXT("Arccosine");
+			inline constexpr const TCHAR* Arcsine = TEXT("Arcsine");
+			inline constexpr const TCHAR* Arctangent2 = TEXT("Arctangent2");
+			inline constexpr const TCHAR* BlackBody = TEXT("BlackBody");
+			inline constexpr const TCHAR* CameraVectorWS = TEXT("CameraVectorWS");
+			inline constexpr const TCHAR* Ceil = TEXT("Ceil");
+			inline constexpr const TCHAR* Clamp = TEXT("Clamp");
+			inline constexpr const TCHAR* Contrast = TEXT("MaterialXContrast");
+			inline constexpr const TCHAR* Cosine = TEXT("Cosine");
+			inline constexpr const TCHAR* Crossproduct = TEXT("Crossproduct");
+			inline constexpr const TCHAR* Desaturation = TEXT("Desaturation");
+			inline constexpr const TCHAR* Divide = TEXT("Divide");
+			inline constexpr const TCHAR* Distance = TEXT("Distance");
+			inline constexpr const TCHAR* Dotproduct = TEXT("Dotproduct");
+			inline constexpr const TCHAR* Exponential = TEXT("Exponential");
+			inline constexpr const TCHAR* Floor = TEXT("Floor");
+			inline constexpr const TCHAR* Frac = TEXT("Frac");
+			inline constexpr const TCHAR* HsvToRgb = TEXT("HsvToRgb");
+			inline constexpr const TCHAR* If = TEXT("If");
+			inline constexpr const TCHAR* Logarithm = TEXT("Logarithm");
+			inline constexpr const TCHAR* Length = TEXT("Length");
+			inline constexpr const TCHAR* Lerp = TEXT("Lerp");
+			inline constexpr const TCHAR* LocalPosition = TEXT("LocalPosition");
+			inline constexpr const TCHAR* MaterialXAppend3Vector = TEXT("MaterialXAppend3Vector");
+			inline constexpr const TCHAR* MaterialXAppend4Vector = TEXT("MaterialXAppend4Vector");
+			inline constexpr const TCHAR* MaterialXBurn = TEXT("MaterialXBurn");
+			inline constexpr const TCHAR* MaterialXDifference = TEXT("MaterialXDifference");
+			inline constexpr const TCHAR* MaterialXDisjointover = TEXT("MaterialXDisjointover");
+			inline constexpr const TCHAR* MaterialXDodge = TEXT("MaterialXDodge");
+			inline constexpr const TCHAR* MaterialXFractal3D = TEXT("MaterialXFractal3D");
+			inline constexpr const TCHAR* MaterialXIn = TEXT("MaterialXIn");
+			inline constexpr const TCHAR* MaterialXLuminance = TEXT("MaterialXLuminance");
+			inline constexpr const TCHAR* MaterialXMask = TEXT("MaterialXMask");
+			inline constexpr const TCHAR* MaterialXMatte = TEXT("MaterialXMatte");
+			inline constexpr const TCHAR* MaterialXMinus = TEXT("MaterialXMinus");
+			inline constexpr const TCHAR* MaterialXMod = TEXT("MaterialXMod");
+			inline constexpr const TCHAR* MaterialXOut = TEXT("MaterialXOut");
+			inline constexpr const TCHAR* MaterialXOver = TEXT("MaterialXOver");
+			inline constexpr const TCHAR* MaterialXOverlay = TEXT("MaterialXOverlay");
+			inline constexpr const TCHAR* MaterialXPlus = TEXT("MaterialXPlus");
+			inline constexpr const TCHAR* MaterialXPremult = TEXT("MaterialXPremult");
+			inline constexpr const TCHAR* MaterialXRamp4 = TEXT("MaterialXRamp4");
+			inline constexpr const TCHAR* MaterialXRampLeftRight = TEXT("MaterialXRampLeftRight");
+			inline constexpr const TCHAR* MaterialXRampTopBottom = TEXT("MaterialXRampTopBottom");
+			inline constexpr const TCHAR* MaterialXRange = TEXT("MaterialXRange");
+			inline constexpr const TCHAR* MaterialXRemap = TEXT("MaterialXRemap");
+			inline constexpr const TCHAR* MaterialXScreen = TEXT("MaterialXScreen");
+			inline constexpr const TCHAR* MaterialXSplitLeftRight = TEXT("MaterialXSplitLeftRight");
+			inline constexpr const TCHAR* MaterialXSplitTopBottom = TEXT("MaterialXSplitTopBottom");
+			inline constexpr const TCHAR* MaterialXUnpremult = TEXT("MaterialXUnpremult");
+			inline constexpr const TCHAR* Max = TEXT("Max");
+			inline constexpr const TCHAR* Min = TEXT("Min");
+			inline constexpr const TCHAR* Multiply = TEXT("Multiply");
+			inline constexpr const TCHAR* Normalize = TEXT("Normalize");
+			inline constexpr const TCHAR* Power = TEXT("Power");
+			inline constexpr const TCHAR* RgbToHsv = TEXT("RgbToHsv");
+			inline constexpr const TCHAR* RotateAboutAxis = TEXT("RotateAboutAxis");
+			inline constexpr const TCHAR* Round = TEXT("Round");
+			inline constexpr const TCHAR* Sign = TEXT("Sign");
+			inline constexpr const TCHAR* Sine = TEXT("Sine");
+			inline constexpr const TCHAR* SmoothStep = TEXT("SmoothStep");
+			inline constexpr const TCHAR* SquareRoot = TEXT("SquareRoot");
+			inline constexpr const TCHAR* SubstrateAdd = TEXT("SubstrateAdd");
+			inline constexpr const TCHAR* SubstrateHorizontalMixing = TEXT("SubstrateHorizontalMixing");
+			inline constexpr const TCHAR* SubstrateVerticalLayering = TEXT("SubstrateVerticalLayering");
+			inline constexpr const TCHAR* SubstrateWeight = TEXT("SubstrateWeight");
+			inline constexpr const TCHAR* Subtract = TEXT("Subtract");
+			inline constexpr const TCHAR* Tangent = TEXT("Tangent");
+			inline constexpr const TCHAR* Time = TEXT("Time");
+			inline constexpr const TCHAR* Transform = TEXT("Transform");
+			inline constexpr const TCHAR* TransformPosition = TEXT("TransformPosition");
+			inline constexpr const TCHAR* VertexColor = TEXT("VertexColor");
+			inline constexpr const TCHAR* VertexNormalWS = TEXT("VertexNormalWS");
+			inline constexpr const TCHAR* VertexTangentWS = TEXT("VertexTangentWS");
+			inline constexpr const TCHAR* WorldPosition = TEXT("WorldPosition");
 		}
 
 		namespace Inputs
 		{
-			static constexpr const TCHAR* A = TEXT("A");
-			static constexpr const TCHAR* AbsoluteWorldPosition = TEXT("Absolute World Position");
-			static constexpr const TCHAR* AEqualsB = TEXT("AEqualsB");
-			static constexpr const TCHAR* AGreaterThanB = TEXT("AGreaterThanB");
-			static constexpr const TCHAR* Alpha = TEXT("Alpha");
-			static constexpr const TCHAR* ALessThanB = TEXT("ALessThanB");
-			static constexpr const TCHAR* Amount= TEXT("Amount");
-			static constexpr const TCHAR* Amplitude = TEXT("Amplitude");
-			static constexpr const TCHAR* B = TEXT("B");
-			static constexpr const TCHAR* Background = TEXT("Background");
-			static constexpr const TCHAR* Base = TEXT("Base");
-			static constexpr const TCHAR* Bottom = TEXT("Bottom");
-			static constexpr const TCHAR* C = TEXT("C");
-			static constexpr const TCHAR* Clamp = TEXT("Clamp");
-			static constexpr const TCHAR* Center = TEXT("Center");
-			static constexpr const TCHAR* Coordinate = TEXT("Coordinate");
-			static constexpr const TCHAR* Coordinates = TEXT("Coordinates");
-			static constexpr const TCHAR* D = TEXT("D");
-			static constexpr const TCHAR* Default = TEXT("Default");
-			static constexpr const TCHAR* Diminish = TEXT("Diminish");
-			static constexpr const TCHAR* Exponent = TEXT("Exponent");
-			static constexpr const TCHAR* Factor = TEXT("Factor");
-			static constexpr const TCHAR* Flatness = TEXT("Flatness");
-			static constexpr const TCHAR* Float2 = TEXT("Float2");
-			static constexpr const TCHAR* Float3 = TEXT("Float3");
-			static constexpr const TCHAR* Float4 = TEXT("Float4");
-			static constexpr const TCHAR* Foreground = TEXT("Foreground");
-			static constexpr const TCHAR* Fraction = TEXT("Fraction");
-			static constexpr const TCHAR* Gamma = TEXT("Gamma");
-			static constexpr const TCHAR* in1 = TEXT("in1");
-			static constexpr const TCHAR* in2 = TEXT("in2");
-			static constexpr const TCHAR* in3 = TEXT("in3");
-			static constexpr const TCHAR* in4 = TEXT("in4");
-			static constexpr const TCHAR* in5 = TEXT("in5");
-			static constexpr const TCHAR* in6 = TEXT("in6");
-			static constexpr const TCHAR* in7 = TEXT("in7");
-			static constexpr const TCHAR* in8 = TEXT("in8");
-			static constexpr const TCHAR* in9 = TEXT("in9");
-			static constexpr const TCHAR* in10 = TEXT("in10");
-			static constexpr const TCHAR* Input = TEXT("Input");
-			static constexpr const TCHAR* InputLow = TEXT("InputLow");
-			static constexpr const TCHAR* InputHigh = TEXT("InputHigh");
-			static constexpr const TCHAR* Lacunarity = TEXT("Lacunarity");
-			static constexpr const TCHAR* LuminanceFactors = TEXT("LuminanceFactors");
-			static constexpr const TCHAR* Max = TEXT("Max");
-			static constexpr const TCHAR* Min = TEXT("Min");
-			static constexpr const TCHAR* Mix = TEXT("Mix");
-			static constexpr const TCHAR* Normal = TEXT("Normal");
-			static constexpr const TCHAR* NormalizedRotationAxis = TEXT("NormalizedRotationAxis");
-			static constexpr const TCHAR* Octaves = TEXT("Octaves");
-			static constexpr const TCHAR* Offset = TEXT("Offset");
-			static constexpr const TCHAR* Pivot = TEXT("Pivot");
-			static constexpr const TCHAR* Position = TEXT("Position");
-			static constexpr const TCHAR* RayDirection = TEXT("Ray Direction");
-			static constexpr const TCHAR* RefractiveIndexOrigin = TEXT("Refractive Index Origin");
-			static constexpr const TCHAR* RefractiveIndexTarget = TEXT("Refractive Index Target");
-			static constexpr const TCHAR* RotationAngle = TEXT("RotationAngle");
-			static constexpr const TCHAR* Scale = TEXT("Scale");
-			static constexpr const TCHAR* SurfaceNormal = TEXT("Surface Normal");
-			static constexpr const TCHAR* SwitchValue = TEXT("SwitchValue");
-			static constexpr const TCHAR* TargetLow = TEXT("TargetLow");
-			static constexpr const TCHAR* TargetHigh = TEXT("TargetHigh");
-			static constexpr const TCHAR* Temp = TEXT("Temp");
-			static constexpr const TCHAR* Time = TEXT("Time");
-			static constexpr const TCHAR* Top = TEXT("Top");
-			static constexpr const TCHAR* Value = TEXT("Value");
-			static constexpr const TCHAR* VectorInput = TEXT("VectorInput");
-			static constexpr const TCHAR* Weight = TEXT("Weight");
-			static constexpr const TCHAR* X = TEXT("X");
-			static constexpr const TCHAR* Y = TEXT("Y");
+			inline constexpr const TCHAR* A = TEXT("A");
+			inline constexpr const TCHAR* AbsoluteWorldPosition = TEXT("Absolute World Position");
+			inline constexpr const TCHAR* AEqualsB = TEXT("AEqualsB");
+			inline constexpr const TCHAR* AGreaterThanB = TEXT("AGreaterThanB");
+			inline constexpr const TCHAR* Alpha = TEXT("Alpha");
+			inline constexpr const TCHAR* ALessThanB = TEXT("ALessThanB");
+			inline constexpr const TCHAR* Amount= TEXT("Amount");
+			inline constexpr const TCHAR* Amplitude = TEXT("Amplitude");
+			inline constexpr const TCHAR* B = TEXT("B");
+			inline constexpr const TCHAR* Background = TEXT("Background");
+			inline constexpr const TCHAR* Base = TEXT("Base");
+			inline constexpr const TCHAR* Bottom = TEXT("Bottom");
+			inline constexpr const TCHAR* C = TEXT("C");
+			inline constexpr const TCHAR* Clamp = TEXT("Clamp");
+			inline constexpr const TCHAR* Center = TEXT("Center");
+			inline constexpr const TCHAR* Coordinate = TEXT("Coordinate");
+			inline constexpr const TCHAR* Coordinates = TEXT("Coordinates");
+			inline constexpr const TCHAR* D = TEXT("D");
+			inline constexpr const TCHAR* Default = TEXT("Default");
+			inline constexpr const TCHAR* Diminish = TEXT("Diminish");
+			inline constexpr const TCHAR* Exponent = TEXT("Exponent");
+			inline constexpr const TCHAR* Factor = TEXT("Factor");
+			inline constexpr const TCHAR* Flatness = TEXT("Flatness");
+			inline constexpr const TCHAR* Float2 = TEXT("Float2");
+			inline constexpr const TCHAR* Float3 = TEXT("Float3");
+			inline constexpr const TCHAR* Float4 = TEXT("Float4");
+			inline constexpr const TCHAR* Foreground = TEXT("Foreground");
+			inline constexpr const TCHAR* Fraction = TEXT("Fraction");
+			inline constexpr const TCHAR* Gamma = TEXT("Gamma");
+			inline constexpr const TCHAR* in1 = TEXT("in1");
+			inline constexpr const TCHAR* in2 = TEXT("in2");
+			inline constexpr const TCHAR* in3 = TEXT("in3");
+			inline constexpr const TCHAR* in4 = TEXT("in4");
+			inline constexpr const TCHAR* in5 = TEXT("in5");
+			inline constexpr const TCHAR* in6 = TEXT("in6");
+			inline constexpr const TCHAR* in7 = TEXT("in7");
+			inline constexpr const TCHAR* in8 = TEXT("in8");
+			inline constexpr const TCHAR* in9 = TEXT("in9");
+			inline constexpr const TCHAR* in10 = TEXT("in10");
+			inline constexpr const TCHAR* Input = TEXT("Input");
+			inline constexpr const TCHAR* InputLow = TEXT("InputLow");
+			inline constexpr const TCHAR* InputHigh = TEXT("InputHigh");
+			inline constexpr const TCHAR* Lacunarity = TEXT("Lacunarity");
+			inline constexpr const TCHAR* LuminanceFactors = TEXT("LuminanceFactors");
+			inline constexpr const TCHAR* Max = TEXT("Max");
+			inline constexpr const TCHAR* Min = TEXT("Min");
+			inline constexpr const TCHAR* Mix = TEXT("Mix");
+			inline constexpr const TCHAR* Normal = TEXT("Normal");
+			inline constexpr const TCHAR* NormalizedRotationAxis = TEXT("NormalizedRotationAxis");
+			inline constexpr const TCHAR* Octaves = TEXT("Octaves");
+			inline constexpr const TCHAR* Offset = TEXT("Offset");
+			inline constexpr const TCHAR* Pivot = TEXT("Pivot");
+			inline constexpr const TCHAR* Position = TEXT("Position");
+			inline constexpr const TCHAR* RayDirection = TEXT("Ray Direction");
+			inline constexpr const TCHAR* RefractiveIndexOrigin = TEXT("Refractive Index Origin");
+			inline constexpr const TCHAR* RefractiveIndexTarget = TEXT("Refractive Index Target");
+			inline constexpr const TCHAR* RotationAngle = TEXT("RotationAngle");
+			inline constexpr const TCHAR* Scale = TEXT("Scale");
+			inline constexpr const TCHAR* SurfaceNormal = TEXT("Surface Normal");
+			inline constexpr const TCHAR* SwitchValue = TEXT("SwitchValue");
+			inline constexpr const TCHAR* TargetLow = TEXT("TargetLow");
+			inline constexpr const TCHAR* TargetHigh = TEXT("TargetHigh");
+			inline constexpr const TCHAR* Temp = TEXT("Temp");
+			inline constexpr const TCHAR* Time = TEXT("Time");
+			inline constexpr const TCHAR* Top = TEXT("Top");
+			inline constexpr const TCHAR* Value = TEXT("Value");
+			inline constexpr const TCHAR* VectorInput = TEXT("VectorInput");
+			inline constexpr const TCHAR* Weight = TEXT("Weight");
+			inline constexpr const TCHAR* X = TEXT("X");
+			inline constexpr const TCHAR* Y = TEXT("Y");
 		}
 	}
 
@@ -303,42 +303,42 @@ namespace UE
 	{
 		namespace Path
 		{
-			constexpr const TCHAR* BreakOutFloat2Components = TEXT("/Engine/Functions/Engine_MaterialFunctions02/Utility/BreakOutFloat2Components.BreakOutFloat2Components");
-			constexpr const TCHAR* BreakOutFloat3Components = TEXT("/Engine/Functions/Engine_MaterialFunctions02/Utility/BreakOutFloat3Components.BreakOutFloat3Components");
-			constexpr const TCHAR* BreakOutFloat4Components = TEXT("/Engine/Functions/Engine_MaterialFunctions02/Utility/BreakOutFloat4Components.BreakOutFloat4Components");
-			constexpr const TCHAR* DitherMask= TEXT("/InterchangeAssets/Substrate/MF_DitherMask.MF_DitherMask");
-			constexpr const TCHAR* HeightToNormalSmooth = TEXT("/Engine/Functions/Engine_MaterialFunctions03/Procedurals/HeightToNormalSmooth.HeightToNormalSmooth");
-			constexpr const TCHAR* MxACEScgToRec709 = TEXT("/InterchangeAssets/Functions/MX_ACEScgToRec709.MX_ACEScgToRec709");
-			constexpr const TCHAR* MxAdobeRGBToRec709 = TEXT("/InterchangeAssets/Functions/MX_AdobeRGBToRec709.MX_AdobeRGBToRec709");
-			constexpr const TCHAR* MxArtisticIOR = TEXT("/InterchangeAssets/Functions/MX_Artistic_IOR.MX_Artistic_IOR");
-			constexpr const TCHAR* MxAnd = TEXT("/InterchangeAssets/Functions/MX_And.MX_And");
-			constexpr const TCHAR* MxCircle = TEXT("/InterchangeAssets/Functions/MX_Circle.MX_Circle");
-			constexpr const TCHAR* MxCheckerboard = TEXT("/InterchangeAssets/Functions/MX_Checkerboard.MX_Checkerboard");
-			constexpr const TCHAR* MxChiangHairAbsorptionFromColor = TEXT("/InterchangeAssets/Functions/MX_ChiangHairAbsorptionFromColor.MX_ChiangHairAbsorptionFromColor");
-			constexpr const TCHAR* MxChiangHairRoughness = TEXT("/InterchangeAssets/Functions/MX_ChiangHairRoughness.MX_ChiangHairRoughness");
-			constexpr const TCHAR* MxColorCorrect = TEXT("/InterchangeAssets/Functions/MX_ColorCorrect.MX_ColorCorrect");
-			constexpr const TCHAR* MxDeonHairAbsorptionFromMelanin = TEXT("/InterchangeAssets/Functions/MX_DeonHairAbsorptionFromMelanin.MX_DeonHairAbsorptionFromMelanin");
-			constexpr const TCHAR* MxGoochShade = TEXT("/InterchangeAssets/Functions/MX_GoochShade.MX_GoochShade");
-			constexpr const TCHAR* MxHsvAdjust = TEXT("/InterchangeAssets/Functions/MX_HSVAdjust.MX_HSVAdjust");
-			constexpr const TCHAR* MxLine = TEXT("/InterchangeAssets/Functions/MX_Line.MX_Line");
-			constexpr const TCHAR* MxLinearAdobeRGBToRec709 = TEXT("/InterchangeAssets/Functions/MX_LinearAdobeRGBToRec709.MX_LinearAdobeRGBToRec709");
-			constexpr const TCHAR* MxLinearDisplayP3ToRec709 = TEXT("/InterchangeAssets/Functions/MX_LinearDisplayP3ToRec709.MX_LinearDisplayP3ToRec709");
-			constexpr const TCHAR* MxOr = TEXT("/InterchangeAssets/Functions/MX_Or.MX_Or");
-			constexpr const TCHAR* MxNot = TEXT("/InterchangeAssets/Functions/MX_Not.MX_Not");
-			constexpr const TCHAR* MxPlace2D = TEXT("/InterchangeAssets/Functions/MX_Place2D.MX_Place2D");
-			constexpr const TCHAR* MxRamp = TEXT("/InterchangeAssets/Functions/MX_Ramp.MX_Ramp");
-			constexpr const TCHAR* MxRampGradient = TEXT("/InterchangeAssets/Functions/MX_RampGradient.MX_RampGradient");
-			constexpr const TCHAR* MxRandomFloat = TEXT("/InterchangeAssets/Functions/MX_RandomFloat.MX_RandomFloat");
-			constexpr const TCHAR* MxRandomColor = TEXT("/InterchangeAssets/Functions/MX_RandomColor.MX_RandomColor");
-			constexpr const TCHAR* MxRoughnesAnisotropy = TEXT("/InterchangeAssets/Functions/MX_Roughness_Anisotropy.MX_Roughness_Anisotropy");
-			constexpr const TCHAR* MxRoughnessDual = TEXT("/InterchangeAssets/Functions/MX_Roughness_Dual.MX_Roughness_Dual");
-			constexpr const TCHAR* MxSrgbDisplayP3ToRec709 = TEXT("/InterchangeAssets/Functions/MX_SrgbDisplayP3ToRec709.MX_SrgbDisplayP3ToRec709");
-			constexpr const TCHAR* MxSrgbToRec709 = TEXT("/InterchangeAssets/Functions/MX_SrgbToRec709.MX_SrgbToRec709");
-			constexpr const TCHAR* MxUnifiedNoise3D = TEXT("/InterchangeAssets/Functions/MX_UnifiedNoise3d.MX_UnifiedNoise3d");
-			constexpr const TCHAR* MxTiledCircles = TEXT("/InterchangeAssets/Functions/MX_TiledCircles.MX_TiledCircles");
-			constexpr const TCHAR* MxXor = TEXT("/InterchangeAssets/Functions/MX_Xor.MX_Xor");
-			constexpr const TCHAR* NormalFromHeightMap = TEXT("/Engine/Functions/Engine_MaterialFunctions03/Procedurals/NormalFromHeightmap.NormalFromHeightmap");
-			constexpr const TCHAR* Refract = TEXT("/Engine/Functions/Engine_MaterialFunctions01/Vectors/Refract.Refract");
+			inline constexpr const TCHAR* BreakOutFloat2Components = TEXT("/Engine/Functions/Engine_MaterialFunctions02/Utility/BreakOutFloat2Components.BreakOutFloat2Components");
+			inline constexpr const TCHAR* BreakOutFloat3Components = TEXT("/Engine/Functions/Engine_MaterialFunctions02/Utility/BreakOutFloat3Components.BreakOutFloat3Components");
+			inline constexpr const TCHAR* BreakOutFloat4Components = TEXT("/Engine/Functions/Engine_MaterialFunctions02/Utility/BreakOutFloat4Components.BreakOutFloat4Components");
+			inline constexpr const TCHAR* DitherMask= TEXT("/InterchangeAssets/Substrate/MF_DitherMask.MF_DitherMask");
+			inline constexpr const TCHAR* HeightToNormalSmooth = TEXT("/Engine/Functions/Engine_MaterialFunctions03/Procedurals/HeightToNormalSmooth.HeightToNormalSmooth");
+			inline constexpr const TCHAR* MxACEScgToRec709 = TEXT("/InterchangeAssets/Functions/MX_ACEScgToRec709.MX_ACEScgToRec709");
+			inline constexpr const TCHAR* MxAdobeRGBToRec709 = TEXT("/InterchangeAssets/Functions/MX_AdobeRGBToRec709.MX_AdobeRGBToRec709");
+			inline constexpr const TCHAR* MxArtisticIOR = TEXT("/InterchangeAssets/Functions/MX_Artistic_IOR.MX_Artistic_IOR");
+			inline constexpr const TCHAR* MxAnd = TEXT("/InterchangeAssets/Functions/MX_And.MX_And");
+			inline constexpr const TCHAR* MxCircle = TEXT("/InterchangeAssets/Functions/MX_Circle.MX_Circle");
+			inline constexpr const TCHAR* MxCheckerboard = TEXT("/InterchangeAssets/Functions/MX_Checkerboard.MX_Checkerboard");
+			inline constexpr const TCHAR* MxChiangHairAbsorptionFromColor = TEXT("/InterchangeAssets/Functions/MX_ChiangHairAbsorptionFromColor.MX_ChiangHairAbsorptionFromColor");
+			inline constexpr const TCHAR* MxChiangHairRoughness = TEXT("/InterchangeAssets/Functions/MX_ChiangHairRoughness.MX_ChiangHairRoughness");
+			inline constexpr const TCHAR* MxColorCorrect = TEXT("/InterchangeAssets/Functions/MX_ColorCorrect.MX_ColorCorrect");
+			inline constexpr const TCHAR* MxDeonHairAbsorptionFromMelanin = TEXT("/InterchangeAssets/Functions/MX_DeonHairAbsorptionFromMelanin.MX_DeonHairAbsorptionFromMelanin");
+			inline constexpr const TCHAR* MxGoochShade = TEXT("/InterchangeAssets/Functions/MX_GoochShade.MX_GoochShade");
+			inline constexpr const TCHAR* MxHsvAdjust = TEXT("/InterchangeAssets/Functions/MX_HSVAdjust.MX_HSVAdjust");
+			inline constexpr const TCHAR* MxLine = TEXT("/InterchangeAssets/Functions/MX_Line.MX_Line");
+			inline constexpr const TCHAR* MxLinearAdobeRGBToRec709 = TEXT("/InterchangeAssets/Functions/MX_LinearAdobeRGBToRec709.MX_LinearAdobeRGBToRec709");
+			inline constexpr const TCHAR* MxLinearDisplayP3ToRec709 = TEXT("/InterchangeAssets/Functions/MX_LinearDisplayP3ToRec709.MX_LinearDisplayP3ToRec709");
+			inline constexpr const TCHAR* MxOr = TEXT("/InterchangeAssets/Functions/MX_Or.MX_Or");
+			inline constexpr const TCHAR* MxNot = TEXT("/InterchangeAssets/Functions/MX_Not.MX_Not");
+			inline constexpr const TCHAR* MxPlace2D = TEXT("/InterchangeAssets/Functions/MX_Place2D.MX_Place2D");
+			inline constexpr const TCHAR* MxRamp = TEXT("/InterchangeAssets/Functions/MX_Ramp.MX_Ramp");
+			inline constexpr const TCHAR* MxRampGradient = TEXT("/InterchangeAssets/Functions/MX_RampGradient.MX_RampGradient");
+			inline constexpr const TCHAR* MxRandomFloat = TEXT("/InterchangeAssets/Functions/MX_RandomFloat.MX_RandomFloat");
+			inline constexpr const TCHAR* MxRandomColor = TEXT("/InterchangeAssets/Functions/MX_RandomColor.MX_RandomColor");
+			inline constexpr const TCHAR* MxRoughnesAnisotropy = TEXT("/InterchangeAssets/Functions/MX_Roughness_Anisotropy.MX_Roughness_Anisotropy");
+			inline constexpr const TCHAR* MxRoughnessDual = TEXT("/InterchangeAssets/Functions/MX_Roughness_Dual.MX_Roughness_Dual");
+			inline constexpr const TCHAR* MxSrgbDisplayP3ToRec709 = TEXT("/InterchangeAssets/Functions/MX_SrgbDisplayP3ToRec709.MX_SrgbDisplayP3ToRec709");
+			inline constexpr const TCHAR* MxSrgbToRec709 = TEXT("/InterchangeAssets/Functions/MX_SrgbToRec709.MX_SrgbToRec709");
+			inline constexpr const TCHAR* MxUnifiedNoise3D = TEXT("/InterchangeAssets/Functions/MX_UnifiedNoise3d.MX_UnifiedNoise3d");
+			inline constexpr const TCHAR* MxTiledCircles = TEXT("/InterchangeAssets/Functions/MX_TiledCircles.MX_TiledCircles");
+			inline constexpr const TCHAR* MxXor = TEXT("/InterchangeAssets/Functions/MX_Xor.MX_Xor");
+			inline constexpr const TCHAR* NormalFromHeightMap = TEXT("/Engine/Functions/Engine_MaterialFunctions03/Procedurals/NormalFromHeightmap.NormalFromHeightmap");
+			inline constexpr const TCHAR* Refract = TEXT("/Engine/Functions/Engine_MaterialFunctions01/Vectors/Refract.Refract");
 		}
 	}
 }
@@ -349,93 +349,93 @@ MATERIALX_NAMESPACE_BEGIN
 	{
 		namespace Input
 		{
-			static constexpr const char* BaseWeight = "base_weight";
-			static constexpr const char* BaseColor = "base_color";
-			static constexpr const char* BaseDiffuseRoughness = "base_diffuse_roughness";
-			static constexpr const char* BaseMetalness = "base_metalness";
-			static constexpr const char* SpecularWeight = "specular_weight";
-			static constexpr const char* SpecularColor = "specular_color";
-			static constexpr const char* SpecularRoughness = "specular_roughness";
-			static constexpr const char* SpecularIOR = "specular_ior";
-			static constexpr const char* SpecularRoughnessAnisotropy = "specular_roughness_anisotropy";
-			static constexpr const char* TransmissionWeight = "transmission_weight";
-			static constexpr const char* TransmissionColor = "transmission_color";
-			static constexpr const char* TransmissionDepth = "transmission_depth";
-			static constexpr const char* TransmissionScatter = "transmission_scatter";
-			static constexpr const char* TransmissionScatterAnisotropy= "transmission_scatter_anisotropy";
-			static constexpr const char* TransmissionDispersionScale = "transmission_dispersion_scale";
-			static constexpr const char* TransmissionDispersionAbbeNumber= "transmission_dispersion_abbe_number";
-			static constexpr const char* SubsurfaceWeight = "subsurface_weight";
-			static constexpr const char* SubsurfaceColor = "subsurface_color";
-			static constexpr const char* SubsurfaceRadius = "subsurface_radius";
-			static constexpr const char* SubsurfaceRadiusScale = "subsurface_radius_scale";
-			static constexpr const char* SubsurfaceScatterAnisotropy = "subsurface_scatter_anisotropy";
-			static constexpr const char* FuzzWeight = "fuzz_weight";
-			static constexpr const char* FuzzColor = "fuzz_color";
-			static constexpr const char* FuzzRoughness = "fuzz_roughness";
-			static constexpr const char* CoatWeight = "coat_weight";
-			static constexpr const char* CoatColor = "coat_color";
-			static constexpr const char* CoatRoughness = "coat_roughness";
-			static constexpr const char* CoatRoughnessAnisotropy = "coat_roughness_anisotropy";
-			static constexpr const char* CoatRotation = "coat_rotation";
-			static constexpr const char* CoatIOR = "coat_ior";
-			static constexpr const char* CoatDarkening = "coat_darkening";
-			static constexpr const char* ThinFilmWeight= "thin_film_weight";
-			static constexpr const char* ThinFilmThickness = "thin_film_thickness";
-			static constexpr const char* ThinFilmIOR = "thin_film_ior";
-			static constexpr const char* EmissionLuminance = "emission_luminance";
-			static constexpr const char* EmissionColor = "emission_color";
-			static constexpr const char* GeometryOpacity = "geometry_opacity";
-			static constexpr const char* GeometryThinWalled = "geometry_thin_walled";
-			static constexpr const char* GeometryNormal = "geometry_normal";
-			static constexpr const char* GeometryCoatNormal = "geometry_coat_normal";
-			static constexpr const char* GeometryTangent = "geometry_tangent";
-			static constexpr const char* GeometryCoatTangent = "geometry_coat_tangent";
+			inline constexpr const char* BaseWeight = "base_weight";
+			inline constexpr const char* BaseColor = "base_color";
+			inline constexpr const char* BaseDiffuseRoughness = "base_diffuse_roughness";
+			inline constexpr const char* BaseMetalness = "base_metalness";
+			inline constexpr const char* SpecularWeight = "specular_weight";
+			inline constexpr const char* SpecularColor = "specular_color";
+			inline constexpr const char* SpecularRoughness = "specular_roughness";
+			inline constexpr const char* SpecularIOR = "specular_ior";
+			inline constexpr const char* SpecularRoughnessAnisotropy = "specular_roughness_anisotropy";
+			inline constexpr const char* TransmissionWeight = "transmission_weight";
+			inline constexpr const char* TransmissionColor = "transmission_color";
+			inline constexpr const char* TransmissionDepth = "transmission_depth";
+			inline constexpr const char* TransmissionScatter = "transmission_scatter";
+			inline constexpr const char* TransmissionScatterAnisotropy= "transmission_scatter_anisotropy";
+			inline constexpr const char* TransmissionDispersionScale = "transmission_dispersion_scale";
+			inline constexpr const char* TransmissionDispersionAbbeNumber= "transmission_dispersion_abbe_number";
+			inline constexpr const char* SubsurfaceWeight = "subsurface_weight";
+			inline constexpr const char* SubsurfaceColor = "subsurface_color";
+			inline constexpr const char* SubsurfaceRadius = "subsurface_radius";
+			inline constexpr const char* SubsurfaceRadiusScale = "subsurface_radius_scale";
+			inline constexpr const char* SubsurfaceScatterAnisotropy = "subsurface_scatter_anisotropy";
+			inline constexpr const char* FuzzWeight = "fuzz_weight";
+			inline constexpr const char* FuzzColor = "fuzz_color";
+			inline constexpr const char* FuzzRoughness = "fuzz_roughness";
+			inline constexpr const char* CoatWeight = "coat_weight";
+			inline constexpr const char* CoatColor = "coat_color";
+			inline constexpr const char* CoatRoughness = "coat_roughness";
+			inline constexpr const char* CoatRoughnessAnisotropy = "coat_roughness_anisotropy";
+			inline constexpr const char* CoatRotation = "coat_rotation";
+			inline constexpr const char* CoatIOR = "coat_ior";
+			inline constexpr const char* CoatDarkening = "coat_darkening";
+			inline constexpr const char* ThinFilmWeight= "thin_film_weight";
+			inline constexpr const char* ThinFilmThickness = "thin_film_thickness";
+			inline constexpr const char* ThinFilmIOR = "thin_film_ior";
+			inline constexpr const char* EmissionLuminance = "emission_luminance";
+			inline constexpr const char* EmissionColor = "emission_color";
+			inline constexpr const char* GeometryOpacity = "geometry_opacity";
+			inline constexpr const char* GeometryThinWalled = "geometry_thin_walled";
+			inline constexpr const char* GeometryNormal = "geometry_normal";
+			inline constexpr const char* GeometryCoatNormal = "geometry_coat_normal";
+			inline constexpr const char* GeometryTangent = "geometry_tangent";
+			inline constexpr const char* GeometryCoatTangent = "geometry_coat_tangent";
 		}
 
 		namespace DefaultValue
 		{
-			static constexpr float BaseWeight = 1.f;
-			static constexpr FLinearColor BaseColor{ 0.8, 0.8, 0.8 };
-			static constexpr float BaseDiffuseRoughness = 0.f;
-			static constexpr float BaseMetalness = 0.f;
-			static constexpr float SpecularWeight = 1.f;
-			static constexpr FLinearColor SpecularColor{ 1, 1, 1 };
-			static constexpr float SpecularRoughness = 0.3f;
-			static constexpr float SpecularIOR = 1.5f;
-			static constexpr float SpecularRoughnessAnisotropy = 0.f;
-			static constexpr float TransmissionWeight = 0.f;
-			static constexpr FLinearColor TransmissionColor{ 1, 1, 1 };
-			static constexpr float TransmissionDepth = 0.f;
-			static constexpr FLinearColor TransmissionScatter{ 0, 0, 0 };
-			static constexpr float TransmissionScatterAnisotropy = 0.f;
-			static constexpr float TransmissionDispersionScale = 0.f;
-			static constexpr float TransmissionDispersionAbbeNumber = 20.f;
-			static constexpr float SubsurfaceWeight = 0.f;
-			static constexpr FLinearColor SubsurfaceColor{ 0.8, 0.8, 0.8 };
-			static constexpr float SubsurfaceRadius = 1.f;
-			static constexpr FLinearColor SubsurfaceRadiusScale{1, 0.5, 0.25};
-			static constexpr float SubsurfaceScatterAnisotropy = 0.f;
-			static constexpr float FuzzWeight = 0.f;
-			static constexpr FLinearColor FuzzColor{ 1, 1, 1 };
-			static constexpr float FuzzRoughness = 0.5f;
-			static constexpr float CoatWeight = 0.f;
-			static constexpr FLinearColor CoatColor{ 1, 1, 1 };
-			static constexpr float CoatRoughness = 0.f;
-			static constexpr float CoatRoughnessAnisotropy = 0.f;
-			static constexpr float CoatIOR = 1.6f;
-			static constexpr float CoatDarkening = 1.f;
-			static constexpr float ThinFilmWeight = 0.f;
-			static constexpr float ThinFilmThickness = 0.5f;
-			static constexpr float ThinFilmIOR = 1.4f;
-			static constexpr float EmissionLuminance = 0.f;
-			static constexpr FLinearColor EmissionColor{ 1, 1, 1 };
-			static constexpr float GeometryOpacity = 1.f;
-			static constexpr bool GeometryThinWalled = false;
-			static const FVector GeometryNormal{ 0, 0, 1 };
-			static const FVector GeometryCoatNormal{ 0, 0, 1 };
-			static const FVector GeometryTangent{ 0, 1, 0 };
-			static const FVector GeometryCoatTangent{ 0, 1, 0 };
+			inline constexpr float BaseWeight = 1.f;
+			inline constexpr FLinearColor BaseColor{ 0.8, 0.8, 0.8 };
+			inline constexpr float BaseDiffuseRoughness = 0.f;
+			inline constexpr float BaseMetalness = 0.f;
+			inline constexpr float SpecularWeight = 1.f;
+			inline constexpr FLinearColor SpecularColor{ 1, 1, 1 };
+			inline constexpr float SpecularRoughness = 0.3f;
+			inline constexpr float SpecularIOR = 1.5f;
+			inline constexpr float SpecularRoughnessAnisotropy = 0.f;
+			inline constexpr float TransmissionWeight = 0.f;
+			inline constexpr FLinearColor TransmissionColor{ 1, 1, 1 };
+			inline constexpr float TransmissionDepth = 0.f;
+			inline constexpr FLinearColor TransmissionScatter{ 0, 0, 0 };
+			inline constexpr float TransmissionScatterAnisotropy = 0.f;
+			inline constexpr float TransmissionDispersionScale = 0.f;
+			inline constexpr float TransmissionDispersionAbbeNumber = 20.f;
+			inline constexpr float SubsurfaceWeight = 0.f;
+			inline constexpr FLinearColor SubsurfaceColor{ 0.8, 0.8, 0.8 };
+			inline constexpr float SubsurfaceRadius = 1.f;
+			inline constexpr FLinearColor SubsurfaceRadiusScale{1, 0.5, 0.25};
+			inline constexpr float SubsurfaceScatterAnisotropy = 0.f;
+			inline constexpr float FuzzWeight = 0.f;
+			inline constexpr FLinearColor FuzzColor{ 1, 1, 1 };
+			inline constexpr float FuzzRoughness = 0.5f;
+			inline constexpr float CoatWeight = 0.f;
+			inline constexpr FLinearColor CoatColor{ 1, 1, 1 };
+			inline constexpr float CoatRoughness = 0.f;
+			inline constexpr float CoatRoughnessAnisotropy = 0.f;
+			inline constexpr float CoatIOR = 1.6f;
+			inline constexpr float CoatDarkening = 1.f;
+			inline constexpr float ThinFilmWeight = 0.f;
+			inline constexpr float ThinFilmThickness = 0.5f;
+			inline constexpr float ThinFilmIOR = 1.4f;
+			inline constexpr float EmissionLuminance = 0.f;
+			inline constexpr FLinearColor EmissionColor{ 1, 1, 1 };
+			inline constexpr float GeometryOpacity = 1.f;
+			inline constexpr bool GeometryThinWalled = false;
+			inline const FVector GeometryNormal{ 0, 0, 1 };
+			inline const FVector GeometryCoatNormal{ 0, 0, 1 };
+			inline const FVector GeometryTangent{ 0, 1, 0 };
+			inline const FVector GeometryCoatTangent{ 0, 1, 0 };
 		}
 	}
 
@@ -443,91 +443,91 @@ MATERIALX_NAMESPACE_BEGIN
 	{
 		namespace Input
 		{
-			static constexpr const char* Base = "base";
-			static constexpr const char* BaseColor = "base_color";
-			static constexpr const char* DiffuseRoughness = "diffuse_roughness";
-			static constexpr const char* Metalness = "metalness";
-			static constexpr const char* Specular = "specular";
-			static constexpr const char* SpecularColor = "specular_color";
-			static constexpr const char* SpecularRoughness = "specular_roughness";
-			static constexpr const char* SpecularIOR = "specular_IOR";
-			static constexpr const char* SpecularAnisotropy = "specular_anisotropy";
-			static constexpr const char* SpecularRotation = "specular_rotation";
-			static constexpr const char* Transmission = "transmission";
-			static constexpr const char* TransmissionColor = "transmission_color";
-			static constexpr const char* TransmissionDepth = "transmission_depth";
-			static constexpr const char* TransmissionScatter = "transmission_scatter";
-			static constexpr const char* TransmissionScatterAnisotropy = "transmission_scatter_anisotropy";
-			static constexpr const char* TransmissionDispersion = "transmission_dispersion";
-			static constexpr const char* TransmissionExtraRoughness = "transmission_extra_roughness";
-			static constexpr const char* Subsurface = "subsurface";
-			static constexpr const char* SubsurfaceColor = "subsurface_color";
-			static constexpr const char* SubsurfaceRadius = "subsurface_radius";
-			static constexpr const char* SubsurfaceScale = "subsurface_scale";
-			static constexpr const char* SubsurfaceAnisotropy = "subsurface_anisotropy";
-			static constexpr const char* Sheen = "sheen";
-			static constexpr const char* SheenColor = "sheen_color";
-			static constexpr const char* SheenRoughness = "sheen_roughness";
-			static constexpr const char* Coat = "coat";
-			static constexpr const char* CoatColor = "coat_color";
-			static constexpr const char* CoatRoughness = "coat_roughness";
-			static constexpr const char* CoatAnisotropy = "coat_anisotropy";
-			static constexpr const char* CoatRotation = "coat_rotation";
-			static constexpr const char* CoatIOR = "coat_IOR";
-			static constexpr const char* CoatNormal = "coat_normal";
-			static constexpr const char* CoatAffectColor = "coat_affect_color";
-			static constexpr const char* CoatAffectRoughness = "coat_affect_roughness";
-			static constexpr const char* ThinFilmThickness = "thin_film_thickness";
-			static constexpr const char* ThinFilmIOR = "thin_film_IOR";
-			static constexpr const char* Emission = "emission";
-			static constexpr const char* EmissionColor = "emission_color";
-			static constexpr const char* Opacity = "opacity";
-			static constexpr const char* ThinWalled = "thin_walled";
-			static constexpr const char* Normal = "normal";
-			static constexpr const char* Tangent = "tangent";
+			inline constexpr const char* Base = "base";
+			inline constexpr const char* BaseColor = "base_color";
+			inline constexpr const char* DiffuseRoughness = "diffuse_roughness";
+			inline constexpr const char* Metalness = "metalness";
+			inline constexpr const char* Specular = "specular";
+			inline constexpr const char* SpecularColor = "specular_color";
+			inline constexpr const char* SpecularRoughness = "specular_roughness";
+			inline constexpr const char* SpecularIOR = "specular_IOR";
+			inline constexpr const char* SpecularAnisotropy = "specular_anisotropy";
+			inline constexpr const char* SpecularRotation = "specular_rotation";
+			inline constexpr const char* Transmission = "transmission";
+			inline constexpr const char* TransmissionColor = "transmission_color";
+			inline constexpr const char* TransmissionDepth = "transmission_depth";
+			inline constexpr const char* TransmissionScatter = "transmission_scatter";
+			inline constexpr const char* TransmissionScatterAnisotropy = "transmission_scatter_anisotropy";
+			inline constexpr const char* TransmissionDispersion = "transmission_dispersion";
+			inline constexpr const char* TransmissionExtraRoughness = "transmission_extra_roughness";
+			inline constexpr const char* Subsurface = "subsurface";
+			inline constexpr const char* SubsurfaceColor = "subsurface_color";
+			inline constexpr const char* SubsurfaceRadius = "subsurface_radius";
+			inline constexpr const char* SubsurfaceScale = "subsurface_scale";
+			inline constexpr const char* SubsurfaceAnisotropy = "subsurface_anisotropy";
+			inline constexpr const char* Sheen = "sheen";
+			inline constexpr const char* SheenColor = "sheen_color";
+			inline constexpr const char* SheenRoughness = "sheen_roughness";
+			inline constexpr const char* Coat = "coat";
+			inline constexpr const char* CoatColor = "coat_color";
+			inline constexpr const char* CoatRoughness = "coat_roughness";
+			inline constexpr const char* CoatAnisotropy = "coat_anisotropy";
+			inline constexpr const char* CoatRotation = "coat_rotation";
+			inline constexpr const char* CoatIOR = "coat_IOR";
+			inline constexpr const char* CoatNormal = "coat_normal";
+			inline constexpr const char* CoatAffectColor = "coat_affect_color";
+			inline constexpr const char* CoatAffectRoughness = "coat_affect_roughness";
+			inline constexpr const char* ThinFilmThickness = "thin_film_thickness";
+			inline constexpr const char* ThinFilmIOR = "thin_film_IOR";
+			inline constexpr const char* Emission = "emission";
+			inline constexpr const char* EmissionColor = "emission_color";
+			inline constexpr const char* Opacity = "opacity";
+			inline constexpr const char* ThinWalled = "thin_walled";
+			inline constexpr const char* Normal = "normal";
+			inline constexpr const char* Tangent = "tangent";
 		}
 
 		namespace DefaultValue
 		{
-			static constexpr float Base = 1.f;
-			constexpr FLinearColor BaseColor{ 0.8, 0.8, 0.8 };
-			static constexpr float DiffuseRoughness = 0.f;
-			static constexpr float Metalness = 0.f;
-			static constexpr float Specular = 1.f;
-			constexpr FLinearColor SpecularColor{ 1, 1, 1 };
-			static constexpr float SpecularRoughness = 0.2f;
-			static constexpr float SpecularIOR = 1.5f;
-			static constexpr float SpecularAnisotropy = 0.f;
-			static constexpr float SpecularRotation = 0.f;
-			static constexpr float Transmission = 0.f;
-			static constexpr float TransmissionDepth = 0.f;
-			static constexpr float TransmissionScatterAnisotropy = 0.f;
-			static constexpr float TransmissionDispersion = 0.f;
-			static constexpr float TransmissionExtraRoughness = 0.f;
-			static constexpr float Subsurface = 0.f;
-			constexpr FLinearColor SubsurfaceColor{ 1, 1, 1 };
-			constexpr FLinearColor SubsurfaceRadius{ 1, 1, 1 };
-			static constexpr float SubsurfaceScale = 1.f;
-			static constexpr float SubsurfaceAnisotropy = 0.f;
-			static constexpr float Sheen = 0.f;
-			constexpr FLinearColor SheenColor{ 1, 1, 1 };
-			static constexpr float SheenRoughness = 0.3f;
-			static constexpr float Coat = 0.f;
-			constexpr FLinearColor CoatColor{ 1, 1, 1 };
-			static constexpr float CoatRoughness = 0.1f;
-			static constexpr float CoatAnisotropy = 0.f;
-			static constexpr float CoatRotation = 0.f;
-			static constexpr float CoatIOR = 1.5f;
-			static constexpr float CoatAffectColor = 0.f;
-			static constexpr float CoatAffectRoughness = 0.f;
-			static constexpr float ThinFilmThickness = 0.f;
-			static constexpr float ThinFilmIOR = 1.5f;
-			constexpr FLinearColor TransmissionColor{ 1, 1, 1 };
-			constexpr FLinearColor TransmissionScatter{ 0, 0, 0 };
-			static constexpr float Emission = 0.f;
-			constexpr FLinearColor EmissionColor{ 1, 1, 1 };
-			constexpr FLinearColor Opacity{ 1, 1, 1 };
-			constexpr bool ThinWalled = false;
+			inline constexpr float Base = 1.f;
+			inline constexpr FLinearColor BaseColor{ 0.8, 0.8, 0.8 };
+			inline constexpr float DiffuseRoughness = 0.f;
+			inline constexpr float Metalness = 0.f;
+			inline constexpr float Specular = 1.f;
+			inline constexpr FLinearColor SpecularColor{ 1, 1, 1 };
+			inline constexpr float SpecularRoughness = 0.2f;
+			inline constexpr float SpecularIOR = 1.5f;
+			inline constexpr float SpecularAnisotropy = 0.f;
+			inline constexpr float SpecularRotation = 0.f;
+			inline constexpr float Transmission = 0.f;
+			inline constexpr float TransmissionDepth = 0.f;
+			inline constexpr float TransmissionScatterAnisotropy = 0.f;
+			inline constexpr float TransmissionDispersion = 0.f;
+			inline constexpr float TransmissionExtraRoughness = 0.f;
+			inline constexpr float Subsurface = 0.f;
+			inline constexpr FLinearColor SubsurfaceColor{ 1, 1, 1 };
+			inline constexpr FLinearColor SubsurfaceRadius{ 1, 1, 1 };
+			inline constexpr float SubsurfaceScale = 1.f;
+			inline constexpr float SubsurfaceAnisotropy = 0.f;
+			inline constexpr float Sheen = 0.f;
+			inline constexpr FLinearColor SheenColor{ 1, 1, 1 };
+			inline constexpr float SheenRoughness = 0.3f;
+			inline constexpr float Coat = 0.f;
+			inline constexpr FLinearColor CoatColor{ 1, 1, 1 };
+			inline constexpr float CoatRoughness = 0.1f;
+			inline constexpr float CoatAnisotropy = 0.f;
+			inline constexpr float CoatRotation = 0.f;
+			inline constexpr float CoatIOR = 1.5f;
+			inline constexpr float CoatAffectColor = 0.f;
+			inline constexpr float CoatAffectRoughness = 0.f;
+			inline constexpr float ThinFilmThickness = 0.f;
+			inline constexpr float ThinFilmIOR = 1.5f;
+			inline constexpr FLinearColor TransmissionColor{ 1, 1, 1 };
+			inline constexpr FLinearColor TransmissionScatter{ 0, 0, 0 };
+			inline constexpr float Emission = 0.f;
+			inline constexpr FLinearColor EmissionColor{ 1, 1, 1 };
+			inline constexpr FLinearColor Opacity{ 1, 1, 1 };
+			inline constexpr bool ThinWalled = false;
 		}
 	}
 
@@ -535,26 +535,26 @@ MATERIALX_NAMESPACE_BEGIN
 	{
 		namespace Input
 		{
-			static constexpr const char* Emission = "emission";
-			static constexpr const char* EmissionColor = "emission_color";
-			static constexpr const char* Transmission = "transmission";
-			static constexpr const char* TransmissionColor = "transmission_color";
-			static constexpr const char* Opacity = "opacity";
+			inline constexpr const char* Emission = "emission";
+			inline constexpr const char* EmissionColor = "emission_color";
+			inline constexpr const char* Transmission = "transmission";
+			inline constexpr const char* TransmissionColor = "transmission_color";
+			inline constexpr const char* Opacity = "opacity";
 		}
 
 		namespace DefaultValue
 		{
 			namespace Float
 			{
-				constexpr float Emission = 1.f;
-				constexpr float Transmission = 0.f;
-				constexpr float Opacity = 1.f;
+				inline constexpr float Emission = 1.f;
+				inline constexpr float Transmission = 0.f;
+				inline constexpr float Opacity = 1.f;
 			}
 
 			namespace Color3
 			{
-				constexpr FLinearColor EmissionColor{ 1, 1, 1 };
-				constexpr FLinearColor TransmissionColor{ 1, 1, 1 };
+				inline constexpr FLinearColor EmissionColor{ 1, 1, 1 };
+				inline constexpr FLinearColor TransmissionColor{ 1, 1, 1 };
 			}
 		}
 	}
@@ -563,16 +563,16 @@ MATERIALX_NAMESPACE_BEGIN
 	{
 		namespace Input
 		{
-			static constexpr const char* Bsdf = "bsdf";
-			static constexpr const char* Edf = "edf";
-			static constexpr const char* Opacity = "opacity";
+			inline constexpr const char* Bsdf = "bsdf";
+			inline constexpr const char* Edf = "edf";
+			inline constexpr const char* Opacity = "opacity";
 		}
 
 		namespace DefaultValue
 		{
 			namespace Float
 			{
-				constexpr float Opacity = 1.f;
+				inline constexpr float Opacity = 1.f;
 			}
 		}
 	}
@@ -581,46 +581,46 @@ MATERIALX_NAMESPACE_BEGIN
 	{
 		namespace Input
 		{
-			static constexpr const char* DiffuseColor = "diffuseColor";
-			static constexpr const char* EmissiveColor = "emissiveColor";
-			static constexpr const char* SpecularColor = "specularColor";
-			static constexpr const char* Metallic = "metallic";
-			static constexpr const char* Roughness = "roughness";
-			static constexpr const char* Clearcoat = "clearcoat";
-			static constexpr const char* ClearcoatRoughness = "clearcoatRoughness";
-			static constexpr const char* Opacity = "opacity";
-			static constexpr const char* OpacityThreshold = "opacityThreshold";
-			static constexpr const char* IOR = "ior";
-			static constexpr const char* Normal = "normal";
-			static constexpr const char* Displacement = "displacement";
-			static constexpr const char* Occlusion = "occlusion";
+			inline constexpr const char* DiffuseColor = "diffuseColor";
+			inline constexpr const char* EmissiveColor = "emissiveColor";
+			inline constexpr const char* SpecularColor = "specularColor";
+			inline constexpr const char* Metallic = "metallic";
+			inline constexpr const char* Roughness = "roughness";
+			inline constexpr const char* Clearcoat = "clearcoat";
+			inline constexpr const char* ClearcoatRoughness = "clearcoatRoughness";
+			inline constexpr const char* Opacity = "opacity";
+			inline constexpr const char* OpacityThreshold = "opacityThreshold";
+			inline constexpr const char* IOR = "ior";
+			inline constexpr const char* Normal = "normal";
+			inline constexpr const char* Displacement = "displacement";
+			inline constexpr const char* Occlusion = "occlusion";
 		}
 
 		namespace DefaultValue
 		{
 			namespace Float
 			{
-				constexpr float Metallic = 0.f;
-				constexpr float Roughness = 0.5f;
-				constexpr float Clearcoat = 0.f;
-				constexpr float ClearcoatRoughness = 0.01f;
-				constexpr float Opacity = 1.f;
-				constexpr float OpacityThreshold = 0.f;
-				constexpr float IOR = 1.5f;
-				constexpr float Displacement = 0.f;
-				constexpr float Occlusion = 1.f;
+				inline constexpr float Metallic = 0.f;
+				inline constexpr float Roughness = 0.5f;
+				inline constexpr float Clearcoat = 0.f;
+				inline constexpr float ClearcoatRoughness = 0.01f;
+				inline constexpr float Opacity = 1.f;
+				inline constexpr float OpacityThreshold = 0.f;
+				inline constexpr float IOR = 1.5f;
+				inline constexpr float Displacement = 0.f;
+				inline constexpr float Occlusion = 1.f;
 			}
 
 			namespace Color3
 			{
-				constexpr FLinearColor DiffuseColor{ 0.18f, 0.18f, 0.18f };
-				constexpr FLinearColor EmissiveColor{ 0, 0, 0 };
-				constexpr FLinearColor SpecularColor{ 0, 0, 0 };
+				inline constexpr FLinearColor DiffuseColor{ 0.18f, 0.18f, 0.18f };
+				inline constexpr FLinearColor EmissiveColor{ 0, 0, 0 };
+				inline constexpr FLinearColor SpecularColor{ 0, 0, 0 };
 			}
 			
 			namespace Vector3
 			{
-				static const FVector3f Normal{ 0.f, 0.f, 1.f };
+				inline const FVector3f Normal{ 0.f, 0.f, 1.f };
 			}
 		}
 	}
@@ -630,8 +630,8 @@ MATERIALX_NAMESPACE_BEGIN
 		//There's no input per se in a Light, but we can find some common inputs among those lights
 		namespace Input
 		{
-			static constexpr const char* Color = "color";
-			static constexpr const char* Intensity = "intensity";
+			inline constexpr const char* Color = "color";
+			inline constexpr const char* Intensity = "intensity";
 		}
 
 		namespace PointLight
@@ -639,8 +639,8 @@ MATERIALX_NAMESPACE_BEGIN
 			namespace Input
 			{
 				using namespace Lights::Input;
-				static constexpr const char* Position = "position";
-				static constexpr const char* DecayRate = "decay_rate";
+				inline constexpr const char* Position = "position";
+				inline constexpr const char* DecayRate = "decay_rate";
 			}
 		}
 
@@ -649,7 +649,7 @@ MATERIALX_NAMESPACE_BEGIN
 			namespace Input
 			{
 				using namespace Lights::Input;
-				static constexpr const char* Direction = "direction";
+				inline constexpr const char* Direction = "direction";
 			}
 		}
 
@@ -659,279 +659,279 @@ MATERIALX_NAMESPACE_BEGIN
 			{
 				using namespace PointLight::Input;
 				using namespace DirectionalLight::Input;
-				static constexpr const char* InnerAngle = "inner_angle";
-				static constexpr const char* OuterAngle = "outer_angle";
+				inline constexpr const char* InnerAngle = "inner_angle";
+				inline constexpr const char* OuterAngle = "outer_angle";
 			}
 		}
 	}
 
 	namespace Attributes
 	{
-		static constexpr const char* IsVisited = "UE:IsVisited";
-		static constexpr const char* NewName = "UE:NewName";
-		static constexpr const char* ParentName = "UE:ParentName";
-		static constexpr const char* UniqueName = "UE:UniqueName";
-		static constexpr const char* GeomPropImage = "UE:GeomPropImage";
-		static constexpr const char* GeomPropSparseVolume = "UE:GeomPropSparseVolume";
+		inline constexpr const char* IsVisited = "UE:IsVisited";
+		inline constexpr const char* NewName = "UE:NewName";
+		inline constexpr const char* ParentName = "UE:ParentName";
+		inline constexpr const char* UniqueName = "UE:UniqueName";
+		inline constexpr const char* GeomPropImage = "UE:GeomPropImage";
+		inline constexpr const char* GeomPropSparseVolume = "UE:GeomPropSparseVolume";
 	}
 
 	namespace Category
 	{
 		// Math nodes
-		static constexpr const char* Absval = "absval";
-		static constexpr const char* Acos = "acos";
-		static constexpr const char* Add = "add";
-		static constexpr const char* ArrayAppend = "arrayappend";
-		static constexpr const char* Asin = "asin";
-		static constexpr const char* Atan2 = "atan2";
-		static constexpr const char* Ceil = "ceil";
-		static constexpr const char* Clamp = "clamp";
-		static constexpr const char* Cos = "cos";
-		static constexpr const char* CrossProduct = "crossproduct";
-		static constexpr const char* Determinant = "determinant";
-		static constexpr const char* Divide = "divide";
-		static constexpr const char* Distance = "distance";
-		static constexpr const char* DotProduct = "dotproduct";
-		static constexpr const char* Exp = "exp";
-		static constexpr const char* Floor = "floor";
-		static constexpr const char* Fract = "fract";
-		static constexpr const char* Invert = "invert";
-		static constexpr const char* InvertMatrix = "invertmatrix";
-		static constexpr const char* Ln = "ln";
-		static constexpr const char* Magnitude = "magnitude";
-		static constexpr const char* Max = "max";
-		static constexpr const char* Min = "min";
-		static constexpr const char* Modulo = "modulo";
-		static constexpr const char* Multiply = "multiply";
-		static constexpr const char* Normalize = "normalize";
-		static constexpr const char* NormalMap = "normalmap";
-		static constexpr const char* Place2D = "place2d";
-		static constexpr const char* Power = "power";
-		static constexpr const char* Reflect = "reflect";
-		static constexpr const char* Refract = "refract";
-		static constexpr const char* Rotate2D = "rotate2d";
-		static constexpr const char* Rotate3D = "rotate3d";
-		static constexpr const char* Round = "round";
-		static constexpr const char* SafePower = "safepower";
-		static constexpr const char* Sign = "sign";
-		static constexpr const char* Sin = "sin";
-		static constexpr const char* Sqrt = "sqrt";
-		static constexpr const char* Sub = "subtract";
-		static constexpr const char* Tan = "tan";
-		static constexpr const char* TransformMatrix = "transformmatrix";
-		static constexpr const char* TransformNormal = "transformnormal";
-		static constexpr const char* TransformPoint = "transformpoint";
-		static constexpr const char* TransformVector = "transformvector";
-		static constexpr const char* Transpose = "transpose";
-		static constexpr const char* TriangleWave = "trianglewave";
+		inline constexpr const char* Absval = "absval";
+		inline constexpr const char* Acos = "acos";
+		inline constexpr const char* Add = "add";
+		inline constexpr const char* ArrayAppend = "arrayappend";
+		inline constexpr const char* Asin = "asin";
+		inline constexpr const char* Atan2 = "atan2";
+		inline constexpr const char* Ceil = "ceil";
+		inline constexpr const char* Clamp = "clamp";
+		inline constexpr const char* Cos = "cos";
+		inline constexpr const char* CrossProduct = "crossproduct";
+		inline constexpr const char* Determinant = "determinant";
+		inline constexpr const char* Divide = "divide";
+		inline constexpr const char* Distance = "distance";
+		inline constexpr const char* DotProduct = "dotproduct";
+		inline constexpr const char* Exp = "exp";
+		inline constexpr const char* Floor = "floor";
+		inline constexpr const char* Fract = "fract";
+		inline constexpr const char* Invert = "invert";
+		inline constexpr const char* InvertMatrix = "invertmatrix";
+		inline constexpr const char* Ln = "ln";
+		inline constexpr const char* Magnitude = "magnitude";
+		inline constexpr const char* Max = "max";
+		inline constexpr const char* Min = "min";
+		inline constexpr const char* Modulo = "modulo";
+		inline constexpr const char* Multiply = "multiply";
+		inline constexpr const char* Normalize = "normalize";
+		inline constexpr const char* NormalMap = "normalmap";
+		inline constexpr const char* Place2D = "place2d";
+		inline constexpr const char* Power = "power";
+		inline constexpr const char* Reflect = "reflect";
+		inline constexpr const char* Refract = "refract";
+		inline constexpr const char* Rotate2D = "rotate2d";
+		inline constexpr const char* Rotate3D = "rotate3d";
+		inline constexpr const char* Round = "round";
+		inline constexpr const char* SafePower = "safepower";
+		inline constexpr const char* Sign = "sign";
+		inline constexpr const char* Sin = "sin";
+		inline constexpr const char* Sqrt = "sqrt";
+		inline constexpr const char* Sub = "subtract";
+		inline constexpr const char* Tan = "tan";
+		inline constexpr const char* TransformMatrix = "transformmatrix";
+		inline constexpr const char* TransformNormal = "transformnormal";
+		inline constexpr const char* TransformPoint = "transformpoint";
+		inline constexpr const char* TransformVector = "transformvector";
+		inline constexpr const char* Transpose = "transpose";
+		inline constexpr const char* TriangleWave = "trianglewave";
 		// Compositing nodes
-		static constexpr const char* Burn = "burn";
-		static constexpr const char* Difference = "difference";
-		static constexpr const char* Disjointover = "disjointover";
-		static constexpr const char* Dodge = "dodge";
-		static constexpr const char* In = "in";
-		static constexpr const char* Inside = "inside";
-		static constexpr const char* Mask = "mask";
-		static constexpr const char* Matte = "matte";
-		static constexpr const char* Minus = "minus";
-		static constexpr const char* Mix = "mix";
-		static constexpr const char* Out = "out";
-		static constexpr const char* Outside = "outside";
-		static constexpr const char* Over = "over";
-		static constexpr const char* Overlay = "overlay";
-		static constexpr const char* Plus = "plus";
-		static constexpr const char* Premult = "premult";
-		static constexpr const char* Screen = "screen";
-		static constexpr const char* Unpremult = "unpremult";
+		inline constexpr const char* Burn = "burn";
+		inline constexpr const char* Difference = "difference";
+		inline constexpr const char* Disjointover = "disjointover";
+		inline constexpr const char* Dodge = "dodge";
+		inline constexpr const char* In = "in";
+		inline constexpr const char* Inside = "inside";
+		inline constexpr const char* Mask = "mask";
+		inline constexpr const char* Matte = "matte";
+		inline constexpr const char* Minus = "minus";
+		inline constexpr const char* Mix = "mix";
+		inline constexpr const char* Out = "out";
+		inline constexpr const char* Outside = "outside";
+		inline constexpr const char* Over = "over";
+		inline constexpr const char* Overlay = "overlay";
+		inline constexpr const char* Plus = "plus";
+		inline constexpr const char* Premult = "premult";
+		inline constexpr const char* Screen = "screen";
+		inline constexpr const char* Unpremult = "unpremult";
 		// Conditional nodes
-		static constexpr const char* And = "and";
-		static constexpr const char* IfGreater = "ifgreater";
-		static constexpr const char* IfGreaterEq = "ifgreatereq";
-		static constexpr const char* IfEqual = "ifequal";
-		static constexpr const char* Not = "not";
-		static constexpr const char* Or = "or";
-		static constexpr const char* Switch = "switch";
-		static constexpr const char* Xor = "xor";
+		inline constexpr const char* And = "and";
+		inline constexpr const char* IfGreater = "ifgreater";
+		inline constexpr const char* IfGreaterEq = "ifgreatereq";
+		inline constexpr const char* IfEqual = "ifequal";
+		inline constexpr const char* Not = "not";
+		inline constexpr const char* Or = "or";
+		inline constexpr const char* Switch = "switch";
+		inline constexpr const char* Xor = "xor";
 		// Channel nodes
-		static constexpr const char* Convert = "convert";
-		static constexpr const char* Combine2 = "combine2";
-		static constexpr const char* Combine3 = "combine3";
-		static constexpr const char* Combine4 = "combine4";
-		static constexpr const char* Extract = "extract";
-		static constexpr const char* Separate2 = "separate2";
-		static constexpr const char* Separate3 = "separate3";
-		static constexpr const char* Separate4 = "separate4";
+		inline constexpr const char* Convert = "convert";
+		inline constexpr const char* Combine2 = "combine2";
+		inline constexpr const char* Combine3 = "combine3";
+		inline constexpr const char* Combine4 = "combine4";
+		inline constexpr const char* Extract = "extract";
+		inline constexpr const char* Separate2 = "separate2";
+		inline constexpr const char* Separate3 = "separate3";
+		inline constexpr const char* Separate4 = "separate4";
 		// Procedural nodes 
-		static constexpr const char* Constant = "constant";
+		inline constexpr const char* Constant = "constant";
 		// Procedural2D nodes
-		static constexpr const char* CellNoise2D = "cellnoise2d";
-		static constexpr const char* Checkerboard = "checkerboard";
-		static constexpr const char* Circle = "circle";
-		static constexpr const char* Line = "line";
-		static constexpr const char* Noise2D = "noise2d";
-		static constexpr const char* Ramp4 = "ramp4";
-		static constexpr const char* Ramp = "ramp";
-		static constexpr const char* RampGradient = "ramp_gradient";
-		static constexpr const char* RampLR = "ramplr";
-		static constexpr const char* RampTB = "ramptb";
-		static constexpr const char* RandomFloat = "randomfloat";
-		static constexpr const char* SplitLR = "splitlr";
-		static constexpr const char* SplitTB = "splittb";
-		static constexpr const char* TiledCircles = "tiledcircles";
-		static constexpr const char* WorleyNoise2D = "worleynoise2d";
+		inline constexpr const char* CellNoise2D = "cellnoise2d";
+		inline constexpr const char* Checkerboard = "checkerboard";
+		inline constexpr const char* Circle = "circle";
+		inline constexpr const char* Line = "line";
+		inline constexpr const char* Noise2D = "noise2d";
+		inline constexpr const char* Ramp4 = "ramp4";
+		inline constexpr const char* Ramp = "ramp";
+		inline constexpr const char* RampGradient = "ramp_gradient";
+		inline constexpr const char* RampLR = "ramplr";
+		inline constexpr const char* RampTB = "ramptb";
+		inline constexpr const char* RandomFloat = "randomfloat";
+		inline constexpr const char* SplitLR = "splitlr";
+		inline constexpr const char* SplitTB = "splittb";
+		inline constexpr const char* TiledCircles = "tiledcircles";
+		inline constexpr const char* WorleyNoise2D = "worleynoise2d";
 		// Procedural3D nodes 
-		static constexpr const char* CellNoise3D = "cellnoise3d";
-		static constexpr const char* Fractal3D = "fractal3d";
-		static constexpr const char* Noise3D = "noise3d";
-		static constexpr const char* RandomColor = "randomcolor";
-		static constexpr const char* UnifiedNoise3D = "unifiednoise3d";
-		static constexpr const char* WorleyNoise3D = "worleynoise3d";
+		inline constexpr const char* CellNoise3D = "cellnoise3d";
+		inline constexpr const char* Fractal3D = "fractal3d";
+		inline constexpr const char* Noise3D = "noise3d";
+		inline constexpr const char* RandomColor = "randomcolor";
+		inline constexpr const char* UnifiedNoise3D = "unifiednoise3d";
+		inline constexpr const char* WorleyNoise3D = "worleynoise3d";
 		// Organization nodes 
-		static constexpr const char* Dot = "dot";
+		inline constexpr const char* Dot = "dot";
 		// Texture nodes 
-		static constexpr const char* Image = "image";
-		static constexpr const char* TiledImage = "tiledimage";
+		inline constexpr const char* Image = "image";
+		inline constexpr const char* TiledImage = "tiledimage";
 		// Geometric nodes
-		static constexpr const char* Bitangent = "bitangent";
-		static constexpr const char* GeomColor = "geomcolor";
-		static constexpr const char* GeomPropValue = "geompropvalue";
-		static constexpr const char* Normal = "normal";
-		static constexpr const char* Position = "position";
-		static constexpr const char* Tangent = "tangent";
-		static constexpr const char* TexCoord = "texcoord";
+		inline constexpr const char* Bitangent = "bitangent";
+		inline constexpr const char* GeomColor = "geomcolor";
+		inline constexpr const char* GeomPropValue = "geompropvalue";
+		inline constexpr const char* Normal = "normal";
+		inline constexpr const char* Position = "position";
+		inline constexpr const char* Tangent = "tangent";
+		inline constexpr const char* TexCoord = "texcoord";
 		// Light nodes
-		static constexpr const char* DirectionalLight = "directional_light";
-		static constexpr const char* PointLight = "point_light";
-		static constexpr const char* SpotLight = "spot_light";
+		inline constexpr const char* DirectionalLight = "directional_light";
+		inline constexpr const char* PointLight = "point_light";
+		inline constexpr const char* SpotLight = "spot_light";
 		// Adjustment nodes
-		static constexpr const char* ColorCorrect = "colorcorrect";
-		static constexpr const char* Contrast = "contrast";
-		static constexpr const char* HsvAdjust = "hsvadjust";
-		static constexpr const char* HsvToRgb = "hsvtorgb";
-		static constexpr const char* Luminance = "luminance";
-		static constexpr const char* Range = "range";
-		static constexpr const char* Remap = "remap";
-		static constexpr const char* RgbToHsv = "rgbtohsv";
-		static constexpr const char* Saturate = "saturate";
-		static constexpr const char* Smoothstep = "smoothstep";
+		inline constexpr const char* ColorCorrect = "colorcorrect";
+		inline constexpr const char* Contrast = "contrast";
+		inline constexpr const char* HsvAdjust = "hsvadjust";
+		inline constexpr const char* HsvToRgb = "hsvtorgb";
+		inline constexpr const char* Luminance = "luminance";
+		inline constexpr const char* Range = "range";
+		inline constexpr const char* Remap = "remap";
+		inline constexpr const char* RgbToHsv = "rgbtohsv";
+		inline constexpr const char* Saturate = "saturate";
+		inline constexpr const char* Smoothstep = "smoothstep";
 		// Application
-		static constexpr const char* Time= "time";
+		inline constexpr const char* Time= "time";
 		// PBR
 		// BSDF
-		static constexpr const char* BurleyDiffuseBSDF= "burley_diffuse_bsdf";
-		static constexpr const char* ChiangHairBSDF = "chiang_hair_bsdf";
-		static constexpr const char* ConductorBSDF = "conductor_bsdf";
-		static constexpr const char* DielectricBSDF = "dielectric_bsdf";
-		static constexpr const char* GeneralizedSchlickBSDF= "generalized_schlick_bsdf";
-		static constexpr const char* OrenNayarDiffuseBSDF = "oren_nayar_diffuse_bsdf";
-		static constexpr const char* SheenBSDF = "sheen_bsdf";
-		static constexpr const char* SubsurfaceBSDF= "subsurface_bsdf";
-		static constexpr const char* TranslucentBSDF = "translucent_bsdf";
+		inline constexpr const char* BurleyDiffuseBSDF= "burley_diffuse_bsdf";
+		inline constexpr const char* ChiangHairBSDF = "chiang_hair_bsdf";
+		inline constexpr const char* ConductorBSDF = "conductor_bsdf";
+		inline constexpr const char* DielectricBSDF = "dielectric_bsdf";
+		inline constexpr const char* GeneralizedSchlickBSDF= "generalized_schlick_bsdf";
+		inline constexpr const char* OrenNayarDiffuseBSDF = "oren_nayar_diffuse_bsdf";
+		inline constexpr const char* SheenBSDF = "sheen_bsdf";
+		inline constexpr const char* SubsurfaceBSDF= "subsurface_bsdf";
+		inline constexpr const char* TranslucentBSDF = "translucent_bsdf";
 		// EDF
-		static constexpr const char* ConicalEDF = "conical_edf";
-		static constexpr const char* MeasuredEDF = "measured_edf";
-		static constexpr const char* UniformEDF = "uniform_edf";
+		inline constexpr const char* ConicalEDF = "conical_edf";
+		inline constexpr const char* MeasuredEDF = "measured_edf";
+		inline constexpr const char* UniformEDF = "uniform_edf";
 		// VDF
-		static constexpr const char* AbsorptionVDF = "absorption_vdf";
-		static constexpr const char* AnisotropicVDF = "anisotropic_vdf";
+		inline constexpr const char* AbsorptionVDF = "absorption_vdf";
+		inline constexpr const char* AnisotropicVDF = "anisotropic_vdf";
 		// PBR Utility Nodes
-		static constexpr const char* ArtisticIOR = "artistic_ior";
-		static constexpr const char* Blackbody = "blackbody";
-		static constexpr const char* ChiangHairAbsorptionFromColor = "chiang_hair_absorption_from_color";
-		static constexpr const char* ChiangHairRoughness = "chiang_hair_roughness";
-		static constexpr const char* DeonHairAbsorptionFromMelanin = "deon_hair_absorption_from_melanin";
-		static constexpr const char* Layer = "layer";
-		static constexpr const char* RoughnessAnisotropy = "roughness_anisotropy";
-		static constexpr const char* RoughnessDual = "roughness_dual";
+		inline constexpr const char* ArtisticIOR = "artistic_ior";
+		inline constexpr const char* Blackbody = "blackbody";
+		inline constexpr const char* ChiangHairAbsorptionFromColor = "chiang_hair_absorption_from_color";
+		inline constexpr const char* ChiangHairRoughness = "chiang_hair_roughness";
+		inline constexpr const char* DeonHairAbsorptionFromMelanin = "deon_hair_absorption_from_melanin";
+		inline constexpr const char* Layer = "layer";
+		inline constexpr const char* RoughnessAnisotropy = "roughness_anisotropy";
+		inline constexpr const char* RoughnessDual = "roughness_dual";
 		//Surface Shaders
-		static constexpr const char* GltfPbr = "gltf_pbr";
-		static constexpr const char* DisneyBSDF2012 = "disney_brdf_2012";
-		static constexpr const char* DisneyBSDF2015 = "disney_bsdf_2015";
-		static constexpr const char* OpenPBRSurface = "open_pbr_surface";
-		static constexpr const char* StandardSurface = "standard_surface";
-		static constexpr const char* Surface = "surface";
-		static constexpr const char* UsdPreviewSurface = "UsdPreviewSurface";
+		inline constexpr const char* GltfPbr = "gltf_pbr";
+		inline constexpr const char* DisneyBSDF2012 = "disney_brdf_2012";
+		inline constexpr const char* DisneyBSDF2015 = "disney_bsdf_2015";
+		inline constexpr const char* OpenPBRSurface = "open_pbr_surface";
+		inline constexpr const char* StandardSurface = "standard_surface";
+		inline constexpr const char* Surface = "surface";
+		inline constexpr const char* UsdPreviewSurface = "UsdPreviewSurface";
 		// Shader
-		static constexpr const char* SurfaceUnlit = "surface_unlit";
+		inline constexpr const char* SurfaceUnlit = "surface_unlit";
 		// Displacement Shader
-		static constexpr const char* Displacement = "displacement";
+		inline constexpr const char* Displacement = "displacement";
 		// Convolution
-		static constexpr const char* Blur = "blur";
-		static constexpr const char* HeightToNormal = "heighttonormal";
+		inline constexpr const char* Blur = "blur";
+		inline constexpr const char* HeightToNormal = "heighttonormal";
 		// ColorTransform
-		static constexpr const char* ACEScgToLinRec709 = "acescg_to_lin_rec709";
-		static constexpr const char* AdobeRgbToLinRec709 = "adobergb_to_lin_rec709";
-		static constexpr const char* G18Rec709ToLinRec709 = "g18_rec709_to_lin_rec709";
-		static constexpr const char* G22ApP1ToLinRec709 = "g22_ap1_to_lin_rec709";
-		static constexpr const char* LinAdobeRgbToLinRec709 = "lin_adobergb_to_lin_rec709";
-		static constexpr const char* LinDisplayP3ToLinRec709 = "lin_displayp3_to_lin_rec709";
-		static constexpr const char* Rec709DisplayToLinRec709 = "rec709_display_to_lin_rec709";
-		static constexpr const char* SrgbDisplayP3ToLinRec709 = "srgb_displayp3_to_lin_rec709";
-		static constexpr const char* SrgbTextureToLinRec709 = "srgb_texture_to_lin_rec709";
+		inline constexpr const char* ACEScgToLinRec709 = "acescg_to_lin_rec709";
+		inline constexpr const char* AdobeRgbToLinRec709 = "adobergb_to_lin_rec709";
+		inline constexpr const char* G18Rec709ToLinRec709 = "g18_rec709_to_lin_rec709";
+		inline constexpr const char* G22ApP1ToLinRec709 = "g22_ap1_to_lin_rec709";
+		inline constexpr const char* LinAdobeRgbToLinRec709 = "lin_adobergb_to_lin_rec709";
+		inline constexpr const char* LinDisplayP3ToLinRec709 = "lin_displayp3_to_lin_rec709";
+		inline constexpr const char* Rec709DisplayToLinRec709 = "rec709_display_to_lin_rec709";
+		inline constexpr const char* SrgbDisplayP3ToLinRec709 = "srgb_displayp3_to_lin_rec709";
+		inline constexpr const char* SrgbTextureToLinRec709 = "srgb_texture_to_lin_rec709";
 		// NPR
-		static constexpr const char* FacingRatio = "facingratio";
-		static constexpr const char* GoochShade = "gooch_shade";
-		static constexpr const char* ViewDirection = "viewdirection";
+		inline constexpr const char* FacingRatio = "facingratio";
+		inline constexpr const char* GoochShade = "gooch_shade";
+		inline constexpr const char* ViewDirection = "viewdirection";
 		// Volume shaders
-		static constexpr const char* Volume = "volume";
+		inline constexpr const char* Volume = "volume";
 		// Unreal Engine nodes defined in Engine\Binaries\ThirdParty\MaterialX\libraries\Interchange\Interchange_defs.mtlx
-		static constexpr const char* SparseVolume = "sparse_volume";
+		inline constexpr const char* SparseVolume = "sparse_volume";
 	}
 
 	namespace NodeDefinition
 	{
-		static constexpr const char* OpenPBRSurface = "ND_open_pbr_surface_surfaceshader";
-		static constexpr const char* StandardSurface = "ND_standard_surface_surfaceshader";
-		static constexpr const char* SurfaceUnlit = "ND_surface_unlit";
-		static constexpr const char* UsdPreviewSurface = "ND_UsdPreviewSurface_surfaceshader";
-		static constexpr const char* PointLight = "ND_point_light";
-		static constexpr const char* DirectionalLight = "ND_directional_light";
-		static constexpr const char* SpotLight = "ND_spot_light";
-		static constexpr const char* Surface = "ND_surface";
-		static constexpr const char* DisplacementFloat = "ND_displacement_float";
-		static constexpr const char* DisplacementVector3 = "ND_displacement_vector3";
-		static constexpr const char* MixSurfaceShader = "ND_mix_surfaceshader";
-		static constexpr const char* Volume = "ND_volume";
+		inline constexpr const char* OpenPBRSurface = "ND_open_pbr_surface_surfaceshader";
+		inline constexpr const char* StandardSurface = "ND_standard_surface_surfaceshader";
+		inline constexpr const char* SurfaceUnlit = "ND_surface_unlit";
+		inline constexpr const char* UsdPreviewSurface = "ND_UsdPreviewSurface_surfaceshader";
+		inline constexpr const char* PointLight = "ND_point_light";
+		inline constexpr const char* DirectionalLight = "ND_directional_light";
+		inline constexpr const char* SpotLight = "ND_spot_light";
+		inline constexpr const char* Surface = "ND_surface";
+		inline constexpr const char* DisplacementFloat = "ND_displacement_float";
+		inline constexpr const char* DisplacementVector3 = "ND_displacement_vector3";
+		inline constexpr const char* MixSurfaceShader = "ND_mix_surfaceshader";
+		inline constexpr const char* Volume = "ND_volume";
 	}
 
 	namespace Library
 	{
-		static constexpr const char* Libraries = "libraries";
+		inline constexpr const char* Libraries = "libraries";
 	}
 
 	namespace Type
 	{
-		static constexpr const char* Boolean = "boolean";
-		static constexpr const char* Integer = "integer";
-		static constexpr const char* Float = "float";
-		static constexpr const char* Color3 = "color3";
-		static constexpr const char* Color4 = "color4";
-		static constexpr const char* Vector2 = "vector2";
-		static constexpr const char* Vector3 = "vector3";
-		static constexpr const char* Vector4 = "vector4";
-		static constexpr const char* Matrix33 = "matrix33";
-		static constexpr const char* Matrix44 = "matrix44";
-		static constexpr const char* String = "string";
-		static constexpr const char* Filename = "filename";
-		static constexpr const char* GeomName = "geomname";
-		static constexpr const char* SurfaceShader = "surfaceshader";
-		static constexpr const char* DisplacementShader = "displacementshader";
-		static constexpr const char* VolumeShader = "volumeshader";
-		static constexpr const char* LightShader = "lightshader";
-		static constexpr const char* Material = "material";
-		static constexpr const char* None = "none";
-		static constexpr const char* IntegerArray = "integerarray";
-		static constexpr const char* FloatArray = "floatarray";
-		static constexpr const char* Color3Array = "color3array";
-		static constexpr const char* Color4Array = "color4array";
-		static constexpr const char* Vector2Array = "vector2array";
-		static constexpr const char* Vector3Array = "vector3array";
-		static constexpr const char* Vector4Array = "vector4array";
-		static constexpr const char* StringArray = "stringarray";
-		static constexpr const char* GeomNameArray = "geomnamearray";
-		static constexpr const char* BSDF = "BSDF";
-		static constexpr const char* EDF = "EDF";
-		static constexpr const char* VDF = "VDF";
+		inline constexpr const char* Boolean = "boolean";
+		inline constexpr const char* Integer = "integer";
+		inline constexpr const char* Float = "float";
+		inline constexpr const char* Color3 = "color3";
+		inline constexpr const char* Color4 = "color4";
+		inline constexpr const char* Vector2 = "vector2";
+		inline constexpr const char* Vector3 = "vector3";
+		inline constexpr const char* Vector4 = "vector4";
+		inline constexpr const char* Matrix33 = "matrix33";
+		inline constexpr const char* Matrix44 = "matrix44";
+		inline constexpr const char* String = "string";
+		inline constexpr const char* Filename = "filename";
+		inline constexpr const char* GeomName = "geomname";
+		inline constexpr const char* SurfaceShader = "surfaceshader";
+		inline constexpr const char* DisplacementShader = "displacementshader";
+		inline constexpr const char* VolumeShader = "volumeshader";
+		inline constexpr const char* LightShader = "lightshader";
+		inline constexpr const char* Material = "material";
+		inline constexpr const char* None = "none";
+		inline constexpr const char* IntegerArray = "integerarray";
+		inline constexpr const char* FloatArray = "floatarray";
+		inline constexpr const char* Color3Array = "color3array";
+		inline constexpr const char* Color4Array = "color4array";
+		inline constexpr const char* Vector2Array = "vector2array";
+		inline constexpr const char* Vector3Array = "vector3array";
+		inline constexpr const char* Vector4Array = "vector4array";
+		inline constexpr const char* StringArray = "stringarray";
+		inline constexpr const char* GeomNameArray = "geomnamearray";
+		inline constexpr const char* BSDF = "BSDF";
+		inline constexpr const char* EDF = "EDF";
+		inline constexpr const char* VDF = "VDF";
 
 	}
 
@@ -941,27 +941,27 @@ MATERIALX_NAMESPACE_BEGIN
 		{
 			namespace Inputs
 			{
-				static constexpr const char* File = "file";
-				static constexpr const char* Default = "default";
-				static constexpr const char* TexCoord = "texcoord";
-				static constexpr const char* FilterType = "filtertype";
-				static constexpr const char* FrameRange = "framerange";
-				static constexpr const char* FrameOffset = "frameoffset";
-				static constexpr const char* FrameEndAction = "frameendaction";
+				inline constexpr const char* File = "file";
+				inline constexpr const char* Default = "default";
+				inline constexpr const char* TexCoord = "texcoord";
+				inline constexpr const char* FilterType = "filtertype";
+				inline constexpr const char* FrameRange = "framerange";
+				inline constexpr const char* FrameOffset = "frameoffset";
+				inline constexpr const char* FrameEndAction = "frameendaction";
 			}
 		}
 
 		namespace Math
 		{
 			template<typename T>
-			static const T NeutralZero = T{ 0 };
+			inline const T NeutralZero = T{ 0 };
 
 			template<typename T>
-			static const T NeutralOne = T{ 1 };
+			inline const T NeutralOne = T{ 1 };
 		}
 
-		static constexpr const char* Compositing = "compositing";
-		static constexpr const char* PBR = "pbr";
+		inline constexpr const char* Compositing = "compositing";
+		inline constexpr const char* PBR = "pbr";
 	}
 
 	namespace Image
@@ -970,9 +970,9 @@ MATERIALX_NAMESPACE_BEGIN
 		{
 			using namespace NodeGroup::Texture2D::Inputs;
 
-			static constexpr const char* Layer = "layer";
-			static constexpr const char* UAddressMode = "uaddressmode";
-			static constexpr const char* VAddressMode = "vaddressmode";
+			inline constexpr const char* Layer = "layer";
+			inline constexpr const char* UAddressMode = "uaddressmode";
+			inline constexpr const char* VAddressMode = "vaddressmode";
 		}
 	}
 
@@ -982,10 +982,10 @@ MATERIALX_NAMESPACE_BEGIN
 		{
 			using namespace NodeGroup::Texture2D::Inputs;
 
-			static constexpr const char* UVTiling = "uvtiling";
-			static constexpr const char* UVOffset = "uvoffset";
-			static constexpr const char* RealWorldImageSize = "realworldimagesize";
-			static constexpr const char* RealWorldTileSize = "realworldtilesize";
+			inline constexpr const char* UVTiling = "uvtiling";
+			inline constexpr const char* UVOffset = "uvoffset";
+			inline constexpr const char* RealWorldImageSize = "realworldimagesize";
+			inline constexpr const char* RealWorldTileSize = "realworldtilesize";
 		}
 	}
 
