Interchange USD: Added support for importing groom cache.

diff --git a/Engine/Plugins/Interchange/Extensions/OpenUSD/Source/Import/Private/InterchangeOpenUSDImportModule.cpp b/Engine/Plugins/Interchange/Extensions/OpenUSD/Source/Import/Private/InterchangeOpenUSDImportModule.cpp
--- a/Engine/Plugins/Interchange/Extensions/OpenUSD/Source/Import/Private/InterchangeOpenUSDImportModule.cpp
+++ b/Engine/Plugins/Interchange/Extensions/OpenUSD/Source/Import/Private/InterchangeOpenUSDImportModule.cpp
@@ -71,6 +71,7 @@ void FInterchangeOpenUSDImportModule::StartupModule()
 				EInterchangeTranslatorAssetType::Materials,
 				EInterchangeTranslatorAssetType::Meshes,
 				EInterchangeTranslatorAssetType::Animations,
+				EInterchangeTranslatorAssetType::Grooms,
 			};
 
 			for (EInterchangeTranslatorAssetType AssetType : AssetTypes)
diff --git a/Engine/Plugins/Interchange/Extensions/OpenUSD/Source/Import/Private/InterchangeUsdTranslator.cpp b/Engine/Plugins/Interchange/Extensions/OpenUSD/Source/Import/Private/InterchangeUsdTranslator.cpp
--- a/Engine/Plugins/Interchange/Extensions/OpenUSD/Source/Import/Private/InterchangeUsdTranslator.cpp
+++ b/Engine/Plugins/Interchange/Extensions/OpenUSD/Source/Import/Private/InterchangeUsdTranslator.cpp
@@ -316,7 +316,7 @@ namespace UE::InterchangeUsdTranslator::Private
 			bool& bOutNeedsFrameTrack
 		);
 
-		FString AddGroomNode(
+		TOptional<FString> AddGroomNode(
 			const UE::FUsdPrim& Prim,
 			UInterchangeBaseNodeContainer& NodeContainer
 		);
@@ -3237,7 +3237,7 @@ namespace UE::InterchangeUsdTranslator::Private
 		return VolumeAssetNodeUids.Array();
 	}
 
-	FString UInterchangeUSDTranslatorImpl::AddGroomNode(
+	TOptional<FString> UInterchangeUSDTranslatorImpl::AddGroomNode(
 		const UE::FUsdPrim& Prim,
 		UInterchangeBaseNodeContainer& NodeContainer
 	)
@@ -3254,10 +3254,38 @@ namespace UE::InterchangeUsdTranslator::Private
 			return NodeUid;
 		}
 
+		// Check if there's any actual groom data
+		FGroomAnimationInfo AnimInfo;
+		FHairDescription HairDescription;
+		if (!UsdToUnreal::ConvertGroomHierarchy(Prim, UsdUtils::GetEarliestTimeCode(), FTransform::Identity, HairDescription, &AnimInfo))
+		{
+			return {};
+		}
+
+		const double StageTimeCodesPerSecond = Prim.GetStage().GetTimeCodesPerSecond();
+		AnimInfo.SecondsPerFrame = static_cast<float>(1 / StageTimeCodesPerSecond);
+
+		const bool bIsValidGroomCache = AnimInfo.IsValid();
+
 		UInterchangeGroomNode* GroomNode = NewObject<UInterchangeGroomNode>(&NodeContainer);
 		NodeContainer.SetupNode(GroomNode, NodeUid, NodeName, EInterchangeNodeContainerType::TranslatedAsset);
 		GroomNode->SetAssetName(NodeName);
-		GroomNode->SetPayloadKey(PrimPath);
+		GroomNode->SetPayloadKey(PrimPath, bIsValidGroomCache ? EInterchangeGroomPayLoadType::ANIMATED : EInterchangeGroomPayLoadType::STATIC);
+
+		if (bIsValidGroomCache)
+		{
+			GroomNode->SetCustomNumFrames(AnimInfo.NumFrames);
+			GroomNode->SetCustomStartFrame(AnimInfo.StartFrame);
+			GroomNode->SetCustomEndFrame(AnimInfo.EndFrame);
+			GroomNode->SetCustomFrameRate(StageTimeCodesPerSecond);
+
+			static_assert((int)EGroomCacheAttributes::None == (int)EInterchangeGroomCacheAttributes::None);
+			static_assert((int)EGroomCacheAttributes::Position == (int)EInterchangeGroomCacheAttributes::Position);
+			static_assert((int)EGroomCacheAttributes::Width == (int)EInterchangeGroomCacheAttributes::Width);
+			static_assert((int)EGroomCacheAttributes::Color == (int)EInterchangeGroomCacheAttributes::Color);
+
+			GroomNode->SetCustomGroomCacheAttributes(static_cast<EInterchangeGroomCacheAttributes>(AnimInfo.Attributes));
+		}
 
 		return NodeUid;
 	}
@@ -4214,7 +4242,11 @@ namespace UE::InterchangeUsdTranslator::Private
 		}
 		else if ((Prim.IsA(TEXT("BasisCurves")) || Prim.IsA(TEXT("Xform"))) && Prim.HasAPI(TEXT("GroomAPI")))
 		{
-			AssetNodeUids = {TranslatorImpl.AddGroomNode(Prim, NodeContainer)};
+			TOptional<FString> AssetNodeUid = TranslatorImpl.AddGroomNode(Prim, NodeContainer);
+			if (AssetNodeUid.IsSet())
+			{
+				AssetNodeUids = {*AssetNodeUid};
+			}
 		}
 		else if (Prim.IsA(TEXT("SpatialAudio")))
 		{
@@ -5160,7 +5192,7 @@ EInterchangeTranslatorType UInterchangeUSDTranslator::GetTranslatorType() const
 
 EInterchangeTranslatorAssetType UInterchangeUSDTranslator::GetSupportedAssetTypes() const
 {
-	return EInterchangeTranslatorAssetType::Materials | EInterchangeTranslatorAssetType::Meshes | EInterchangeTranslatorAssetType::Animations;
+	return EInterchangeTranslatorAssetType::Materials | EInterchangeTranslatorAssetType::Meshes | EInterchangeTranslatorAssetType::Animations | EInterchangeTranslatorAssetType::Grooms;
 }
 
 TArray<FString> UInterchangeUSDTranslator::GetSupportedFormats() const
@@ -5922,10 +5954,11 @@ TOptional<UE::Interchange::FGroomPayloadData> UInterchangeUSDTranslator::GetGroo
 
 	FGroomAnimationInfo AnimInfo;
 	UE::Interchange::FGroomPayloadData PayloadData;
+	const double TimeCode = PayloadKey.Type == EInterchangeGroomPayLoadType::ANIMATED ? PayloadKey.FrameNumber : UsdUtils::GetEarliestTimeCode();
 
 	const bool bSuccess = UsdToUnreal::ConvertGroomHierarchy(
 		Prim,
-		UsdUtils::GetEarliestTimeCode(),
+		TimeCode,
 		FTransform::Identity,
 		PayloadData.HairDescription,
 		&AnimInfo
diff --git a/Engine/Plugins/Interchange/Runtime/Source/FactoryNodes/Private/InterchangeGroomCacheFactoryNode.cpp b/Engine/Plugins/Interchange/Runtime/Source/FactoryNodes/Private/InterchangeGroomCacheFactoryNode.cpp
new file mode 100644
--- /dev/null
+++ b/Engine/Plugins/Interchange/Runtime/Source/FactoryNodes/Private/InterchangeGroomCacheFactoryNode.cpp
@@ -0,0 +1,88 @@
+// Copyright Epic Games, Inc. All Rights Reserved.
+
+#include "InterchangeGroomCacheFactoryNode.h"
+
+#include "GroomCache.h"
+
+#include UE_INLINE_GENERATED_CPP_BY_NAME(InterchangeGroomCacheFactoryNode)
+
+FString UInterchangeGroomCacheFactoryNode::GetTypeName() const
+{
+	const FString TypeName = TEXT("GroomCacheFactoryNode");
+	return TypeName;
+}
+
+UClass* UInterchangeGroomCacheFactoryNode::GetObjectClass() const
+{
+	return UGroomCache::StaticClass();
+}
+
+bool UInterchangeGroomCacheFactoryNode::GetCustomStartFrame(int32& AttributeValue) const
+{
+	IMPLEMENT_NODE_ATTRIBUTE_GETTER(StartFrame, int32);
+}
+
+bool UInterchangeGroomCacheFactoryNode::SetCustomStartFrame(const int32& AttributeValue)
+{
+	IMPLEMENT_NODE_ATTRIBUTE_SETTER_NODELEGATE(StartFrame, int32);
+}
+
+bool UInterchangeGroomCacheFactoryNode::GetCustomEndFrame(int32& AttributeValue) const
+{
+	IMPLEMENT_NODE_ATTRIBUTE_GETTER(EndFrame, int32);
+}
+
+bool UInterchangeGroomCacheFactoryNode::SetCustomEndFrame(const int32& AttributeValue)
+{
+	IMPLEMENT_NODE_ATTRIBUTE_SETTER_NODELEGATE(EndFrame, int32);
+}
+
+bool UInterchangeGroomCacheFactoryNode::GetCustomNumFrames(int32& AttributeValue) const
+{
+	IMPLEMENT_NODE_ATTRIBUTE_GETTER(NumFrames, int32);
+}
+
+bool UInterchangeGroomCacheFactoryNode::SetCustomNumFrames(const int32& AttributeValue)
+{
+	IMPLEMENT_NODE_ATTRIBUTE_SETTER_NODELEGATE(NumFrames, int32);
+}
+
+bool UInterchangeGroomCacheFactoryNode::GetCustomFrameRate(double& AttributeValue) const
+{
+	IMPLEMENT_NODE_ATTRIBUTE_GETTER(FrameRate, double);
+}
+
+bool UInterchangeGroomCacheFactoryNode::SetCustomFrameRate(const double& AttributeValue)
+{
+	IMPLEMENT_NODE_ATTRIBUTE_SETTER_NODELEGATE(FrameRate, double);
+}
+
+bool UInterchangeGroomCacheFactoryNode::GetCustomGroomCacheAttributes(EInterchangeGroomCacheAttributes& AttributeValue) const
+{
+	IMPLEMENT_NODE_ATTRIBUTE_GETTER(GroomCacheAttributes, EInterchangeGroomCacheAttributes);
+}
+
+bool UInterchangeGroomCacheFactoryNode::SetCustomGroomCacheAttributes(const EInterchangeGroomCacheAttributes& AttributeValue)
+{
+	IMPLEMENT_NODE_ATTRIBUTE_SETTER_NODELEGATE(GroomCacheAttributes, EInterchangeGroomCacheAttributes);
+}
+
+bool UInterchangeGroomCacheFactoryNode::GetCustomGroomCacheImportType(EInterchangeGroomCacheImportType& AttributeValue) const
+{
+	IMPLEMENT_NODE_ATTRIBUTE_GETTER(ImportType, EInterchangeGroomCacheImportType);
+}
+
+bool UInterchangeGroomCacheFactoryNode::SetCustomGroomCacheImportType(const EInterchangeGroomCacheImportType& AttributeValue)
+{
+	IMPLEMENT_NODE_ATTRIBUTE_SETTER_NODELEGATE(ImportType, EInterchangeGroomCacheImportType);
+}
+
+bool UInterchangeGroomCacheFactoryNode::GetCustomGroomAssetPath(FSoftObjectPath& AttributeValue) const
+{
+	IMPLEMENT_NODE_ATTRIBUTE_GETTER(GroomAssetPath, FSoftObjectPath);
+}
+
+bool UInterchangeGroomCacheFactoryNode::SetCustomGroomAssetPath(const FSoftObjectPath& AttributeValue)
+{
+	IMPLEMENT_NODE_ATTRIBUTE_SETTER_NODELEGATE(GroomAssetPath, FSoftObjectPath);
+}
diff --git a/Engine/Plugins/Interchange/Runtime/Source/FactoryNodes/Private/InterchangeGroomFactoryNode.cpp b/Engine/Plugins/Interchange/Runtime/Source/FactoryNodes/Private/InterchangeGroomFactoryNode.cpp
--- a/Engine/Plugins/Interchange/Runtime/Source/FactoryNodes/Private/InterchangeGroomFactoryNode.cpp
+++ b/Engine/Plugins/Interchange/Runtime/Source/FactoryNodes/Private/InterchangeGroomFactoryNode.cpp
@@ -3,6 +3,7 @@
 #include "InterchangeGroomFactoryNode.h"
 
 #include "GroomAsset.h"
+#include "GroomAssetInterpolation.h"
 
 #include UE_INLINE_GENERATED_CPP_BY_NAME(InterchangeGroomFactoryNode)
 
@@ -16,3 +17,46 @@ UClass* UInterchangeGroomFactoryNode::GetObjectClass() const
 {
 	return UGroomAsset::StaticClass();
 }
+
+#define GET_ATTRIBUTE(Struct, AttributeName, AttributeType) \
+	static const FString AttributeName##OperationName = GetTypeName() + TEXT(".Get" #AttributeName); \
+	if (!InterchangePrivateNodeBase::GetCustomAttribute<AttributeType>(*Attributes, Macro_Custom##AttributeName##Key, AttributeName##OperationName, AttributeValue.Struct.AttributeName)) return false;
+
+#define SET_ATTRIBUTE(Struct, AttributeName, AttributeType) \
+	static const FString AttributeName##OperationName = GetTypeName() + TEXT(".Set" #AttributeName); \
+	if (!InterchangePrivateNodeBase::SetCustomAttribute<AttributeType>(*Attributes, Macro_Custom##AttributeName##Key, AttributeName##OperationName, AttributeValue.Struct.AttributeName)) return false;
+
+bool UInterchangeGroomFactoryNode::GetCustomGroupInterpolationSettings(FHairGroupsInterpolation& AttributeValue) const
+{
+	GET_ATTRIBUTE(DecimationSettings, CurveDecimation, float);
+	GET_ATTRIBUTE(DecimationSettings, VertexDecimation, float);
+	GET_ATTRIBUTE(InterpolationSettings, GuideType, EGroomGuideType);
+	GET_ATTRIBUTE(InterpolationSettings, HairToGuideDensity, float);
+	GET_ATTRIBUTE(InterpolationSettings, RiggedGuideNumCurves, int32);
+	GET_ATTRIBUTE(InterpolationSettings, RiggedGuideNumPoints, int32);
+	GET_ATTRIBUTE(InterpolationSettings, InterpolationQuality, EHairInterpolationQuality);
+	GET_ATTRIBUTE(InterpolationSettings, InterpolationDistance, EHairInterpolationWeight);
+	GET_ATTRIBUTE(InterpolationSettings, bRandomizeGuide, bool);
+	GET_ATTRIBUTE(InterpolationSettings, bUseUniqueGuide, bool);
+
+	return true;
+}
+
+bool UInterchangeGroomFactoryNode::SetCustomGroupInterpolationSettings(const FHairGroupsInterpolation& AttributeValue)
+{
+	SET_ATTRIBUTE(DecimationSettings, CurveDecimation, float);
+	SET_ATTRIBUTE(DecimationSettings, VertexDecimation, float);
+	SET_ATTRIBUTE(InterpolationSettings, GuideType, EGroomGuideType);
+	SET_ATTRIBUTE(InterpolationSettings, HairToGuideDensity, float);
+	SET_ATTRIBUTE(InterpolationSettings, RiggedGuideNumCurves, int32);
+	SET_ATTRIBUTE(InterpolationSettings, RiggedGuideNumPoints, int32);
+	SET_ATTRIBUTE(InterpolationSettings, InterpolationQuality, EHairInterpolationQuality);
+	SET_ATTRIBUTE(InterpolationSettings, InterpolationDistance, EHairInterpolationWeight);
+	SET_ATTRIBUTE(InterpolationSettings, bRandomizeGuide, bool);
+	SET_ATTRIBUTE(InterpolationSettings, bUseUniqueGuide, bool);
+
+	return true;
+}
+
+#undef SET_ATTRIBUTE
+#undef GET_ATTRIBUTE
diff --git a/Engine/Plugins/Interchange/Runtime/Source/FactoryNodes/Public/InterchangeGroomCacheFactoryNode.h b/Engine/Plugins/Interchange/Runtime/Source/FactoryNodes/Public/InterchangeGroomCacheFactoryNode.h
new file mode 100644
--- /dev/null
+++ b/Engine/Plugins/Interchange/Runtime/Source/FactoryNodes/Public/InterchangeGroomCacheFactoryNode.h
@@ -0,0 +1,81 @@
+// Copyright Epic Games, Inc. All Rights Reserved.
+
+#pragma once
+
+#include "Nodes/InterchangeFactoryBaseNode.h"
+
+#include "InterchangeGroomCacheFactoryNode.generated.h"
+
+#define UE_API INTERCHANGEFACTORYNODES_API
+
+UENUM(BlueprintType)
+enum class EInterchangeGroomCacheImportType : uint8
+{
+	None = 0x00 UMETA(Hidden),
+	Strands = 0x01,
+	Guides = 0x02,
+	All = Strands | Guides
+};
+
+UCLASS(MinimalAPI, BlueprintType)
+class UInterchangeGroomCacheFactoryNode : public UInterchangeFactoryBaseNode
+{
+	GENERATED_BODY()
+
+public:
+	UE_API virtual FString GetTypeName() const override;
+	UE_API virtual class UClass* GetObjectClass() const override;
+
+	UFUNCTION(BlueprintCallable, Category = "Interchange | Node | Groom Cache")
+	UE_API bool GetCustomStartFrame(int32& AttributeValue) const;
+
+	UFUNCTION(BlueprintCallable, Category = "Interchange | Node | Groom Cache")
+	UE_API bool SetCustomStartFrame(const int32& AttributeValue);
+
+	UFUNCTION(BlueprintCallable, Category = "Interchange | Node | Groom Cache")
+	UE_API bool GetCustomEndFrame(int32& AttributeValue) const;
+
+	UFUNCTION(BlueprintCallable, Category = "Interchange | Node | Groom Cache")
+	UE_API bool SetCustomEndFrame(const int32& AttributeValue);
+
+	UFUNCTION(BlueprintCallable, Category = "Interchange | Node | Groom Cache")
+	UE_API bool GetCustomNumFrames(int32& AttributeValue) const;
+
+	UFUNCTION(BlueprintCallable, Category = "Interchange | Node | Groom Cache")
+	UE_API bool SetCustomNumFrames(const int32& AttributeValue);
+
+	UFUNCTION(BlueprintCallable, Category = "Interchange | Node | Groom Cache")
+	UE_API bool GetCustomFrameRate(double& AttributeValue) const;
+
+	UFUNCTION(BlueprintCallable, Category = "Interchange | Node | Groom Cache")
+	UE_API bool SetCustomFrameRate(const double& AttributeValue);
+
+	UFUNCTION(BlueprintCallable, Category = "Interchange | Node | Groom Cache")
+	UE_API bool GetCustomGroomCacheAttributes(EInterchangeGroomCacheAttributes& AttributeValue) const;
+
+	UFUNCTION(BlueprintCallable, Category = "Interchange | Node | Groom Cache")
+	UE_API bool SetCustomGroomCacheAttributes(const EInterchangeGroomCacheAttributes& AttributeValue);
+
+	UFUNCTION(BlueprintCallable, Category = "Interchange | Node | Groom Cache")
+	UE_API bool GetCustomGroomCacheImportType(EInterchangeGroomCacheImportType& AttributeValue) const;
+
+	UFUNCTION(BlueprintCallable, Category = "Interchange | Node | Groom Cache")
+	UE_API bool SetCustomGroomCacheImportType(const EInterchangeGroomCacheImportType& AttributeValue);
+
+	UFUNCTION(BlueprintCallable, Category = "Interchange | Node | Groom Cache")
+	UE_API bool GetCustomGroomAssetPath(FSoftObjectPath& AttributeValue) const;
+
+	UFUNCTION(BlueprintCallable, Category = "Interchange | Node | Groom Cache")
+	UE_API bool SetCustomGroomAssetPath(const FSoftObjectPath& AttributeValue);
+
+private:
+	IMPLEMENT_NODE_ATTRIBUTE_KEY(StartFrame);
+	IMPLEMENT_NODE_ATTRIBUTE_KEY(EndFrame);
+	IMPLEMENT_NODE_ATTRIBUTE_KEY(NumFrames);
+	IMPLEMENT_NODE_ATTRIBUTE_KEY(FrameRate);
+	IMPLEMENT_NODE_ATTRIBUTE_KEY(GroomCacheAttributes);
+	IMPLEMENT_NODE_ATTRIBUTE_KEY(ImportType);
+	IMPLEMENT_NODE_ATTRIBUTE_KEY(GroomAssetPath);
+};
+
+#undef UE_API
\ No newline at end of file
diff --git a/Engine/Plugins/Interchange/Runtime/Source/FactoryNodes/Public/InterchangeGroomFactoryNode.h b/Engine/Plugins/Interchange/Runtime/Source/FactoryNodes/Public/InterchangeGroomFactoryNode.h
--- a/Engine/Plugins/Interchange/Runtime/Source/FactoryNodes/Public/InterchangeGroomFactoryNode.h
+++ b/Engine/Plugins/Interchange/Runtime/Source/FactoryNodes/Public/InterchangeGroomFactoryNode.h
@@ -8,6 +8,8 @@
 
 #define UE_API INTERCHANGEFACTORYNODES_API
 
+struct FHairGroupsInterpolation;
+
 UCLASS(MinimalAPI, BlueprintType)
 class UInterchangeGroomFactoryNode : public UInterchangeFactoryBaseNode
 {
@@ -16,6 +18,27 @@ class UInterchangeGroomFactoryNode : public UInterchangeFactoryBaseNode
 public:
 	UE_API virtual FString GetTypeName() const override;
 	UE_API virtual class UClass* GetObjectClass() const override;
+
+	UFUNCTION(BlueprintCallable, Category = "Interchange | Node | Groom Cache")
+	UE_API bool GetCustomGroupInterpolationSettings(FHairGroupsInterpolation& AttributeValue) const;
+
+	UFUNCTION(BlueprintCallable, Category = "Interchange | Node | Groom Cache")
+	UE_API bool SetCustomGroupInterpolationSettings(const FHairGroupsInterpolation& AttributeValue);
+
+private:
+	// DecimationSettings
+	IMPLEMENT_NODE_ATTRIBUTE_KEY(CurveDecimation);
+	IMPLEMENT_NODE_ATTRIBUTE_KEY(VertexDecimation);
+
+	// InterpolationSettings
+	IMPLEMENT_NODE_ATTRIBUTE_KEY(GuideType);
+	IMPLEMENT_NODE_ATTRIBUTE_KEY(HairToGuideDensity);
+	IMPLEMENT_NODE_ATTRIBUTE_KEY(RiggedGuideNumCurves);
+	IMPLEMENT_NODE_ATTRIBUTE_KEY(RiggedGuideNumPoints);
+	IMPLEMENT_NODE_ATTRIBUTE_KEY(InterpolationQuality);
+	IMPLEMENT_NODE_ATTRIBUTE_KEY(InterpolationDistance);
+	IMPLEMENT_NODE_ATTRIBUTE_KEY(bRandomizeGuide);
+	IMPLEMENT_NODE_ATTRIBUTE_KEY(bUseUniqueGuide);
 };
 
 #undef UE_API
\ No newline at end of file
diff --git a/Engine/Plugins/Interchange/Runtime/Source/Import/Private/Groom/InterchangeGroomCacheFactory.cpp b/Engine/Plugins/Interchange/Runtime/Source/Import/Private/Groom/InterchangeGroomCacheFactory.cpp
new file mode 100644
--- /dev/null
+++ b/Engine/Plugins/Interchange/Runtime/Source/Import/Private/Groom/InterchangeGroomCacheFactory.cpp
@@ -0,0 +1,409 @@
+// Copyright Epic Games, Inc. All Rights Reserved.
+
+#include "Groom/InterchangeGroomCacheFactory.h"
+
+#include "Groom/InterchangeGroomPayloadData.h"
+#include "Groom/InterchangeGroomPayloadInterface.h"
+#include "GroomBuilder.h"
+#include "GroomCache.h"
+#include "GroomImportOptions.h"
+#include "InterchangeGroomCacheFactoryNode.h"
+#include "InterchangeGroomFactoryNode.h"
+#include "InterchangeGroomNode.h"
+#include "InterchangeImportCommon.h"
+#include "InterchangeImportLog.h"
+#include "InterchangeSourceData.h"
+#include "InterchangeTranslatorBase.h"
+#include "Misc/ScopedSlowTask.h"
+#include "Nodes/InterchangeBaseNodeContainer.h"
+
+#include UE_INLINE_GENERATED_CPP_BY_NAME(InterchangeGroomCacheFactory)
+
+#define LOCTEXT_NAMESPACE "InterchangeGroomCacheFactory"
+
+UClass* UInterchangeGroomCacheFactory::GetFactoryClass() const
+{
+	return UGroomCache::StaticClass();
+}
+
+UInterchangeFactoryBase::FImportAssetResult UInterchangeGroomCacheFactory::BeginImportAsset_GameThread(const FImportAssetObjectParams& Arguments)
+{
+	TRACE_CPUPROFILER_EVENT_SCOPE(UInterchangeGroomCacheFactory::BeginImportAsset_GameThread);
+
+	FImportAssetResult ImportAssetResult;
+
+	auto LogGroomCacheFactoryError = [this, &Arguments, &ImportAssetResult](const FText& Info)
+	{
+		UInterchangeResultError_Generic* Message = AddMessage<UInterchangeResultError_Generic>();
+		Message->SourceAssetName = Arguments.SourceData->GetFilename();
+		Message->DestinationAssetName = Arguments.AssetName;
+		Message->AssetType = GetFactoryClass();
+		Message->Text = FText::Format(
+			LOCTEXT("GroomCacheFactory_Failure", "UInterchangeGroomCacheFactory: Could not create Groom Cache '{0}'. Reason: {1}"),
+			FText::FromString(Arguments.AssetName),
+			Info
+		);
+		ImportAssetResult.bIsFactorySkipAsset = true;
+	};
+
+	const IInterchangeGroomPayloadInterface* GroomPayloadInterface = Cast<IInterchangeGroomPayloadInterface>(Arguments.Translator);
+	if (!GroomPayloadInterface)
+	{
+		LogGroomCacheFactoryError(LOCTEXT("GroomCacheFactory_NoPayloadInterface", "The translator does not implement IInterchangeGroomPayloadInterface"));
+		return ImportAssetResult;
+	}
+
+	UInterchangeGroomCacheFactoryNode* GroomCacheFactoryNode = Cast<UInterchangeGroomCacheFactoryNode>(Arguments.AssetNode);
+	if (!GroomCacheFactoryNode)
+	{
+		LogGroomCacheFactoryError(LOCTEXT("GroomCacheFactory_AssetNodeNull", "Asset node parameter is not a UInterchangeGroomCacheFactoryNode."));
+		return ImportAssetResult;
+	}
+
+	const UClass* ObjectClass = GroomCacheFactoryNode->GetObjectClass();
+	if (!ObjectClass || !ObjectClass->IsChildOf(GetFactoryClass()))
+	{
+		LogGroomCacheFactoryError(LOCTEXT("GroomCacheFactory_NodeClassMissmatch", "Asset node parameter class doesn't derive the UGroomCache class."));
+		return ImportAssetResult;
+	}
+
+	TArray<FString> GroomUids;
+	GroomCacheFactoryNode->GetTargetNodeUids(GroomUids);
+	if (GroomUids.Num() == 0)
+	{
+		LogGroomCacheFactoryError(LOCTEXT("GroomCacheFactory_MissingTargetNodeId", "Groom cache factory node is missing the target node id."));
+		return ImportAssetResult;
+	}
+
+	const FString& GroomUid = GroomUids[0]; // the groom pipeline adds only one target node
+	const UInterchangeBaseNode* Node = Arguments.NodeContainer->GetNode(GroomUid);
+	const UInterchangeGroomNode* GroomNode = Cast<const UInterchangeGroomNode>(Node);
+	if (!GroomNode)
+	{
+		LogGroomCacheFactoryError(LOCTEXT("GroomCacheFactory_MissingGroomNode", "Groom node not found."));
+		return ImportAssetResult;
+	}
+
+	TOptional<FInterchangeGroomPayloadKey> OptionalPayloadKey = GroomNode->GetPayloadKey();
+	if (!ensure(OptionalPayloadKey.IsSet()))
+	{
+		LogGroomCacheFactoryError(LOCTEXT("GroomCacheFactory_MissingPayloadKey", "Groom node has no payload key."));
+		return ImportAssetResult;
+	}
+
+	// Retrieve the groom asset dependency
+	// First check if a groom asset was specifically selected by the user
+	FSoftObjectPath GroomAssetPath;
+	if (!GroomCacheFactoryNode->GetCustomGroomAssetPath(GroomAssetPath))
+	{
+		// Otherwise, the user wanted to import the groom asset
+		bool bSuccess = false;
+		TArray<FString> Dependencies;
+		GroomCacheFactoryNode->GetFactoryDependencies(Dependencies);
+		if (Dependencies.Num() > 0)
+		{
+			const UInterchangeGroomFactoryNode* GroomFactoryNode = Cast<UInterchangeGroomFactoryNode>(Arguments.NodeContainer->GetNode(Dependencies[0]));
+			bSuccess = GroomFactoryNode && GroomFactoryNode->GetCustomReferenceObject(GroomAssetPath);
+		}
+
+		if (!bSuccess)
+		{
+			LogGroomCacheFactoryError(LOCTEXT("GroomCacheFactory_MissingGroomAssetDependency", "Could not retrieve groom asset dependency."));
+			return ImportAssetResult;
+		}
+	}
+
+	UGroomAsset* GroomAssetForCache = Cast<UGroomAsset>(GroomAssetPath.TryLoad());
+	if (!GroomAssetForCache)
+	{
+		LogGroomCacheFactoryError(LOCTEXT("GroomCacheFactory_MissingGroomAsset", "Could not retrieve groom asset."));
+		return ImportAssetResult;
+	}
+
+	// Reconstruct and validate the GroomAnimationInfo
+	FGroomAnimationInfo AnimInfo;
+	int32 Value = 0;
+	if (GroomCacheFactoryNode->GetCustomNumFrames(Value))
+	{
+		AnimInfo.NumFrames = Value;
+	}
+
+	if (GroomCacheFactoryNode->GetCustomStartFrame(Value))
+	{
+		AnimInfo.StartFrame = Value;
+	}
+
+	if (GroomCacheFactoryNode->GetCustomEndFrame(Value))
+	{
+		AnimInfo.EndFrame = Value;
+	}
+
+	double FrameRate = 0;
+	if (GroomCacheFactoryNode->GetCustomFrameRate(FrameRate))
+	{
+		AnimInfo.SecondsPerFrame = static_cast<float>(1 / FrameRate);
+	}
+
+	EInterchangeGroomCacheAttributes Attributes;
+	if (GroomCacheFactoryNode->GetCustomGroomCacheAttributes(Attributes))
+	{
+		AnimInfo.Attributes = static_cast<EGroomCacheAttributes>(Attributes);
+	}
+
+	if (!AnimInfo.IsValid())
+	{
+		LogGroomCacheFactoryError(LOCTEXT("GroomCacheFactory_InvalidAnimInfo", "Invalid number of frames, start/end frame, frame rate or no attribute is animated."));
+		return ImportAssetResult;
+	}
+
+	AnimInfo.Duration = AnimInfo.NumFrames * AnimInfo.SecondsPerFrame;
+	AnimInfo.StartTime = AnimInfo.StartFrame * AnimInfo.SecondsPerFrame;
+	AnimInfo.EndTime = AnimInfo.EndFrame * AnimInfo.SecondsPerFrame;
+
+	UObject* ExistingAsset = Arguments.ReimportObject;
+	if (!ExistingAsset)
+	{
+		FSoftObjectPath ReferenceObject;
+		if (GroomCacheFactoryNode->GetCustomReferenceObject(ReferenceObject))
+		{
+			ExistingAsset = ReferenceObject.TryLoad();
+		}
+	}
+
+	UGroomCache* GroomCache = Cast<UGroomCache>(ExistingAsset);
+	if (!GroomCache)
+	{
+		GroomCache = NewObject<UGroomCache>(Arguments.Parent, *Arguments.AssetName, RF_Public | RF_Standalone);
+	}
+	if (!GroomCache)
+	{
+		LogGroomCacheFactoryError(LOCTEXT("GroomCacheFactory_GroomCreationFailure", "Could not allocate new GroomCache."));
+		return ImportAssetResult;
+	}
+
+	static_assert((int)EGroomCacheType::None == (int)EInterchangeGroomCacheImportType::None);
+	static_assert((int)EGroomCacheType::Strands == (int)EInterchangeGroomCacheImportType::Strands);
+	static_assert((int)EGroomCacheType::Guides == (int)EInterchangeGroomCacheImportType::Guides);
+
+	// Initialize groom cache
+	EInterchangeGroomCacheImportType ImportType;
+	GroomCacheFactoryNode->GetCustomGroomCacheImportType(ImportType);
+	EGroomCacheType CacheType = static_cast<EGroomCacheType>(ImportType);
+	GroomCache->Initialize(CacheType);
+
+	const int32 NumFrames = AnimInfo.NumFrames + 1;
+	FScopedSlowTask SlowTask(NumFrames, LOCTEXT("ImportGroomCache", "Importing GroomCache frames"));
+	SlowTask.MakeDialog();
+
+	const TArray<FHairGroupPlatformData>& GroupPlatformData = GroomAssetForCache->GetHairGroupsPlatformData();
+	const FInterchangeGroomPayloadKey& PayloadKey = OptionalPayloadKey.GetValue();
+
+	FGroomCacheProcessor CacheProcessor(CacheType, AnimInfo.Attributes);
+	bool bSuccess = true;
+	FText ImportErrorMessage;
+
+	// Ref. FGroomCacheImporter::ImportGroomCache
+	// 
+	// Each frame the HairDescription payload is translated and processed into GroomCacheInputData
+	// If the frame groom data fails the validation, the import is aborted
+	// 
+	// Sample one extra frame so that we can interpolate between EndFrame - 1 and EndFrame
+	for (int32 FrameIndex = AnimInfo.StartFrame; FrameIndex < AnimInfo.EndFrame + 1; ++FrameIndex)
+	{
+		TRACE_CPUPROFILER_EVENT_SCOPE(FGroomCacheImporter::ImportGroomCache::OneFrame);
+
+		const uint32 CurrentFrame = FrameIndex - AnimInfo.StartFrame;
+
+		FTextBuilder TextBuilder;
+		TextBuilder.AppendLineFormat(LOCTEXT("ImportGroomCacheFrame", "Importing {2} frame {0} of {1}"), FText::AsNumber(CurrentFrame), FText::AsNumber(NumFrames), FText::FromString(Arguments.AssetName));
+		SlowTask.EnterProgressFrame(1.f, TextBuilder.ToText());
+
+		// Get the payload for this frame
+		FInterchangeGroomPayloadKey FramePayloadKey = PayloadKey;
+		FramePayloadKey.FrameNumber = CurrentFrame;
+		TOptional<UE::Interchange::FGroomPayloadData> GroomPayload = GroomPayloadInterface->GetGroomPayloadData(FramePayloadKey);
+
+		if (!GroomPayload.IsSet())
+		{
+			bSuccess = false;
+			ImportErrorMessage = FText::Format(LOCTEXT("GroomCacheFactory_NoPayload", "Missing payload for frame {0}"), FrameIndex);
+			break;
+		}
+
+		FHairDescriptionGroups HairDescriptionGroups;
+		// Do not add extra control points at the end of curve when hair strip geometry is enabled. This is because groom cache data are 
+		// serialized within the uasset (i.e. do not use intermediate cached/build data), and need the asset to be compatible with or without hair strip geometry.
+		if (!FGroomBuilder::BuildHairDescriptionGroups(GroomPayload->HairDescription, HairDescriptionGroups, false /*bAllowAddEndControlPoint*/))
+		{
+			bSuccess = false;
+			ImportErrorMessage = LOCTEXT("GroomCacheFactory_HairDescriptionGroupsBuildFailure", "Failed to build HairDescriptionGroups.");
+			break;
+		}
+
+		const uint32 GroupCount = HairDescriptionGroups.HairGroups.Num();
+		if (GroupCount != GroupPlatformData.Num())
+		{
+			bSuccess = false;
+			ImportErrorMessage = FText::Format(LOCTEXT("GroomCacheFactory_GroupCountMismatch", "GroomCache does not have the same number of groups as the static groom ({0} instead of {1})."), GroupCount, GroupPlatformData.Num());
+			break;
+		}
+
+		TArray<FHairGroupInfoWithVisibility> HairGroupsInfo = GroomAssetForCache->GetHairGroupsInfo();
+		TArray<FGroomCacheInputData> CacheInputDatas;
+		CacheInputDatas.SetNum(GroupCount);
+		for (uint32 GroupIndex = 0; GroupIndex < GroupCount; ++GroupIndex)
+		{
+			FGroomBuilder::BuildData(
+				HairDescriptionGroups.HairGroups[GroupIndex],
+				GroomAssetForCache->GetHairGroupsInterpolation()[GroupIndex],
+				HairGroupsInfo[GroupIndex],
+				CacheInputDatas[GroupIndex].Strands,
+				CacheInputDatas[GroupIndex].Guides);
+		}
+
+		// Validate that the GroomCache has the same topology as the static groom
+		for (uint32 GroupIndex = 0; GroupIndex < GroupCount; ++GroupIndex)
+		{
+			if (CacheType == EGroomCacheType::Strands)
+			{
+				// When UsesTriangleStrips is enabled, we add an extra control point at the end of each curve in the groom asset.
+				// Since the groom cache needs to be independent of these settings (since the data is serialized directly into the asset
+				// and not rebuilt from a description), it does not contains these extra control points.
+				// We account for them for validation here
+				const bool bHasExtraPoint = GroupPlatformData[GroupIndex].Strands.BulkData.HasFlags(FHairStrandsBulkData::DataFlags_HasExtraEndingPoint);
+				const uint32 ExtraControlPointCount = bHasExtraPoint ? CacheInputDatas[GroupIndex].Strands.GetNumCurves() : 0u;
+
+				const bool bCurveCountMatch = CacheInputDatas[GroupIndex].Strands.GetNumCurves() == GroupPlatformData[GroupIndex].Strands.BulkData.GetNumCurves();
+				const bool bPointCountMatch = CacheInputDatas[GroupIndex].Strands.GetNumPoints() + ExtraControlPointCount == GroupPlatformData[GroupIndex].Strands.BulkData.GetNumPoints();
+
+				if (!bCurveCountMatch || !bPointCountMatch)
+				{
+					bSuccess = false;
+					ImportErrorMessage = FText::Format(LOCTEXT("GroomCacheFactory_StrandsValidationFailure", 
+						"GroomCache frame {0} does not have the same number of curves ({1}) or vertices ({2}) for the strands as the static groom ({3} and {4} respectively)."),
+						FrameIndex, CacheInputDatas[GroupIndex].Strands.GetNumCurves(), CacheInputDatas[GroupIndex].Strands.GetNumPoints() + ExtraControlPointCount,
+						GroupPlatformData[GroupIndex].Strands.BulkData.GetNumCurves(), GroupPlatformData[GroupIndex].Strands.BulkData.GetNumPoints());
+					break;
+				}
+			}
+			else if (CacheType == EGroomCacheType::Guides)
+			{
+				const bool bHasExtraPoint = GroupPlatformData[GroupIndex].Guides.BulkData.HasFlags(FHairStrandsBulkData::DataFlags_HasExtraEndingPoint);
+				const uint32 ExtraControlPointCount = bHasExtraPoint ? CacheInputDatas[GroupIndex].Guides.GetNumCurves() : 0u;
+
+				const bool bCurveCountMatch = CacheInputDatas[GroupIndex].Guides.GetNumCurves() == GroupPlatformData[GroupIndex].Guides.BulkData.GetNumCurves();
+				const bool bPointCountMatch = CacheInputDatas[GroupIndex].Guides.GetNumPoints() + ExtraControlPointCount == GroupPlatformData[GroupIndex].Guides.BulkData.GetNumPoints();
+
+				if (!bCurveCountMatch || !bPointCountMatch)
+				{
+					bSuccess = false;
+					ImportErrorMessage = FText::Format(LOCTEXT("GroomCacheFactory_GuidesValidationFailure",
+						"GroomCache frame {0} does not have the same number of curves ({1}) or vertices ({2}) for the guides as the static groom ({3} and {4} respectively)."),
+						FrameIndex, CacheInputDatas[GroupIndex].Guides.GetNumCurves(), CacheInputDatas[GroupIndex].Guides.GetNumPoints() + ExtraControlPointCount,
+						GroupPlatformData[GroupIndex].Guides.BulkData.GetNumCurves(), GroupPlatformData[GroupIndex].Guides.BulkData.GetNumPoints());
+					break;
+				}
+			}
+		}
+
+		if (!bSuccess)
+		{
+			break;
+		}
+
+		CacheProcessor.AddGroomSample(MoveTemp(CacheInputDatas));
+	}
+
+	if (!bSuccess)
+	{
+		if (GroomCache != ExistingAsset)
+		{
+			GroomCache->MarkAsGarbage();
+		}
+
+		LogGroomCacheFactoryError(ImportErrorMessage);
+		return ImportAssetResult;
+	}
+
+	// Finalize groom cache
+	CacheProcessor.TransferChunks(GroomCache);
+	GroomCache->SetGroomAnimationInfo(AnimInfo);
+
+	ImportAssetResult.ImportedObject = GroomCache;
+
+	return ImportAssetResult;
+}
+
+void UInterchangeGroomCacheFactory::SetupObject_GameThread(const FSetupObjectParams& Arguments)
+{
+	TRACE_CPUPROFILER_EVENT_SCOPE(UInterchangeGroomCacheFactory::SetupObject_GameThread);
+
+	Super::SetupObject_GameThread(Arguments);
+
+#if WITH_EDITORONLY_DATA
+	if (ensure(Arguments.ImportedObject && Arguments.SourceData))
+	{
+		// We must call the Update of the asset source file in the main thread because UAssetImportData::Update execute some delegate we do not control
+		UGroomCache* GroomCache = CastChecked<UGroomCache>(Arguments.ImportedObject);
+
+		UE::Interchange::FFactoryCommon::FUpdateImportAssetDataParameters UpdateImportAssetDataParameters(GroomCache, GroomCache->AssetImportData, Arguments.SourceData, Arguments.NodeUniqueID, Arguments.NodeContainer, Arguments.OriginalPipelines, Arguments.Translator);
+		GroomCache->AssetImportData = UE::Interchange::FFactoryCommon::UpdateImportAssetData(UpdateImportAssetDataParameters);
+	}
+#endif
+}
+
+bool UInterchangeGroomCacheFactory::GetSourceFilenames(const UObject* Object, TArray<FString>& OutSourceFilenames) const
+{
+#if WITH_EDITORONLY_DATA
+	if (const UGroomCache* GroomCache = Cast<UGroomCache>(Object))
+	{
+		return UE::Interchange::FFactoryCommon::GetSourceFilenames(GroomCache->AssetImportData, OutSourceFilenames);
+	}
+#endif
+
+	return false;
+}
+
+bool UInterchangeGroomCacheFactory::SetSourceFilename(const UObject* Object, const FString& SourceFilename, int32 SourceIndex) const
+{
+#if WITH_EDITORONLY_DATA
+	if (const UGroomCache* GroomCache = Cast<UGroomCache>(Object))
+	{
+		return UE::Interchange::FFactoryCommon::SetSourceFilename(GroomCache->AssetImportData, SourceFilename, SourceIndex);
+	}
+#endif
+
+	return false;
+}
+
+void UInterchangeGroomCacheFactory::BackupSourceData(const UObject* Object) const
+{
+#if WITH_EDITORONLY_DATA
+	if (const UGroomCache* GroomCache = Cast<UGroomCache>(Object))
+	{
+		UE::Interchange::FFactoryCommon::BackupSourceData(GroomCache->AssetImportData);
+	}
+#endif
+}
+
+void UInterchangeGroomCacheFactory::ReinstateSourceData(const UObject* Object) const
+{
+#if WITH_EDITORONLY_DATA
+	if (const UGroomCache* GroomCache = Cast<UGroomCache>(Object))
+	{
+		UE::Interchange::FFactoryCommon::ReinstateSourceData(GroomCache->AssetImportData);
+	}
+#endif
+}
+
+void UInterchangeGroomCacheFactory::ClearBackupSourceData(const UObject* Object) const
+{
+#if WITH_EDITORONLY_DATA
+	if (const UGroomCache* GroomCache = Cast<UGroomCache>(Object))
+	{
+		UE::Interchange::FFactoryCommon::ClearBackupSourceData(GroomCache->AssetImportData);
+	}
+#endif
+}
+
+#undef LOCTEXT_NAMESPACE
diff --git a/Engine/Plugins/Interchange/Runtime/Source/Import/Private/Groom/InterchangeGroomFactory.cpp b/Engine/Plugins/Interchange/Runtime/Source/Import/Private/Groom/InterchangeGroomFactory.cpp
--- a/Engine/Plugins/Interchange/Runtime/Source/Import/Private/Groom/InterchangeGroomFactory.cpp
+++ b/Engine/Plugins/Interchange/Runtime/Source/Import/Private/Groom/InterchangeGroomFactory.cpp
@@ -156,8 +156,11 @@ UInterchangeFactoryBase::FImportAssetResult UInterchangeGroomFactory::BeginImpor
 	FHairDescriptionGroups GroupsDescription;
 	FGroomBuilder::BuildHairDescriptionGroups(HairDescription, GroupsDescription);
 
+	FHairGroupsInterpolation GroupInterpolationSettings;
+	GroomFactoryNode->GetCustomGroupInterpolationSettings(GroupInterpolationSettings);
+
 	TArray<FHairGroupsInterpolation> GroomInterpolationSettings;
-	GroomInterpolationSettings.SetNum(GroupsDescription.HairGroups.Num());
+	GroomInterpolationSettings.Init(GroupInterpolationSettings, GroupsDescription.HairGroups.Num());
 	UGroomImportOptions* ImportOptions = CreateGroomImportOptions(GroupsDescription, GroomInterpolationSettings);
 
 	FHairImportContext HairImportContext(ImportOptions);
diff --git a/Engine/Plugins/Interchange/Runtime/Source/Import/Private/InterchangeImportModule.cpp b/Engine/Plugins/Interchange/Runtime/Source/Import/Private/InterchangeImportModule.cpp
--- a/Engine/Plugins/Interchange/Runtime/Source/Import/Private/InterchangeImportModule.cpp
+++ b/Engine/Plugins/Interchange/Runtime/Source/Import/Private/InterchangeImportModule.cpp
@@ -10,6 +10,7 @@
 #include "Fbx/InterchangeFbxTranslator.h"
 #include "Gltf/InterchangeGltfTranslator.h"
 #include "Groom/InterchangeGroomFactory.h"
+#include "Groom/InterchangeGroomCacheFactory.h"
 #include "InterchangeImportLog.h"
 #include "InterchangeManager.h"
 #include "Material/InterchangeMaterialFactory.h"
@@ -111,6 +112,7 @@ void FInterchangeImportModule::StartupModule()
 		InterchangeManager.RegisterFactory(UInterchangeStaticMeshFactory::StaticClass());
 		InterchangeManager.RegisterFactory(UInterchangeGeometryCacheFactory::StaticClass());
 		InterchangeManager.RegisterFactory(UInterchangeGroomFactory::StaticClass());
+		InterchangeManager.RegisterFactory(UInterchangeGroomCacheFactory::StaticClass());
 		InterchangeManager.RegisterFactory(UInterchangePhysicsAssetFactory::StaticClass());
 		InterchangeManager.RegisterFactory(UInterchangeActorFactory::StaticClass());
 		InterchangeManager.RegisterFactory(UInterchangeLevelFactory::StaticClass());
diff --git a/Engine/Plugins/Interchange/Runtime/Source/Import/Public/Groom/InterchangeGroomCacheFactory.h b/Engine/Plugins/Interchange/Runtime/Source/Import/Public/Groom/InterchangeGroomCacheFactory.h
new file mode 100644
--- /dev/null
+++ b/Engine/Plugins/Interchange/Runtime/Source/Import/Public/Groom/InterchangeGroomCacheFactory.h
@@ -0,0 +1,31 @@
+// Copyright Epic Games, Inc. All Rights Reserved.
+
+#pragma once
+
+#include "InterchangeFactoryBase.h"
+#include "UObject/ObjectMacros.h"
+
+#include "InterchangeGroomCacheFactory.generated.h"
+
+#define UE_API INTERCHANGEIMPORT_API
+
+UCLASS(MinimalAPI, BlueprintType)
+class UInterchangeGroomCacheFactory : public UInterchangeFactoryBase
+{
+	GENERATED_BODY()
+
+public:
+	// Begin UInterchangeFactoryBase interface
+	UE_API virtual UClass* GetFactoryClass() const override;
+	virtual EInterchangeFactoryAssetType GetFactoryAssetType() override	{ return EInterchangeFactoryAssetType::Grooms; }
+	UE_API virtual FImportAssetResult BeginImportAsset_GameThread(const FImportAssetObjectParams& Arguments) override;
+	UE_API virtual void SetupObject_GameThread(const FSetupObjectParams& Arguments) override;
+	UE_API virtual bool GetSourceFilenames(const UObject* Object, TArray<FString>& OutSourceFilenames) const override;
+	UE_API virtual bool SetSourceFilename(const UObject* Object, const FString& SourceFilename, int32 SourceIndex) const override;
+	UE_API virtual void BackupSourceData(const UObject* Object) const override;
+	UE_API virtual void ReinstateSourceData(const UObject* Object) const override;
+	UE_API virtual void ClearBackupSourceData(const UObject* Object) const override;
+	// End UInterchangeFactoryBase interface
+};
+
+#undef UE_API
\ No newline at end of file
diff --git a/Engine/Plugins/Interchange/Runtime/Source/Nodes/Private/InterchangeGroomNode.cpp b/Engine/Plugins/Interchange/Runtime/Source/Nodes/Private/InterchangeGroomNode.cpp
--- a/Engine/Plugins/Interchange/Runtime/Source/Nodes/Private/InterchangeGroomNode.cpp
+++ b/Engine/Plugins/Interchange/Runtime/Source/Nodes/Private/InterchangeGroomNode.cpp
@@ -86,3 +86,53 @@ void UInterchangeGroomNode::SetPayloadKey(const FString& PayloadKey, EInterchang
 		LogAttributeStorageErrors(Result, TEXT("UInterchangeGroomNode.SetPayloadTypeKey"), UE::Interchange::FGroomNodeStaticData::PayloadTypeKey());
 	}
 }
+
+bool UInterchangeGroomNode::GetCustomStartFrame(int32& AttributeValue) const
+{
+	IMPLEMENT_NODE_ATTRIBUTE_GETTER(StartFrame, int32);
+}
+
+bool UInterchangeGroomNode::SetCustomStartFrame(const int32& AttributeValue)
+{
+	IMPLEMENT_NODE_ATTRIBUTE_SETTER_NODELEGATE(StartFrame, int32);
+}
+
+bool UInterchangeGroomNode::GetCustomEndFrame(int32& AttributeValue) const
+{
+	IMPLEMENT_NODE_ATTRIBUTE_GETTER(EndFrame, int32);
+}
+
+bool UInterchangeGroomNode::SetCustomEndFrame(const int32& AttributeValue)
+{
+	IMPLEMENT_NODE_ATTRIBUTE_SETTER_NODELEGATE(EndFrame, int32);
+}
+
+bool UInterchangeGroomNode::GetCustomNumFrames(int32& AttributeValue) const
+{
+	IMPLEMENT_NODE_ATTRIBUTE_GETTER(NumFrames, int32);
+}
+
+bool UInterchangeGroomNode::SetCustomNumFrames(const int32& AttributeValue)
+{
+	IMPLEMENT_NODE_ATTRIBUTE_SETTER_NODELEGATE(NumFrames, int32);
+}
+
+bool UInterchangeGroomNode::GetCustomFrameRate(double& AttributeValue) const
+{
+	IMPLEMENT_NODE_ATTRIBUTE_GETTER(FrameRate, double);
+}
+
+bool UInterchangeGroomNode::SetCustomFrameRate(const double& AttributeValue)
+{
+	IMPLEMENT_NODE_ATTRIBUTE_SETTER_NODELEGATE(FrameRate, double);
+}
+
+bool UInterchangeGroomNode::GetCustomGroomCacheAttributes(EInterchangeGroomCacheAttributes& AttributeValue) const
+{
+	IMPLEMENT_NODE_ATTRIBUTE_GETTER(GroomCacheAttributes, EInterchangeGroomCacheAttributes);
+}
+
+bool UInterchangeGroomNode::SetCustomGroomCacheAttributes(const EInterchangeGroomCacheAttributes& AttributeValue)
+{
+	IMPLEMENT_NODE_ATTRIBUTE_SETTER_NODELEGATE(GroomCacheAttributes, EInterchangeGroomCacheAttributes);
+}
diff --git a/Engine/Plugins/Interchange/Runtime/Source/Nodes/Public/InterchangeGroomNode.h b/Engine/Plugins/Interchange/Runtime/Source/Nodes/Public/InterchangeGroomNode.h
--- a/Engine/Plugins/Interchange/Runtime/Source/Nodes/Public/InterchangeGroomNode.h
+++ b/Engine/Plugins/Interchange/Runtime/Source/Nodes/Public/InterchangeGroomNode.h
@@ -24,6 +24,21 @@ enum class EInterchangeGroomPayLoadType : uint8
 	ANIMATED	UMETA(DisplayName = "Groom Cache")
 };
 
+UENUM(BlueprintType)
+enum class EInterchangeGroomCacheAttributes : uint8
+{
+	None = 0,
+	Position = 1,
+	Width = 1 << 1,
+	Color = 1 << 2,
+	
+	// For display names
+	PositionWidth = (Position | Width) UMETA(DisplayName = "Position & Width"),
+	PositionColor = (Position | Color) UMETA(DisplayName = "Position & Color"),
+	WidthColor = (Position | Color) UMETA(DisplayName = "Width & Color"),
+	PositionWidthColor = (Position | Width | Color) UMETA(DisplayName = "Position, Width, Color")
+};
+
 USTRUCT(BlueprintType)
 struct FInterchangeGroomPayloadKey
 {
@@ -78,6 +93,53 @@ public:
 
 	UFUNCTION(BlueprintCallable, Category = "Interchange | Node | Groom")
 	UE_API void SetPayloadKey(const FString& PayloadKey, EInterchangeGroomPayLoadType PayLoadType = EInterchangeGroomPayLoadType::STATIC);
+
+	/** Get the start frame index of the animation */
+	UFUNCTION(BlueprintCallable, Category = "Interchange | Node | Groom Cache")
+	UE_API bool GetCustomStartFrame(int32& AttributeValue) const;
+
+	/** Set the start frame index of the animation */
+	UFUNCTION(BlueprintCallable, Category = "Interchange | Node | Groom Cache")
+	UE_API bool SetCustomStartFrame(const int32& AttributeValue);
+
+	/** Get the end frame index of the animation */
+	UFUNCTION(BlueprintCallable, Category = "Interchange | Node | Groom Cache")
+	UE_API bool GetCustomEndFrame(int32& AttributeValue) const;
+
+	/** Set the end frame index of the animation */
+	UFUNCTION(BlueprintCallable, Category = "Interchange | Node | Groom Cache")
+	UE_API bool SetCustomEndFrame(const int32& AttributeValue);
+
+	/** Get the number of frames in the animation */
+	UFUNCTION(BlueprintCallable, Category = "Interchange | Node | Groom Cache")
+	UE_API bool GetCustomNumFrames(int32& AttributeValue) const;
+
+	/** Set the number of frames in the animation */
+	UFUNCTION(BlueprintCallable, Category = "Interchange | Node | Groom Cache")
+	UE_API bool SetCustomNumFrames(const int32& AttributeValue);
+
+	/** Get the animation's frame rate */
+	UFUNCTION(BlueprintCallable, Category = "Interchange | Node | Groom Cache")
+	UE_API bool GetCustomFrameRate(double& AttributeValue) const;
+
+	/** Set the animation's frame rate */
+	UFUNCTION(BlueprintCallable, Category = "Interchange | Node | Groom Cache")
+	UE_API bool SetCustomFrameRate(const double& AttributeValue);
+
+	/** Get the groom attributes that are animated */
+	UFUNCTION(BlueprintCallable, Category = "Interchange | Node | Groom Cache")
+	UE_API bool GetCustomGroomCacheAttributes(EInterchangeGroomCacheAttributes& AttributeValue) const;
+
+	/** Set the groom attributes that are animated */
+	UFUNCTION(BlueprintCallable, Category = "Interchange | Node | Groom Cache")
+	UE_API bool SetCustomGroomCacheAttributes(const EInterchangeGroomCacheAttributes& AttributeValue);
+
+private:
+	IMPLEMENT_NODE_ATTRIBUTE_KEY(NumFrames);
+	IMPLEMENT_NODE_ATTRIBUTE_KEY(StartFrame);
+	IMPLEMENT_NODE_ATTRIBUTE_KEY(EndFrame);
+	IMPLEMENT_NODE_ATTRIBUTE_KEY(FrameRate);
+	IMPLEMENT_NODE_ATTRIBUTE_KEY(GroomCacheAttributes);
 };
 
 #undef UE_API
\ No newline at end of file
diff --git a/Engine/Plugins/Interchange/Runtime/Source/Pipelines/Private/InterchangeGenericAssetsPipeline.cpp b/Engine/Plugins/Interchange/Runtime/Source/Pipelines/Private/InterchangeGenericAssetsPipeline.cpp
--- a/Engine/Plugins/Interchange/Runtime/Source/Pipelines/Private/InterchangeGenericAssetsPipeline.cpp
+++ b/Engine/Plugins/Interchange/Runtime/Source/Pipelines/Private/InterchangeGenericAssetsPipeline.cpp
@@ -147,6 +147,11 @@ bool UInterchangeGenericAssetsPipeline::IsSettingsAreValid(TOptional<FText>& Out
 		return false;
 	}
 
+	if (GroomPipeline && !GroomPipeline->IsSettingsAreValid(OutInvalidReason))
+	{
+		return false;
+	}
+
 	if (AnimationPipeline && !AnimationPipeline->IsSettingsAreValid(OutInvalidReason))
 	{
 		return false;
@@ -432,7 +437,8 @@ bool UInterchangeGenericAssetsPipeline::IsPropertyChangeNeedRefresh(const FPrope
 		|| (CommonSkeletalMeshesAndAnimationsProperties && CommonSkeletalMeshesAndAnimationsProperties->IsPropertyChangeNeedRefresh(PropertyChangedEvent))
 		|| (MeshPipeline && MeshPipeline->IsPropertyChangeNeedRefresh(PropertyChangedEvent))
 		|| (MaterialPipeline && MaterialPipeline->IsPropertyChangeNeedRefresh(PropertyChangedEvent))
-		|| (AnimationPipeline && AnimationPipeline->IsPropertyChangeNeedRefresh(PropertyChangedEvent)))
+		|| (AnimationPipeline && AnimationPipeline->IsPropertyChangeNeedRefresh(PropertyChangedEvent))
+		|| (GroomPipeline && GroomPipeline->IsPropertyChangeNeedRefresh(PropertyChangedEvent)))
 	{
 		return true;
 	}
diff --git a/Engine/Plugins/Interchange/Runtime/Source/Pipelines/Private/InterchangeGenericGroomPipeline.cpp b/Engine/Plugins/Interchange/Runtime/Source/Pipelines/Private/InterchangeGenericGroomPipeline.cpp
--- a/Engine/Plugins/Interchange/Runtime/Source/Pipelines/Private/InterchangeGenericGroomPipeline.cpp
+++ b/Engine/Plugins/Interchange/Runtime/Source/Pipelines/Private/InterchangeGenericGroomPipeline.cpp
@@ -3,7 +3,9 @@
 #include "InterchangeGenericGroomPipeline.h"
 
 #include "GroomAsset.h"
+#include "GroomCache.h"
 #include "InterchangeGroomNode.h"
+#include "InterchangeGroomCacheFactoryNode.h"
 #include "InterchangeGroomFactoryNode.h"
 #include "InterchangePipelineHelper.h"
 #include "Nodes/InterchangeFactoryBaseNode.h"
@@ -12,6 +14,78 @@
 
 #include UE_INLINE_GENERATED_CPP_BY_NAME(InterchangeGenericGroomPipeline)
 
+namespace UE::InterchangeGroomPipeline::Private
+{
+	UInterchangeFactoryBaseNode* CreateGroomFactoryNode(UInterchangeBaseNodeContainer* InBaseNodeContainer, const UInterchangeGroomNode* GroomNode, EInterchangeGroomCacheImportType Type)
+	{
+		FString NodeID = GroomNode->GetUniqueID();
+		FString FactoryNodeUid;
+
+		UInterchangeFactoryBaseNode* FactoryNode = nullptr;
+		if (Type == EInterchangeGroomCacheImportType::None)
+		{
+			// Create groom asset node
+			FactoryNodeUid = UInterchangeFactoryBaseNode::BuildFactoryNodeUid(NodeID);
+			FactoryNode = NewObject<UInterchangeGroomFactoryNode>(InBaseNodeContainer);
+		}
+		else
+		{
+			// Create groom cache node
+			const FString CacheType = Type == EInterchangeGroomCacheImportType::Strands ? TEXT("_strands_") : TEXT("_guides_");
+			static const FString Suffix(TEXT("cache"));
+
+			FactoryNodeUid = UInterchangeFactoryBaseNode::BuildFactoryNodeUid(NodeID + CacheType + Suffix);
+			UInterchangeGroomCacheFactoryNode* GroomCacheFactoryNode = NewObject<UInterchangeGroomCacheFactoryNode>(InBaseNodeContainer);
+
+			GroomCacheFactoryNode->SetCustomGroomCacheImportType(Type);
+			GroomCacheFactoryNode->SetAssetName(GroomNode->GetAssetName() + CacheType + Suffix);
+
+			int32 Value = 0;
+			if (GroomNode->GetCustomNumFrames(Value))
+			{
+				GroomCacheFactoryNode->SetCustomNumFrames(Value);
+			}
+
+			if (GroomNode->GetCustomStartFrame(Value))
+			{
+				GroomCacheFactoryNode->SetCustomStartFrame(Value);
+			}
+
+			if (GroomNode->GetCustomEndFrame(Value))
+			{
+				GroomCacheFactoryNode->SetCustomEndFrame(Value);
+			}
+
+			double SecsPerFrames = 0;
+			if (GroomNode->GetCustomFrameRate(SecsPerFrames))
+			{
+				GroomCacheFactoryNode->SetCustomFrameRate(SecsPerFrames);
+			}
+
+			EInterchangeGroomCacheAttributes Attributes;
+			if (GroomNode->GetCustomGroomCacheAttributes(Attributes))
+			{
+				GroomCacheFactoryNode->SetCustomGroomCacheAttributes(Attributes);
+			}
+
+			FactoryNode = GroomCacheFactoryNode;
+		}
+
+		InBaseNodeContainer->SetupNode(FactoryNode, FactoryNodeUid, GroomNode->GetDisplayLabel(), EInterchangeNodeContainerType::FactoryData);
+
+		UInterchangeSourceNode* SourceNode = UInterchangeSourceNode::FindOrCreateUniqueInstance(InBaseNodeContainer);
+		UE::Interchange::PipelineHelper::FillSubPathFromSourceNode(FactoryNode, SourceNode);
+
+		const bool bAddSourceNodeName = false;
+		UInterchangeUserDefinedAttributesAPI::DuplicateAllUserDefinedAttribute(GroomNode, FactoryNode, bAddSourceNodeName);
+
+		FactoryNode->AddTargetNodeUid(GroomNode->GetUniqueID());
+		GroomNode->AddTargetNodeUid(FactoryNode->GetUniqueID());
+
+		return FactoryNode;
+	}
+}
+
 FString UInterchangeGenericGroomPipeline::GetPipelineCategory(UClass* AssetClass)
 {
 	return TEXT("Grooms");
@@ -19,6 +93,8 @@ FString UInterchangeGenericGroomPipeline::GetPipelineCategory(UClass* AssetClass
 
 void UInterchangeGenericGroomPipeline::ExecutePipeline(UInterchangeBaseNodeContainer* InBaseNodeContainer, const TArray<UInterchangeSourceData*>& SourceDatas, const FString& ContentBasePath)
 {
+	using namespace UE::InterchangeGroomPipeline::Private;
+
 	if (!bEnableGroomTypesImport)
 	{
 		return;
@@ -31,7 +107,7 @@ void UInterchangeGenericGroomPipeline::ExecutePipeline(UInterchangeBaseNodeConta
 
 	// Find all the translated groom nodes we need for this pipeline
 	TArray<UInterchangeGroomNode*> GroomNodes;
-	if (bImportGrooms)
+	if (bImportGrooms || bImportGroomCaches)
 	{
 		InBaseNodeContainer->IterateNodes(
 			[&GroomNodes](const FString& NodeUid, UInterchangeBaseNode* Node)
@@ -46,19 +122,83 @@ void UInterchangeGenericGroomPipeline::ExecutePipeline(UInterchangeBaseNodeConta
 		// Create groom factory nodes
 		for (const UInterchangeGroomNode* GroomNode : GroomNodes)
 		{
-			const FString FactoryNodeUid = UInterchangeFactoryBaseNode::BuildFactoryNodeUid(GroomNode->GetUniqueID());
+			const TOptional<FInterchangeGroomPayloadKey> PayloadKey = GroomNode->GetPayloadKey();
+			if (!PayloadKey.IsSet())
+			{
+				continue;
+			}
+
+			if (bImportGrooms && PayloadKey->Type == EInterchangeGroomPayLoadType::STATIC)
+			{
+				UInterchangeGroomFactoryNode* GroomFactoryNode = Cast<UInterchangeGroomFactoryNode>(CreateGroomFactoryNode(InBaseNodeContainer, GroomNode, EInterchangeGroomCacheImportType::None));
+				if (GroomFactoryNode)
+				{
+					GroomFactoryNode->SetCustomGroupInterpolationSettings(GroupInterpolationSettings);
+				}
+			}
+			else
+			{
+				UInterchangeGroomFactoryNode* GroomFactoryNode = nullptr;
+				if (bImportGrooms)
+				{
+					// A groom asset can be generated from the first frame of a groom cache
+					GroomFactoryNode = Cast<UInterchangeGroomFactoryNode>(CreateGroomFactoryNode(InBaseNodeContainer, GroomNode, EInterchangeGroomCacheImportType::None));
+					if (GroomFactoryNode)
+					{
+						GroomFactoryNode->SetCustomGroupInterpolationSettings(GroupInterpolationSettings);
+					}
+				}
 
-			UInterchangeGroomFactoryNode* GroomFactoryNode = NewObject<UInterchangeGroomFactoryNode>(InBaseNodeContainer);
-			InBaseNodeContainer->SetupNode(GroomFactoryNode, FactoryNodeUid, GroomNode->GetDisplayLabel(), EInterchangeNodeContainerType::FactoryData);
+				if (bImportGroomCaches)
+				{
+					// To import a groom cache, either a groom asset must be imported from its first frame or a previously imported groom asset must be specified
+					// The groom asset is used to validate that the cache is compatible with the asset
+					if (!bImportGrooms && !GroomAsset.IsValid())
+					{
+						continue;
+					}
 
-			UInterchangeSourceNode* SourceNode = UInterchangeSourceNode::FindOrCreateUniqueInstance(InBaseNodeContainer);
-			UE::Interchange::PipelineHelper::FillSubPathFromSourceNode(GroomFactoryNode, SourceNode);
+					auto InitializeGroomCacheFactoryNode = [this, InBaseNodeContainer, GroomNode, GroomFactoryNode](EInterchangeGroomCacheImportType CacheType)
+					{
+						if (EnumHasAnyFlags(ImportGroomCacheType, CacheType))
+						{
+							UInterchangeGroomCacheFactoryNode* GroomCacheFactoryNode = Cast<UInterchangeGroomCacheFactoryNode>(CreateGroomFactoryNode(InBaseNodeContainer, GroomNode, CacheType));
+							if (GroomCacheFactoryNode)
+							{
+								if (GroomFactoryNode)
+								{
+									// If the groom asset is imported at the same time, set its factory node as a dependency
+									GroomCacheFactoryNode->AddFactoryDependencyUid(GroomFactoryNode->GetUniqueID());
+									if (bIsReimportContext && Cast<UGroomCache>(CacheContextParam.ReimportAsset))
+									{
+										// When reimporting a groom cache, the associated groom asset will not be reimported
+										// but it's still needed as a dependency. Disabling its factory node will try to
+										// fetch it instead of importing it
+										GroomFactoryNode->SetEnabled(false);
+									}
+								}
+								else
+								{
+									// Not importing the groom asset, so set the path to a previously imported groom asset
+									GroomCacheFactoryNode->SetCustomGroomAssetPath(GroomAsset);
+								}
 
-			const bool bAddSourceNodeName = false;
-			UInterchangeUserDefinedAttributesAPI::DuplicateAllUserDefinedAttribute(GroomNode, GroomFactoryNode, bAddSourceNodeName);
+								if (bOverrideTimeRange)
+								{
+									GroomCacheFactoryNode->SetCustomStartFrame(FrameStart);
+									// Make sure the end is after the start and it's at least 2 frames
+									FrameEnd = FMath::Max(FrameEnd, FrameStart + 1);
+									GroomCacheFactoryNode->SetCustomEndFrame(FrameEnd);
+									GroomCacheFactoryNode->SetCustomNumFrames(FrameEnd - FrameStart + 1);
+								}
+							}
+						}
+					};
 
-			GroomFactoryNode->AddTargetNodeUid(GroomNode->GetUniqueID());
-			GroomNode->AddTargetNodeUid(GroomFactoryNode->GetUniqueID());
+					InitializeGroomCacheFactoryNode(EInterchangeGroomCacheImportType::Strands);
+					InitializeGroomCacheFactoryNode(EInterchangeGroomCacheImportType::Guides);
+				}
+			}
 		}
 	}
 }
@@ -68,6 +208,80 @@ void UInterchangeGenericGroomPipeline::ExecutePipeline(UInterchangeBaseNodeConta
 void UInterchangeGenericGroomPipeline::GetSupportAssetClasses(TArray<UClass*>& PipelineSupportAssetClasses) const
 {
 	PipelineSupportAssetClasses.Add(UGroomAsset::StaticClass());
+	PipelineSupportAssetClasses.Add(UGroomCache::StaticClass());
+}
+
+void UInterchangeGenericGroomPipeline::FilterPropertiesFromTranslatedData(UInterchangeBaseNodeContainer* InBaseNodeContainer)
+{
+	int32 NumGroomNodes = 0;
+	int32 NumGroomCaches = 0;
+	InBaseNodeContainer->BreakableIterateNodesOfType<UInterchangeGroomNode>([&NumGroomNodes, &NumGroomCaches](const FString& NodeUid, UInterchangeGroomNode* GroomNode)
+	{
+		++NumGroomNodes;
+		const TOptional<FInterchangeGroomPayloadKey> PayloadKey = GroomNode->GetPayloadKey();
+		if (PayloadKey.IsSet() && PayloadKey->Type == EInterchangeGroomPayLoadType::ANIMATED)
+		{
+			++NumGroomCaches;
+			return true;
+		}
+		return false;
+	});
+
+	UInterchangePipelineBase* OuterMostPipeline = GetMostPipelineOuter();
+	if (NumGroomNodes == 0)
+	{
+		HidePropertiesOfCategory(OuterMostPipeline, this, GetPipelineCategory(nullptr));
+	}
+	else if (NumGroomCaches == 0)
+	{
+		HidePropertiesOfSubCategory(OuterMostPipeline, this, TEXT("Caches"));
+	}
 }
 
-#endif // WITH_EDITOR
\ No newline at end of file
+bool UInterchangeGenericGroomPipeline::IsPropertyChangeNeedRefresh(const FPropertyChangedEvent& PropertyChangedEvent) const
+{
+	static const TSet<FName> NeedRefreshProperties =
+	{
+		GET_MEMBER_NAME_CHECKED(UInterchangeGenericGroomPipeline, bEnableGroomTypesImport),
+		GET_MEMBER_NAME_CHECKED(UInterchangeGenericGroomPipeline, bImportGroomCaches),
+		GET_MEMBER_NAME_CHECKED(UInterchangeGenericGroomPipeline, bImportGrooms),
+		GET_MEMBER_NAME_CHECKED(UInterchangeGenericGroomPipeline, GroomAsset),
+	};
+
+	if (NeedRefreshProperties.Contains(PropertyChangedEvent.GetPropertyName()))
+	{
+		return true;
+	}
+
+	return Super::IsPropertyChangeNeedRefresh(PropertyChangedEvent);
+}
+
+#endif //WITH_EDITOR
+
+bool UInterchangeGenericGroomPipeline::IsSettingsAreValid(TOptional<FText>& OutInvalidReason) const
+{
+	if (bEnableGroomTypesImport && bImportGroomCaches && !bImportGrooms && !GroomAsset.IsValid())
+	{
+		int32 NumGroomCaches = 0;
+		CacheContextParam.BaseNodeContainer->BreakableIterateNodesOfType<UInterchangeGroomNode>([&NumGroomCaches](const FString& NodeUid, UInterchangeGroomNode* GroomNode)
+		{
+			const TOptional<FInterchangeGroomPayloadKey> PayloadKey = GroomNode->GetPayloadKey();
+			if (PayloadKey.IsSet() && PayloadKey->Type == EInterchangeGroomPayLoadType::ANIMATED)
+			{
+				++NumGroomCaches;
+				return true;
+			}
+			return false;
+		});
+
+		// If there's no groom cache, then this validation is not needed
+		if (NumGroomCaches == 0)
+		{
+			return Super::IsSettingsAreValid(OutInvalidReason);
+		}
+
+		OutInvalidReason = NSLOCTEXT("UInterchangeGenericGroomPipeline", "GroomAssetMustBeSpecified", "When importing a groom cache without importing its associated groom asset, a previously imported and compatible groom asset must be specified.");
+		return false;
+	}
+	return Super::IsSettingsAreValid(OutInvalidReason);
+}
diff --git a/Engine/Plugins/Interchange/Runtime/Source/Pipelines/Private/InterchangePipelineHelper.cpp b/Engine/Plugins/Interchange/Runtime/Source/Pipelines/Private/InterchangePipelineHelper.cpp
--- a/Engine/Plugins/Interchange/Runtime/Source/Pipelines/Private/InterchangePipelineHelper.cpp
+++ b/Engine/Plugins/Interchange/Runtime/Source/Pipelines/Private/InterchangePipelineHelper.cpp
@@ -13,6 +13,7 @@
 #include "Framework/Docking/TabManager.h"
 #include "GeometryCache.h"
 #include "GroomAsset.h"
+#include "GroomCache.h"
 #include "FoliageType_InstancedStaticMesh.h"
 #include "HAL/PlatformApplicationMisc.h"
 #include "InterchangePipelineBase.h"
@@ -125,7 +126,7 @@ namespace UE::Interchange::PipelineHelper
 			{
 				Suffix = TEXT("GeometryCaches");
 			}
-			else if (Class->IsChildOf(UGroomAsset::StaticClass()))
+			else if (Class->IsChildOf(UGroomAsset::StaticClass()) || Class->IsChildOf(UGroomCache::StaticClass()))
 			{
 				Suffix = TEXT("Grooms");
 			}
diff --git a/Engine/Plugins/Interchange/Runtime/Source/Pipelines/Public/InterchangeGenericGroomPipeline.h b/Engine/Plugins/Interchange/Runtime/Source/Pipelines/Public/InterchangeGenericGroomPipeline.h
--- a/Engine/Plugins/Interchange/Runtime/Source/Pipelines/Public/InterchangeGenericGroomPipeline.h
+++ b/Engine/Plugins/Interchange/Runtime/Source/Pipelines/Public/InterchangeGenericGroomPipeline.h
@@ -3,6 +3,9 @@
 #pragma once
 
 #include "InterchangePipelineBase.h"
+
+#include "GroomAssetInterpolation.h"
+#include "InterchangeGroomCacheFactoryNode.h"
 #include "Nodes/InterchangeBaseNodeContainer.h"
 #include "UObject/ObjectMacros.h"
 
@@ -27,9 +30,34 @@ public:
 	bool bEnableGroomTypesImport = false;
 
 	/** If enabled, import all groom assets found in the source. */
-	UPROPERTY(EditAnywhere, BlueprintReadWrite, Category = "Grooms")
+	UPROPERTY(EditAnywhere, BlueprintReadWrite, Category = "Grooms", meta = (EditCondition = bEnableGroomTypesImport))
 	bool bImportGrooms = true;
 
+	/** Settings that will be applied to all hair groups in the groom asset */
+	UPROPERTY(EditAnywhere, BlueprintReadWrite, Category = "Grooms")
+	FHairGroupsInterpolation GroupInterpolationSettings;
+
+	/** If enabled, import all groom caches found in the source. */
+	UPROPERTY(EditAnywhere, BlueprintReadWrite, Category = "Grooms", meta = (SubCategory = "Caches", EditCondition = bEnableGroomTypesImport))
+	bool bImportGroomCaches = true;
+
+	/** If the groom asset is not imported, provide an existing one against which the groom cache will be validated */
+	UPROPERTY(EditAnywhere, BlueprintReadWrite, Category = "Grooms", meta = (SubCategory = "Caches", EditCondition = "bImportGrooms == false", MetaClass = "/Script/HairStrandsCore.GroomAsset"))
+	FSoftObjectPath GroomAsset;
+
+	/** Groom Cache types to import */
+	UPROPERTY(EditAnywhere, BlueprintReadWrite, Category = "Grooms", meta = (SubCategory = "Caches", EditCondition = bImportGroomCaches))
+	EInterchangeGroomCacheImportType ImportGroomCacheType = EInterchangeGroomCacheImportType::All;
+
+	UPROPERTY(EditAnywhere, BlueprintReadWrite, Category = "Grooms", meta = (SubCategory = "Caches"))
+	bool bOverrideTimeRange = false;
+
+	UPROPERTY(EditAnywhere, BlueprintReadWrite, Category = "Grooms", meta = (SubCategory = "Caches", EditCondition = bOverrideTimeRange))
+	int32 FrameStart = 0;
+
+	UPROPERTY(EditAnywhere, BlueprintReadWrite, Category = "Grooms", meta = (SubCategory = "Caches", EditCondition = bOverrideTimeRange))
+	int32 FrameEnd = 1;
+
 protected:
 
 	UE_API virtual void ExecutePipeline(UInterchangeBaseNodeContainer* BaseNodeContainer, const TArray<UInterchangeSourceData*>& SourceDatas, const FString& ContentBasePath) override;
@@ -37,7 +65,11 @@ protected:
 public:
 #if WITH_EDITOR
 	UE_API virtual void GetSupportAssetClasses(TArray<UClass*>& PipelineSupportAssetClasses) const override;
+	UE_API virtual void FilterPropertiesFromTranslatedData(UInterchangeBaseNodeContainer* InBaseNodeContainer) override;
+	UE_API virtual bool IsPropertyChangeNeedRefresh(const FPropertyChangedEvent& PropertyChangedEvent) const override;
 #endif
+
+	UE_API virtual bool IsSettingsAreValid(TOptional<FText>& OutInvalidReason) const override;
 };
 
 #undef UE_API
\ No newline at end of file
diff --git a/Engine/Plugins/Runtime/USDCore/Source/USDUtilities/Private/USDGroomConversion.cpp b/Engine/Plugins/Runtime/USDCore/Source/USDUtilities/Private/USDGroomConversion.cpp
--- a/Engine/Plugins/Runtime/USDCore/Source/USDUtilities/Private/USDGroomConversion.cpp
+++ b/Engine/Plugins/Runtime/USDCore/Source/USDUtilities/Private/USDGroomConversion.cpp
@@ -381,9 +381,10 @@ namespace UsdToUnreal
 
 			// Following the USD recommendation that gprims be not nested, a Curves prim is handled as a leaf
 		}
-		else if (pxr::UsdGeomImageable(Prim))
+		else
 		{
-			// UsdGeomImageable includes UsdGeomXformable and UsdGeomScope
+			// Traverse any prim type
+			// Extract the transform if it's a UsdGeomXformable
 			if (pxr::UsdGeomXform Xform = pxr::UsdGeomXform(Prim))
 			{
 				// Propagate the prim transform to the children
diff --git a/Engine/Source/Editor/UnrealEd/Private/Editor.cpp b/Engine/Source/Editor/UnrealEd/Private/Editor.cpp
--- a/Engine/Source/Editor/UnrealEd/Private/Editor.cpp
+++ b/Engine/Source/Editor/UnrealEd/Private/Editor.cpp
@@ -1010,6 +1010,9 @@ void FReimportManager::GetNewReimportPath(UObject* Obj, TArray<FString>& InOutFi
 					case EInterchangeFactoryAssetType::Textures:
 						TranslatorAssetType = EInterchangeTranslatorAssetType::Textures;
 						break;
+					case EInterchangeFactoryAssetType::Grooms:
+						TranslatorAssetType = EInterchangeTranslatorAssetType::Grooms;
+						break;
 				};
 				TranslatorFormats = UInterchangeManager::GetInterchangeManager().GetSupportedAssetTypeFormats(TranslatorAssetType, EInterchangeTranslatorType::Assets, false /*bStrictMatchTranslatorType*/);
 			}
diff --git a/Engine/Source/Runtime/Interchange/Core/Public/InterchangeTranslatorBase.h b/Engine/Source/Runtime/Interchange/Core/Public/InterchangeTranslatorBase.h
--- a/Engine/Source/Runtime/Interchange/Core/Public/InterchangeTranslatorBase.h
+++ b/Engine/Source/Runtime/Interchange/Core/Public/InterchangeTranslatorBase.h
@@ -59,6 +59,7 @@ enum class EInterchangeTranslatorAssetType : uint8
 	Meshes = 0x1 << 2,
 	Animations = 0x1 << 3,
 	Sounds = 0x1 << 4,
+	Grooms = 0x1 << 5,
 };
 ENUM_CLASS_FLAGS(EInterchangeTranslatorAssetType);
 
diff --git a/Engine/Source/Runtime/Interchange/Engine/Private/InterchangeManager.cpp b/Engine/Source/Runtime/Interchange/Engine/Private/InterchangeManager.cpp
--- a/Engine/Source/Runtime/Interchange/Engine/Private/InterchangeManager.cpp
+++ b/Engine/Source/Runtime/Interchange/Engine/Private/InterchangeManager.cpp
@@ -1435,6 +1435,9 @@ TArray<FString> UInterchangeManager::GetSupportedFormatsForObject(const UObject*
 				case EInterchangeFactoryAssetType::Textures:
 					FileExtensions = GetSupportedAssetTypeFormats(EInterchangeTranslatorAssetType::Textures, TranslatorType, bStrictMatchTranslatorType);
 					break;
+				case EInterchangeFactoryAssetType::Grooms:
+					FileExtensions = GetSupportedAssetTypeFormats(EInterchangeTranslatorAssetType::Grooms, TranslatorType, bStrictMatchTranslatorType);
+					break;
 				case EInterchangeFactoryAssetType::None: //Actor factories return None
 					FileExtensions = GetSupportedFormats(EInterchangeTranslatorType::Actors);
 					break;
diff --git a/Engine/Source/Runtime/Interchange/Engine/Private/Tasks/InterchangeTaskImportObject.cpp b/Engine/Source/Runtime/Interchange/Engine/Private/Tasks/InterchangeTaskImportObject.cpp
--- a/Engine/Source/Runtime/Interchange/Engine/Private/Tasks/InterchangeTaskImportObject.cpp
+++ b/Engine/Source/Runtime/Interchange/Engine/Private/Tasks/InterchangeTaskImportObject.cpp
@@ -376,7 +376,13 @@ void UE::Interchange::FTaskImportObject_GameThread::Execute()
 		Private::InternalGetPackageName(*AsyncHelper, SourceIndex, PackageBasePath, FactoryNode, PackageName, AssetName);
 		if (bRefObjectIsMap || !FPackageUtils::IsMapPackageAsset(PackageName))
 		{
+			//Try to find the package in memory
 			Pkg = FindPackage(nullptr, *PackageName);
+			if (!Pkg)
+			{
+				//Try to load the package from disk
+				Pkg = LoadPackage(nullptr, *PackageName, LOAD_NoWarn | LOAD_Quiet);
+			}
 			if (Pkg)
 			{
 				FollowRedirectorCode();
@@ -749,4 +755,4 @@ void UE::Interchange::FTaskImportObjectFinalize_GameThread::Execute()
 		}
 		FactoryNode->SetCustomReferenceObject(FSoftObjectPath(ImportAssetResult.ImportedObject));
 	}
-}
\ No newline at end of file
+}
