Interchange USD: Minor fixes to groom cache import code.

#rnx

diff --git a/Engine/Plugins/Interchange/Runtime/Source/Import/Private/Groom/InterchangeGroomCacheFactory.cpp b/Engine/Plugins/Interchange/Runtime/Source/Import/Private/Groom/InterchangeGroomCacheFactory.cpp
--- a/Engine/Plugins/Interchange/Runtime/Source/Import/Private/Groom/InterchangeGroomCacheFactory.cpp
+++ b/Engine/Plugins/Interchange/Runtime/Source/Import/Private/Groom/InterchangeGroomCacheFactory.cpp
@@ -141,7 +141,15 @@ UInterchangeFactoryBase::FImportAssetResult UInterchangeGroomCacheFactory::Begin
 	double FrameRate = 0;
 	if (GroomCacheFactoryNode->GetCustomFrameRate(FrameRate))
 	{
-		AnimInfo.SecondsPerFrame = static_cast<float>(1 / FrameRate);
+		if (FrameRate != 0)
+		{
+			AnimInfo.SecondsPerFrame = static_cast<float>(1 / FrameRate);
+		}
+		else
+		{
+			// Fallback to 24 fps
+			AnimInfo.SecondsPerFrame = static_cast<float>(1.0 / 24.0);
+		}
 	}
 
 	EInterchangeGroomCacheAttributes Attributes;
@@ -187,7 +195,11 @@ UInterchangeFactoryBase::FImportAssetResult UInterchangeGroomCacheFactory::Begin
 
 	// Initialize groom cache
 	EInterchangeGroomCacheImportType ImportType;
-	GroomCacheFactoryNode->GetCustomGroomCacheImportType(ImportType);
+	if (!GroomCacheFactoryNode->GetCustomGroomCacheImportType(ImportType))
+	{
+		LogGroomCacheFactoryError(LOCTEXT("GroomCacheFactory_NoImportType", "No GroomCache type was specified for import."));
+		return ImportAssetResult;
+	}
 	EGroomCacheType CacheType = static_cast<EGroomCacheType>(ImportType);
 	GroomCache->Initialize(CacheType);
 
diff --git a/Engine/Plugins/Interchange/Runtime/Source/Nodes/Public/InterchangeGroomNode.h b/Engine/Plugins/Interchange/Runtime/Source/Nodes/Public/InterchangeGroomNode.h
--- a/Engine/Plugins/Interchange/Runtime/Source/Nodes/Public/InterchangeGroomNode.h
+++ b/Engine/Plugins/Interchange/Runtime/Source/Nodes/Public/InterchangeGroomNode.h
@@ -35,7 +35,7 @@ enum class EInterchangeGroomCacheAttributes : uint8
 	// For display names
 	PositionWidth = (Position | Width) UMETA(DisplayName = "Position & Width"),
 	PositionColor = (Position | Color) UMETA(DisplayName = "Position & Color"),
-	WidthColor = (Position | Color) UMETA(DisplayName = "Width & Color"),
+	WidthColor = (Width | Color) UMETA(DisplayName = "Width & Color"),
 	PositionWidthColor = (Position | Width | Color) UMETA(DisplayName = "Position, Width, Color")
 };
 
diff --git a/Engine/Plugins/Runtime/HairStrands/Source/HairStrandsCore/Public/GroomCacheData.h b/Engine/Plugins/Runtime/HairStrands/Source/HairStrandsCore/Public/GroomCacheData.h
--- a/Engine/Plugins/Runtime/HairStrands/Source/HairStrandsCore/Public/GroomCacheData.h
+++ b/Engine/Plugins/Runtime/HairStrands/Source/HairStrandsCore/Public/GroomCacheData.h
@@ -22,7 +22,7 @@ enum class EGroomCacheAttributes : uint8
 	// For display names
 	PositionWidth = (Position | Width) UMETA(DisplayName = "Position & Width"),
 	PositionColor = (Position | Color) UMETA(DisplayName = "Position & Color"),
-	WidthColor = (Position | Color) UMETA(DisplayName = "Width & Color"),
+	WidthColor = (Width | Color) UMETA(DisplayName = "Width & Color"),
 	PositionWidthColor = (Position | Width | Color) UMETA(DisplayName = "Position, Width, Color")
 };
 
