[Interchange Audio] Fixes for Automation Tests and fixes for wundef compliance
- Made sure that the Automation Test codepath would not be tied to the #if WITH_EDITOR to be able to run the tests
- Added a temp fix for the #if WITH_SNDFILE_IO. This will be revisited when the Interchange and Sound Factory will resort to using the common Audio File Import Utils.

#rnx

diff --git a/Engine/Plugins/Interchange/Extensions/OpenUSD/Source/Import/Private/InterchangeOpenUSDImportModule.cpp b/Engine/Plugins/Interchange/Extensions/OpenUSD/Source/Import/Private/InterchangeOpenUSDImportModule.cpp
--- a/Engine/Plugins/Interchange/Extensions/OpenUSD/Source/Import/Private/InterchangeOpenUSDImportModule.cpp
+++ b/Engine/Plugins/Interchange/Extensions/OpenUSD/Source/Import/Private/InterchangeOpenUSDImportModule.cpp
@@ -45,6 +45,7 @@ void FInterchangeOpenUSDImportModule::StartupModule()
 			TranslatorPipelines.Pipelines.Add(FSoftObjectPath(TEXT("/InterchangeOpenUSD/Pipelines/DefaultUSDAssetsPipeline.DefaultUSDAssetsPipeline")));
 			TranslatorPipelines.Pipelines.Add(FSoftObjectPath(TEXT("/InterchangeOpenUSD/Pipelines/DefaultUSDPipelineAssetImport.DefaultUSDPipelineAssetImport")));
 			TranslatorPipelines.Pipelines.Add(FSoftObjectPath(TEXT("/Interchange/Pipelines/DefaultMaterialXPipeline.DefaultMaterialXPipeline")));
+			TranslatorPipelines.Pipelines.Add(FSoftObjectPath(TEXT("/Interchange/Pipelines/DefaultAudioPipeline.DefaultAudioPipeline")));
 
 			FInterchangePipelineStack& PipelineStack = AssetImportSettings.PipelineStacks[TEXT("Assets")];
 			PipelineStack.PerTranslatorPipelines.Add(TranslatorPipelines);
diff --git a/Engine/Plugins/Interchange/Extensions/OpenUSD/Source/Import/Private/InterchangeUsdTranslator.cpp b/Engine/Plugins/Interchange/Extensions/OpenUSD/Source/Import/Private/InterchangeUsdTranslator.cpp
--- a/Engine/Plugins/Interchange/Extensions/OpenUSD/Source/Import/Private/InterchangeUsdTranslator.cpp
+++ b/Engine/Plugins/Interchange/Extensions/OpenUSD/Source/Import/Private/InterchangeUsdTranslator.cpp
@@ -106,6 +106,14 @@ static FAutoConsoleVariableRef CVarInterchangeEnableUSDLevelImport(
 	TEXT("Whether support for USD level import is enabled.")
 );
 
+
+static bool GInterchangeEnableUSDAudioImport = true;
+static FAutoConsoleVariableRef CVarInterchangeEnableUSDAudioImport(
+	TEXT("Interchange.FeatureFlags.Import.USD.Audio"),
+	GInterchangeEnableUSDAudioImport,
+	TEXT("Whether support for USD Media Spatial Audio is enabled.")
+);
+
 namespace UE::InterchangeUsdTranslator::Private
 {
 #if USE_USD_SDK
@@ -4250,7 +4258,20 @@ namespace UE::InterchangeUsdTranslator::Private
 		}
 		else if (Prim.IsA(TEXT("SpatialAudio")))
 		{
-			AssetNodeUids = { TranslatorImpl.AddSpatialAudioNode(Prim, NodeContainer) };
+			if (GInterchangeEnableUSDAudioImport)
+			{
+				AssetNodeUids = { TranslatorImpl.AddSpatialAudioNode(Prim, NodeContainer) };
+			}
+			else
+			{
+				UInterchangeResultWarning_Generic* Message = TranslatorImpl.ResultsContainer->Add<UInterchangeResultWarning_Generic>();
+				Message->SourceAssetName = TranslatorImpl.UsdStage.GetRootLayer().GetRealPath();
+				Message->Text = FText::Format(
+					LOCTEXT("UnsupportedSchemaType", "Prim '{0}' has schema '{1}', which is not yet supported via USD Interchange."),
+					FText::FromString(SceneNodeUid),
+					FText::FromName(TypeName)
+				);
+			}
 		}
 		else
 		{
diff --git a/Engine/Plugins/Interchange/Runtime/Source/Import/Private/Audio/InterchangeAudioPayloadData.cpp b/Engine/Plugins/Interchange/Runtime/Source/Import/Private/Audio/InterchangeAudioPayloadData.cpp
--- a/Engine/Plugins/Interchange/Runtime/Source/Import/Private/Audio/InterchangeAudioPayloadData.cpp
+++ b/Engine/Plugins/Interchange/Runtime/Source/Import/Private/Audio/InterchangeAudioPayloadData.cpp
@@ -6,6 +6,12 @@
 
 #include "InterchangeImportLog.h"
 
+// TODO: This should be revisited when the code will 
+// rely on the common Audio File Import Utils used by both the Sound Factory and Interchange.
+#ifndef WITH_SNDFILE_IO
+	#define WITH_SNDFILE_IO 0
+#endif
+
 namespace UE::Interchange::Private
 {
 	bool ConvertAudioFile(const TArray<uint8>& RawWaveData, FWaveModInfo& WaveInfo, TArray<uint8>& ConvertedRawWaveData, FString& ErrorMessage, bool bEnsureEqualNumSamples = true)
@@ -89,7 +95,7 @@ TOptional<UE::Interchange::FInterchangeAudioPayloadData> FInterchangeAudioPayloa
 	if (bFileNeedsConversion || *WaveInfo.pBitsPerSample != 16 || !WaveInfo.IsFormatSupported())
 	{
 		TArray<uint8> ConvertedRawWaveData;
-#if defined(WITH_SNDFILE_IO) && WITH_SNDFILE_IO
+#if WITH_SNDFILE_IO
 		const bool bEnsureEqualNumSamples = !bFileNeedsConversion;
 		if (!Private::ConvertAudioFile(ReadBuffer, WaveInfo, ConvertedRawWaveData, ErrorMessage, bEnsureEqualNumSamples))
 		{
diff --git a/Engine/Plugins/Interchange/Runtime/Source/Import/Private/Audio/InterchangeGenericAudioTranslator.cpp b/Engine/Plugins/Interchange/Runtime/Source/Import/Private/Audio/InterchangeGenericAudioTranslator.cpp
--- a/Engine/Plugins/Interchange/Runtime/Source/Import/Private/Audio/InterchangeGenericAudioTranslator.cpp
+++ b/Engine/Plugins/Interchange/Runtime/Source/Import/Private/Audio/InterchangeGenericAudioTranslator.cpp
@@ -5,6 +5,12 @@
 
 #define LOCTEXT_NAMESPACE "InterchangeGenericAudioTranslator"
 
+// TODO: This should be revisited when the code will 
+// rely on the common Audio File Import Utils used by both the Sound Factory and Interchange.
+#ifndef WITH_SNDFILE_IO
+	#define WITH_SNDFILE_IO 0
+#endif
+
 static bool GInterchangeEnableAudio = false;
 static FAutoConsoleVariableRef CCvarInterchangeEnableAudioSupport(
 	TEXT("Interchange.FeatureFlags.Import.Audio"),
@@ -16,8 +22,7 @@ TArray<FString> UInterchangeGenericAudioTranslator::GetSupportedFormats() const
 {
 	TArray<FString> Formats;
 
-#if WITH_EDITOR
-	if (GInterchangeEnableAudio || GIsAutomationTesting)
+	if (GInterchangeEnableAudio)
 	{
 #if WITH_SNDFILE_IO
 		Formats.Add(TEXT("aif;Audio Interchange File"));
@@ -27,10 +32,8 @@ TArray<FString> UInterchangeGenericAudioTranslator::GetSupportedFormats() const
 		Formats.Add(TEXT("opus;OGG OPUS bitstream format"));
 		Formats.Add(TEXT("mp3;MPEG Layer 3 Audio"));
 #endif // WITH_SNDFILE_IO
-
 		Formats.Add(TEXT("wav;Wave Audio File"));
 	}
-#endif
 	return Formats;
 }
 
diff --git a/Engine/Plugins/Interchange/Runtime/Source/Import/Private/InterchangeImportModule.cpp b/Engine/Plugins/Interchange/Runtime/Source/Import/Private/InterchangeImportModule.cpp
--- a/Engine/Plugins/Interchange/Runtime/Source/Import/Private/InterchangeImportModule.cpp
+++ b/Engine/Plugins/Interchange/Runtime/Source/Import/Private/InterchangeImportModule.cpp
@@ -13,6 +13,8 @@
 #include "Groom/InterchangeGroomCacheFactory.h"
 #include "InterchangeImportLog.h"
 #include "InterchangeManager.h"
+#include "Audio/InterchangeAudioSoundWaveFactory.h"
+#include "Audio/InterchangeGenericAudioTranslator.h"
 #include "Material/InterchangeMaterialFactory.h"
 #include "MaterialX/InterchangeMaterialXTranslator.h"
 #include "Mesh/InterchangeGeometryCacheFactory.h"
@@ -92,6 +94,9 @@ void FInterchangeImportModule::StartupModule()
 		InterchangeManager.RegisterTranslator(UInterchangeGLTFTranslator::StaticClass());
 		InterchangeManager.RegisterTranslator(UInterchangeOBJTranslator::StaticClass());
 
+		// Audio
+		InterchangeManager.RegisterTranslator(UInterchangeGenericAudioTranslator::StaticClass());
+
 		//Materials
 		InterchangeManager.RegisterTranslator(UInterchangeMaterialXTranslator::StaticClass());
 
@@ -130,6 +135,7 @@ void FInterchangeImportModule::StartupModule()
 		InterchangeManager.RegisterFactory(UInterchangeSceneImportAssetFactory::StaticClass());
 		InterchangeManager.RegisterFactory(UInterchangeSpecularProfileFactory::StaticClass());
 		InterchangeManager.RegisterFactory(UInterchangeHeterogeneousVolumeActorFactory::StaticClass());
+		InterchangeManager.RegisterFactory(UInterchangeAudioSoundWaveFactory::StaticClass());
 #if WITH_EDITOR
 		// Editor-only because SVT internals like UStreamableSparseVolumeTexture::AppendFrame are editor-only
 		InterchangeManager.RegisterFactory(UInterchangeSparseVolumeTextureFactory::StaticClass());
