[Interchange Audio] Split the supported audio formats into their own separate translators

#rnx

diff --git a/Engine/Plugins/Interchange/Runtime/Source/Import/Private/Audio/Formats/InterchangeAudioFormatsPrivate.cpp b/Engine/Plugins/Interchange/Runtime/Source/Import/Private/Audio/Formats/InterchangeAudioFormatsPrivate.cpp
new file mode 100644
--- /dev/null
+++ b/Engine/Plugins/Interchange/Runtime/Source/Import/Private/Audio/Formats/InterchangeAudioFormatsPrivate.cpp
@@ -0,0 +1,97 @@
+// Copyright Epic Games, Inc. All Rights Reserved.
+#include "Audio/Formats/InterchangeAudioFormatsPrivate.h"
+#include "HAL/IConsoleManager.h"
+
+namespace Interchange::Audio
+{
+	const FString& FAudioFormatCVarNames::GetAudioFormatCVarName_AIF()
+	{ 
+		static const FString AudioFormatCVarName = TEXT("Interchange.FeatureFlags.Import.Audio.AIF");
+		return AudioFormatCVarName;
+	}
+
+	const FString& FAudioFormatCVarNames::GetAudioFormatCVarName_AIFF()
+	{ 
+		static const FString AudioFormatCVarName = TEXT("Interchange.FeatureFlags.Import.Audio.AIFF");
+		return AudioFormatCVarName;
+	}
+
+	const FString& FAudioFormatCVarNames::GetAudioFormatCVarName_OGG()
+	{ 
+		static const FString AudioFormatCVarName = TEXT("Interchange.FeatureFlags.Import.Audio.OGG");
+		return AudioFormatCVarName;
+	}
+
+	const FString& FAudioFormatCVarNames::GetAudioFormatCVarName_FLAC()
+	{ 
+		static const FString AudioFormatCVarName = TEXT("Interchange.FeatureFlags.Import.Audio.FLAC");
+		return AudioFormatCVarName;
+	}
+
+	const FString& FAudioFormatCVarNames::GetAudioFormatCVarName_OPUS()
+	{ 
+		static const FString AudioFormatCVarName = TEXT("Interchange.FeatureFlags.Import.Audio.OPUS");
+		return AudioFormatCVarName;
+	}
+
+	const FString& FAudioFormatCVarNames::GetAudioFormatCVarName_MP3()
+	{ 
+		static const FString AudioFormatCVarName = TEXT("Interchange.FeatureFlags.Import.Audio.MP3");
+		return AudioFormatCVarName;
+	}
+
+	const FString& FAudioFormatCVarNames::GetAudioFormatCVarName_WAV() 
+	{ 
+		static const FString AudioFormatCVarName = TEXT("Interchange.FeatureFlags.Import.Audio.WAV");
+		return AudioFormatCVarName;
+	}
+}
+
+static bool GInterchangeEnableAIFImport = false;
+static FAutoConsoleVariableRef CCvarInterchangeEnableAIFImport(
+	*Interchange::Audio::FAudioFormatCVarNames::GetAudioFormatCVarName_AIF(),
+	GInterchangeEnableAIFImport,
+	TEXT("Whether AIF support is enabled."),
+	ECVF_Default);
+
+static bool GInterchangeEnableAIFFImport = false;
+static FAutoConsoleVariableRef CCvarInterchangeEnableAIFFImport(
+	*Interchange::Audio::FAudioFormatCVarNames::GetAudioFormatCVarName_AIFF(),
+	GInterchangeEnableAIFFImport,
+	TEXT("Whether AIFF support is enabled."),
+	ECVF_Default);
+
+static bool GInterchangeEnableOGGImport = false;
+static FAutoConsoleVariableRef CCvarInterchangeEnableOGGImport(
+	*Interchange::Audio::FAudioFormatCVarNames::GetAudioFormatCVarName_OGG(),
+	GInterchangeEnableOGGImport,
+	TEXT("Whether OGG support is enabled."),
+	ECVF_Default);
+
+static bool GInterchangeEnableFLACImport = false;
+static FAutoConsoleVariableRef CCvarInterchangeEnableFLACImport(
+	*Interchange::Audio::FAudioFormatCVarNames::GetAudioFormatCVarName_FLAC(),
+	GInterchangeEnableFLACImport,
+	TEXT("Whether FLAC support is enabled."),
+	ECVF_Default);
+
+static bool GInterchangeEnableOPUSImport = false;
+static FAutoConsoleVariableRef CCvarInterchangeEnableOPUSImport(
+	*Interchange::Audio::FAudioFormatCVarNames::GetAudioFormatCVarName_OPUS(),
+	GInterchangeEnableOPUSImport,
+	TEXT("Whether OPUS support is enabled."),
+	ECVF_Default);
+
+static bool GInterchangeEnableMP3Import = false;
+static FAutoConsoleVariableRef CCvarInterchangeEnableMP3Import(
+	*Interchange::Audio::FAudioFormatCVarNames::GetAudioFormatCVarName_MP3(),
+	GInterchangeEnableMP3Import,
+	TEXT("Whether MP3 support is enabled."),
+	ECVF_Default);
+
+static bool GInterchangeEnableWAVImport = false;
+static FAutoConsoleVariableRef CCvarInterchangeEnableWAVImport(
+	*Interchange::Audio::FAudioFormatCVarNames::GetAudioFormatCVarName_WAV(),
+	GInterchangeEnableWAVImport,
+	TEXT("Whether WAV support is enabled."),
+	ECVF_Default);
\ No newline at end of file
diff --git a/Engine/Plugins/Interchange/Runtime/Source/Import/Private/Audio/Formats/InterchangeAudioFormatsPrivate.h b/Engine/Plugins/Interchange/Runtime/Source/Import/Private/Audio/Formats/InterchangeAudioFormatsPrivate.h
new file mode 100644
--- /dev/null
+++ b/Engine/Plugins/Interchange/Runtime/Source/Import/Private/Audio/Formats/InterchangeAudioFormatsPrivate.h
@@ -0,0 +1,24 @@
+// Copyright Epic Games, Inc. All Rights Reserved.
+#pragma once
+#include "Containers/UnrealString.h"
+
+// TODO: This should be revisited when the code will 
+// rely on the common Audio File Import Utils used by both the Sound Factory and Interchange.
+#ifndef WITH_SNDFILE_IO
+	#define WITH_SNDFILE_IO 0
+#endif
+
+namespace Interchange::Audio
+{
+	struct FAudioFormatCVarNames
+	{
+	public:
+		static const FString& GetAudioFormatCVarName_AIF();
+		static const FString& GetAudioFormatCVarName_AIFF();
+		static const FString& GetAudioFormatCVarName_OGG();
+		static const FString& GetAudioFormatCVarName_FLAC();
+		static const FString& GetAudioFormatCVarName_OPUS();
+		static const FString& GetAudioFormatCVarName_MP3();
+		static const FString& GetAudioFormatCVarName_WAV();
+	};
+}
\ No newline at end of file
diff --git a/Engine/Plugins/Interchange/Runtime/Source/Import/Private/Audio/Formats/InterchangeAudioTranslatorBase.cpp b/Engine/Plugins/Interchange/Runtime/Source/Import/Private/Audio/Formats/InterchangeAudioTranslatorBase.cpp
new file mode 100644
--- /dev/null
+++ b/Engine/Plugins/Interchange/Runtime/Source/Import/Private/Audio/Formats/InterchangeAudioTranslatorBase.cpp
@@ -0,0 +1,87 @@
+// Copyright Epic Games, Inc. All Rights Reserved.
+#include "Audio/Formats/InterchangeAudioTranslatorBase.h"
+#include "InterchangeAudioSoundWaveNode.h"
+
+#define LOCTEXT_NAMESPACE "InterchangeGenericAudioTranslator"
+
+TArray<FString> UInterchangeAudioTranslatorBase::GetSupportedFormats() const
+{
+	return TArray<FString>();
+}
+
+EInterchangeTranslatorAssetType UInterchangeAudioTranslatorBase::GetSupportedAssetTypes() const
+{
+	return EInterchangeTranslatorAssetType::Sounds;
+}
+
+bool UInterchangeAudioTranslatorBase::Translate(UInterchangeBaseNodeContainer& BaseNodeContainer) const
+{
+	UInterchangeAudioSoundWaveNode* TranslatedNode = Translate_Internal(BaseNodeContainer);
+	return TranslatedNode != nullptr;
+}
+
+UInterchangeAudioSoundWaveNode* UInterchangeAudioTranslatorBase::Translate_Internal(UInterchangeBaseNodeContainer& BaseNodeContainer) const
+{
+	FString Filename = SourceData->GetFilename();
+	FPaths::NormalizeFilename(Filename);
+	if (!FPaths::FileExists(Filename))
+	{
+		LogError(
+			FText::Format(
+				LOCTEXT("InterchangeGenericAudioTranslator_Translate_FileNotFound", "Filepath {0} doesn't exist"),
+				FText::FromString(Filename)
+			)
+		);
+
+		return nullptr;
+	}
+
+	UInterchangeAudioSoundWaveNode* SoundWaveNode = UInterchangeAudioSoundWaveNode::Create(BaseNodeContainer, FPaths::GetBaseFilename(Filename));
+	SoundWaveNode->SetPayloadKey(Filename);
+
+	return SoundWaveNode;
+}
+
+void UInterchangeAudioTranslatorBase::LogError(FText&& ErrorText) const
+{
+	if (UInterchangeResultError_Generic* Result = AddMessage<UInterchangeResultError_Generic>())
+	{
+		Result->AssetType = GetClass();
+
+		if (SourceData)
+		{
+			const FString Filename = SourceData->GetFilename();
+			Result->SourceAssetName = Filename;
+			Result->InterchangeKey = FPaths::GetBaseFilename(Filename);
+		}
+		else
+		{
+			Result->InterchangeKey = TEXT("Undefined");
+		}
+
+		Result->Text = MoveTemp(ErrorText);
+	}
+}
+
+void UInterchangeAudioTranslatorBase::LogWarning(FText&& WarningText) const
+{
+	if (UInterchangeResultWarning_Generic* Result = AddMessage<UInterchangeResultWarning_Generic>())
+	{
+		Result->AssetType = GetClass();
+
+		if (SourceData)
+		{
+			const FString Filename = SourceData->GetFilename();
+			Result->SourceAssetName = Filename;
+			Result->InterchangeKey = FPaths::GetBaseFilename(Filename);
+		}
+		else
+		{
+			Result->InterchangeKey = TEXT("Undefined");
+		}
+
+		Result->Text = MoveTemp(WarningText);
+	}
+}
+
+#undef LOCTEXT_NAMESPACE
\ No newline at end of file
diff --git a/Engine/Plugins/Interchange/Runtime/Source/Import/Private/Audio/Formats/InterchangeAudioTranslator_AIF.cpp b/Engine/Plugins/Interchange/Runtime/Source/Import/Private/Audio/Formats/InterchangeAudioTranslator_AIF.cpp
new file mode 100644
--- /dev/null
+++ b/Engine/Plugins/Interchange/Runtime/Source/Import/Private/Audio/Formats/InterchangeAudioTranslator_AIF.cpp
@@ -0,0 +1,21 @@
+// Copyright Epic Games, Inc. All Rights Reserved.
+#include "Audio/Formats/InterchangeAudioTranslator_AIF.h"
+
+#include "Audio/Formats/InterchangeAudioFormatsPrivate.h"
+#include "HAL/IConsoleManager.h"
+
+TArray<FString> UInterchangeAudioTranslator_AIF::GetSupportedFormats() const
+{
+	using namespace Interchange::Audio;
+
+	TArray<FString> Formats;
+#if WITH_SNDFILE_IO
+	const IConsoleVariable* CVarFormatEnabled_AIF = IConsoleManager::Get().FindConsoleVariable(*FAudioFormatCVarNames::GetAudioFormatCVarName_AIF());
+	const bool bFormatEnabled_AIF = CVarFormatEnabled_AIF && CVarFormatEnabled_AIF->GetBool();
+	if(bFormatEnabled_AIF)
+	{
+		Formats.Add(TEXT("aif;Audio Interchange File"));
+	}
+#endif
+	return Formats;
+}
\ No newline at end of file
diff --git a/Engine/Plugins/Interchange/Runtime/Source/Import/Private/Audio/Formats/InterchangeAudioTranslator_AIFF.cpp b/Engine/Plugins/Interchange/Runtime/Source/Import/Private/Audio/Formats/InterchangeAudioTranslator_AIFF.cpp
new file mode 100644
--- /dev/null
+++ b/Engine/Plugins/Interchange/Runtime/Source/Import/Private/Audio/Formats/InterchangeAudioTranslator_AIFF.cpp
@@ -0,0 +1,21 @@
+// Copyright Epic Games, Inc. All Rights Reserved.
+#include "Audio/Formats/InterchangeAudioTranslator_AIFF.h"
+
+#include "Audio/Formats/InterchangeAudioFormatsPrivate.h"
+#include "HAL/IConsoleManager.h"
+
+TArray<FString> UInterchangeAudioTranslator_AIFF::GetSupportedFormats() const
+{
+	using namespace Interchange::Audio;
+
+	TArray<FString> Formats;
+#if WITH_SNDFILE_IO
+	const IConsoleVariable* CVarFormatEnabled_AIFF = IConsoleManager::Get().FindConsoleVariable(*FAudioFormatCVarNames::GetAudioFormatCVarName_AIFF());
+	const bool bFormatEnabled_AIFF = CVarFormatEnabled_AIFF && CVarFormatEnabled_AIFF->GetBool();
+	if(bFormatEnabled_AIFF)
+	{
+		Formats.Add(TEXT("aiff;Audio Interchange File Format"));
+	}
+#endif
+	return Formats;
+}
\ No newline at end of file
diff --git a/Engine/Plugins/Interchange/Runtime/Source/Import/Private/Audio/Formats/InterchangeAudioTranslator_FLAC.cpp b/Engine/Plugins/Interchange/Runtime/Source/Import/Private/Audio/Formats/InterchangeAudioTranslator_FLAC.cpp
new file mode 100644
--- /dev/null
+++ b/Engine/Plugins/Interchange/Runtime/Source/Import/Private/Audio/Formats/InterchangeAudioTranslator_FLAC.cpp
@@ -0,0 +1,21 @@
+// Copyright Epic Games, Inc. All Rights Reserved.
+#include "Audio/Formats/InterchangeAudioTranslator_FLAC.h"
+
+#include "Audio/Formats/InterchangeAudioFormatsPrivate.h"
+#include "HAL/IConsoleManager.h"
+
+TArray<FString> UInterchangeAudioTranslator_FLAC::GetSupportedFormats() const
+{
+	using namespace Interchange::Audio;
+
+	TArray<FString> Formats;
+#if WITH_SNDFILE_IO
+	const IConsoleVariable* CVarFormatEnabled_FLAC = IConsoleManager::Get().FindConsoleVariable(*FAudioFormatCVarNames::GetAudioFormatCVarName_FLAC());
+	const bool bFormatEnabled_FLAC = CVarFormatEnabled_FLAC && CVarFormatEnabled_FLAC->GetBool();
+	if(bFormatEnabled_FLAC)
+	{
+		Formats.Add(TEXT("flac;Free Lossless Audio Codec"));
+	}
+#endif
+	return Formats;
+}
\ No newline at end of file
diff --git a/Engine/Plugins/Interchange/Runtime/Source/Import/Private/Audio/Formats/InterchangeAudioTranslator_MP3.cpp b/Engine/Plugins/Interchange/Runtime/Source/Import/Private/Audio/Formats/InterchangeAudioTranslator_MP3.cpp
new file mode 100644
--- /dev/null
+++ b/Engine/Plugins/Interchange/Runtime/Source/Import/Private/Audio/Formats/InterchangeAudioTranslator_MP3.cpp
@@ -0,0 +1,21 @@
+// Copyright Epic Games, Inc. All Rights Reserved.
+#include "Audio/Formats/InterchangeAudioTranslator_MP3.h"
+
+#include "Audio/Formats/InterchangeAudioFormatsPrivate.h"
+#include "HAL/IConsoleManager.h"
+
+TArray<FString> UInterchangeAudioTranslator_MP3::GetSupportedFormats() const
+{
+	using namespace Interchange::Audio;
+
+	TArray<FString> Formats;
+#if WITH_SNDFILE_IO
+	const IConsoleVariable* CVarFormatEnabled_MP3 = IConsoleManager::Get().FindConsoleVariable(*FAudioFormatCVarNames::GetAudioFormatCVarName_MP3());
+	const bool bFormatEnabled_MP3 = CVarFormatEnabled_MP3 && CVarFormatEnabled_MP3->GetBool();
+	if(bFormatEnabled_MP3)
+	{
+		Formats.Add(TEXT("mp3;MPEG Layer 3 Audio"));
+	}
+#endif
+	return Formats;
+}
\ No newline at end of file
diff --git a/Engine/Plugins/Interchange/Runtime/Source/Import/Private/Audio/Formats/InterchangeAudioTranslator_OGG.cpp b/Engine/Plugins/Interchange/Runtime/Source/Import/Private/Audio/Formats/InterchangeAudioTranslator_OGG.cpp
new file mode 100644
--- /dev/null
+++ b/Engine/Plugins/Interchange/Runtime/Source/Import/Private/Audio/Formats/InterchangeAudioTranslator_OGG.cpp
@@ -0,0 +1,21 @@
+// Copyright Epic Games, Inc. All Rights Reserved.
+#include "Audio/Formats/InterchangeAudioTranslator_OGG.h"
+
+#include "Audio/Formats/InterchangeAudioFormatsPrivate.h"
+#include "HAL/IConsoleManager.h"
+
+TArray<FString> UInterchangeAudioTranslator_OGG::GetSupportedFormats() const
+{
+	using namespace Interchange::Audio;
+
+	TArray<FString> Formats;
+#if WITH_SNDFILE_IO
+	const IConsoleVariable* CVarFormatEnabled_OGG = IConsoleManager::Get().FindConsoleVariable(*FAudioFormatCVarNames::GetAudioFormatCVarName_OGG());
+	const bool bFormatEnabled_OGG = CVarFormatEnabled_OGG && CVarFormatEnabled_OGG->GetBool();
+	if(bFormatEnabled_OGG)
+	{
+		Formats.Add(TEXT("ogg;OGG Vorbis bitstream format "));
+	}
+#endif
+	return Formats;
+}
\ No newline at end of file
diff --git a/Engine/Plugins/Interchange/Runtime/Source/Import/Private/Audio/Formats/InterchangeAudioTranslator_OPUS.cpp b/Engine/Plugins/Interchange/Runtime/Source/Import/Private/Audio/Formats/InterchangeAudioTranslator_OPUS.cpp
new file mode 100644
--- /dev/null
+++ b/Engine/Plugins/Interchange/Runtime/Source/Import/Private/Audio/Formats/InterchangeAudioTranslator_OPUS.cpp
@@ -0,0 +1,21 @@
+// Copyright Epic Games, Inc. All Rights Reserved.
+#include "Audio/Formats/InterchangeAudioTranslator_OPUS.h"
+
+#include "Audio/Formats/InterchangeAudioFormatsPrivate.h"
+#include "HAL/IConsoleManager.h"
+
+TArray<FString> UInterchangeAudioTranslator_OPUS::GetSupportedFormats() const
+{
+	using namespace Interchange::Audio;
+
+	TArray<FString> Formats;
+#if WITH_SNDFILE_IO
+	const IConsoleVariable* CVarFormatEnabled_OPUS = IConsoleManager::Get().FindConsoleVariable(*FAudioFormatCVarNames::GetAudioFormatCVarName_OPUS());
+	const bool bFormatEnabled_OPUS = CVarFormatEnabled_OPUS && CVarFormatEnabled_OPUS->GetBool();
+	if(bFormatEnabled_OPUS)
+	{
+		Formats.Add(TEXT("opus;OGG OPUS bitstream format"));
+	}
+#endif
+	return Formats;
+}
\ No newline at end of file
diff --git a/Engine/Plugins/Interchange/Runtime/Source/Import/Private/Audio/Formats/InterchangeAudioTranslator_WAV.cpp b/Engine/Plugins/Interchange/Runtime/Source/Import/Private/Audio/Formats/InterchangeAudioTranslator_WAV.cpp
new file mode 100644
--- /dev/null
+++ b/Engine/Plugins/Interchange/Runtime/Source/Import/Private/Audio/Formats/InterchangeAudioTranslator_WAV.cpp
@@ -0,0 +1,19 @@
+// Copyright Epic Games, Inc. All Rights Reserved.
+#include "Audio/Formats/InterchangeAudioTranslator_WAV.h"
+
+#include "Audio/Formats/InterchangeAudioFormatsPrivate.h"
+#include "HAL/IConsoleManager.h"
+
+TArray<FString> UInterchangeAudioTranslator_WAV::GetSupportedFormats() const
+{
+	using namespace Interchange::Audio;
+
+	TArray<FString> Formats;
+	const IConsoleVariable* CVarFormatEnabled_WAV = IConsoleManager::Get().FindConsoleVariable(*FAudioFormatCVarNames::GetAudioFormatCVarName_WAV());
+	const bool bFormatEnabled_WAV = CVarFormatEnabled_WAV && CVarFormatEnabled_WAV->GetBool();
+	if(bFormatEnabled_WAV)
+	{
+		Formats.Add(TEXT("wav;Wave Audio File"));
+	}
+	return Formats;
+}
\ No newline at end of file
diff --git a/Engine/Plugins/Interchange/Runtime/Source/Import/Private/Audio/InterchangeAudioPayloadData.cpp b/Engine/Plugins/Interchange/Runtime/Source/Import/Private/Audio/InterchangeAudioPayloadData.cpp
--- a/Engine/Plugins/Interchange/Runtime/Source/Import/Private/Audio/InterchangeAudioPayloadData.cpp
+++ b/Engine/Plugins/Interchange/Runtime/Source/Import/Private/Audio/InterchangeAudioPayloadData.cpp
@@ -82,7 +82,7 @@ TOptional<UE::Interchange::FInterchangeAudioPayloadData> FInterchangeAudioPayloa
 	FWaveModInfo WaveInfo;
 	FString ErrorMessage;
 
-	const bool bFileNeedsConversion = !FPaths::GetExtension(*(PayloadSourceFileKey.ToLower())).Equals(TEXT(".wav"));
+	const bool bFileNeedsConversion = !FPaths::GetExtension(*(PayloadSourceFileKey.ToLower())).Equals(TEXT("wav"));
 	if (!WaveInfo.ReadWaveInfo(ReadBuffer.GetData(), ReadBuffer.Num(), &ErrorMessage))
 	{
 		if (!bFileNeedsConversion)
diff --git a/Engine/Plugins/Interchange/Runtime/Source/Import/Private/Audio/InterchangeGenericAudioTranslator.cpp b/Engine/Plugins/Interchange/Runtime/Source/Import/Private/Audio/InterchangeGenericAudioTranslator.cpp
deleted file mode 100644
diff --git a/Engine/Plugins/Interchange/Runtime/Source/Import/Private/InterchangeImportModule.cpp b/Engine/Plugins/Interchange/Runtime/Source/Import/Private/InterchangeImportModule.cpp
--- a/Engine/Plugins/Interchange/Runtime/Source/Import/Private/InterchangeImportModule.cpp
+++ b/Engine/Plugins/Interchange/Runtime/Source/Import/Private/InterchangeImportModule.cpp
@@ -14,7 +14,13 @@
 #include "InterchangeImportLog.h"
 #include "InterchangeManager.h"
 #include "Audio/InterchangeAudioSoundWaveFactory.h"
-#include "Audio/InterchangeGenericAudioTranslator.h"
+#include "Audio/Formats/InterchangeAudioTranslator_AIF.h"
+#include "Audio/Formats/InterchangeAudioTranslator_AIFF.h"
+#include "Audio/Formats/InterchangeAudioTranslator_OGG.h"
+#include "Audio/Formats/InterchangeAudioTranslator_FLAC.h"
+#include "Audio/Formats/InterchangeAudioTranslator_OPUS.h"
+#include "Audio/Formats/InterchangeAudioTranslator_MP3.h"
+#include "Audio/Formats/InterchangeAudioTranslator_WAV.h"
 #include "Material/InterchangeMaterialFactory.h"
 #include "MaterialX/InterchangeMaterialXTranslator.h"
 #include "Mesh/InterchangeGeometryCacheFactory.h"
@@ -95,7 +101,13 @@ void FInterchangeImportModule::StartupModule()
 		InterchangeManager.RegisterTranslator(UInterchangeOBJTranslator::StaticClass());
 
 		// Audio
-		InterchangeManager.RegisterTranslator(UInterchangeGenericAudioTranslator::StaticClass());
+		InterchangeManager.RegisterTranslator(UInterchangeAudioTranslator_AIF::StaticClass());
+		InterchangeManager.RegisterTranslator(UInterchangeAudioTranslator_AIFF::StaticClass());
+		InterchangeManager.RegisterTranslator(UInterchangeAudioTranslator_OGG::StaticClass());
+		InterchangeManager.RegisterTranslator(UInterchangeAudioTranslator_FLAC::StaticClass());
+		InterchangeManager.RegisterTranslator(UInterchangeAudioTranslator_OPUS::StaticClass());
+		InterchangeManager.RegisterTranslator(UInterchangeAudioTranslator_MP3::StaticClass());
+		InterchangeManager.RegisterTranslator(UInterchangeAudioTranslator_WAV::StaticClass());
 
 		//Materials
 		InterchangeManager.RegisterTranslator(UInterchangeMaterialXTranslator::StaticClass());
diff --git a/Engine/Plugins/Interchange/Runtime/Source/Import/Public/Audio/Formats/InterchangeAudioTranslatorBase.h b/Engine/Plugins/Interchange/Runtime/Source/Import/Public/Audio/Formats/InterchangeAudioTranslatorBase.h
new file mode 100644
--- /dev/null
+++ b/Engine/Plugins/Interchange/Runtime/Source/Import/Public/Audio/Formats/InterchangeAudioTranslatorBase.h
@@ -0,0 +1,33 @@
+// Copyright Epic Games, Inc. All Rights Reserved.
+#pragma once
+#include "InterchangeTranslatorBase.h"
+#include "Audio/InterchangeAudioPayloadInterface.h"
+
+#include "InterchangeAudioTranslatorBase.generated.h"
+
+#define UE_API INTERCHANGEIMPORT_API
+
+class UInterchangeAudioSoundWaveNode;
+
+UCLASS(BlueprintType, Abstract)
+class UE_API UInterchangeAudioTranslatorBase
+	: public UInterchangeTranslatorBase
+	, public IInterchangeAudioPayloadInterface
+{
+	GENERATED_BODY()
+
+public:
+	/** Begin UInterchangeTranslatorBase API */
+	virtual TArray<FString> GetSupportedFormats() const override;
+	virtual EInterchangeTranslatorAssetType GetSupportedAssetTypes() const override;
+	virtual bool Translate(UInterchangeBaseNodeContainer& BaseNodeContainer) const override;
+	/** End UInterchangeTranslatorBase API */
+protected:
+	virtual UInterchangeAudioSoundWaveNode* Translate_Internal(UInterchangeBaseNodeContainer& BaseNodeContainer) const;
+
+protected:
+	void LogError(FText&& ErrorText)  const;
+	void LogWarning(FText&& WarningText) const;
+};
+
+#undef UE_API
\ No newline at end of file
diff --git a/Engine/Plugins/Interchange/Runtime/Source/Import/Public/Audio/Formats/InterchangeAudioTranslator_AIF.h b/Engine/Plugins/Interchange/Runtime/Source/Import/Public/Audio/Formats/InterchangeAudioTranslator_AIF.h
new file mode 100644
--- /dev/null
+++ b/Engine/Plugins/Interchange/Runtime/Source/Import/Public/Audio/Formats/InterchangeAudioTranslator_AIF.h
@@ -0,0 +1,17 @@
+// Copyright Epic Games, Inc. All Rights Reserved.
+#pragma once
+#include "Audio/Formats/InterchangeAudioTranslatorBase.h"
+
+#include "InterchangeAudioTranslator_AIF.generated.h"
+
+#define UE_API INTERCHANGEIMPORT_API
+
+UCLASS(BlueprintType)
+class UE_API UInterchangeAudioTranslator_AIF : public UInterchangeAudioTranslatorBase
+{
+	GENERATED_BODY()
+public:
+	virtual TArray<FString> GetSupportedFormats() const override;
+};
+
+#undef UE_API
\ No newline at end of file
diff --git a/Engine/Plugins/Interchange/Runtime/Source/Import/Public/Audio/Formats/InterchangeAudioTranslator_AIFF.h b/Engine/Plugins/Interchange/Runtime/Source/Import/Public/Audio/Formats/InterchangeAudioTranslator_AIFF.h
new file mode 100644
--- /dev/null
+++ b/Engine/Plugins/Interchange/Runtime/Source/Import/Public/Audio/Formats/InterchangeAudioTranslator_AIFF.h
@@ -0,0 +1,17 @@
+// Copyright Epic Games, Inc. All Rights Reserved.
+#pragma once
+#include "Audio/Formats/InterchangeAudioTranslatorBase.h"
+
+#include "InterchangeAudioTranslator_AIFF.generated.h"
+
+#define UE_API INTERCHANGEIMPORT_API
+
+UCLASS(BlueprintType)
+class UE_API UInterchangeAudioTranslator_AIFF : public UInterchangeAudioTranslatorBase
+{
+	GENERATED_BODY()
+public:
+	virtual TArray<FString> GetSupportedFormats() const override;
+};
+
+#undef UE_API
\ No newline at end of file
diff --git a/Engine/Plugins/Interchange/Runtime/Source/Import/Public/Audio/Formats/InterchangeAudioTranslator_FLAC.h b/Engine/Plugins/Interchange/Runtime/Source/Import/Public/Audio/Formats/InterchangeAudioTranslator_FLAC.h
new file mode 100644
--- /dev/null
+++ b/Engine/Plugins/Interchange/Runtime/Source/Import/Public/Audio/Formats/InterchangeAudioTranslator_FLAC.h
@@ -0,0 +1,18 @@
+// Copyright Epic Games, Inc. All Rights Reserved.
+#pragma once
+#include "Audio/Formats/InterchangeAudioTranslatorBase.h"
+
+#include "InterchangeAudioTranslator_FLAC.generated.h"
+
+#define UE_API INTERCHANGEIMPORT_API
+
+UCLASS(BlueprintType)
+class UE_API UInterchangeAudioTranslator_FLAC : public UInterchangeAudioTranslatorBase
+{
+	GENERATED_BODY()
+	
+public:
+	virtual TArray<FString> GetSupportedFormats() const override;
+};
+
+#undef UE_API
\ No newline at end of file
diff --git a/Engine/Plugins/Interchange/Runtime/Source/Import/Public/Audio/Formats/InterchangeAudioTranslator_MP3.h b/Engine/Plugins/Interchange/Runtime/Source/Import/Public/Audio/Formats/InterchangeAudioTranslator_MP3.h
new file mode 100644
--- /dev/null
+++ b/Engine/Plugins/Interchange/Runtime/Source/Import/Public/Audio/Formats/InterchangeAudioTranslator_MP3.h
@@ -0,0 +1,18 @@
+// Copyright Epic Games, Inc. All Rights Reserved.
+#pragma once
+#include "Audio/Formats/InterchangeAudioTranslatorBase.h"
+
+#include "InterchangeAudioTranslator_MP3.generated.h"
+
+#define UE_API INTERCHANGEIMPORT_API
+
+UCLASS(BlueprintType)
+class UE_API UInterchangeAudioTranslator_MP3 : public UInterchangeAudioTranslatorBase
+{
+	GENERATED_BODY()
+	
+public:
+	virtual TArray<FString> GetSupportedFormats() const override;
+};
+
+#undef UE_API
\ No newline at end of file
diff --git a/Engine/Plugins/Interchange/Runtime/Source/Import/Public/Audio/Formats/InterchangeAudioTranslator_OGG.h b/Engine/Plugins/Interchange/Runtime/Source/Import/Public/Audio/Formats/InterchangeAudioTranslator_OGG.h
new file mode 100644
--- /dev/null
+++ b/Engine/Plugins/Interchange/Runtime/Source/Import/Public/Audio/Formats/InterchangeAudioTranslator_OGG.h
@@ -0,0 +1,18 @@
+// Copyright Epic Games, Inc. All Rights Reserved.
+#pragma once
+#include "Audio/Formats/InterchangeAudioTranslatorBase.h"
+
+#include "InterchangeAudioTranslator_OGG.generated.h"
+
+#define UE_API INTERCHANGEIMPORT_API
+
+UCLASS(BlueprintType)
+class UE_API UInterchangeAudioTranslator_OGG : public UInterchangeAudioTranslatorBase
+{
+	GENERATED_BODY()
+	
+public:
+	virtual TArray<FString> GetSupportedFormats() const override;
+};
+
+#undef UE_API
\ No newline at end of file
diff --git a/Engine/Plugins/Interchange/Runtime/Source/Import/Public/Audio/Formats/InterchangeAudioTranslator_OPUS.h b/Engine/Plugins/Interchange/Runtime/Source/Import/Public/Audio/Formats/InterchangeAudioTranslator_OPUS.h
new file mode 100644
--- /dev/null
+++ b/Engine/Plugins/Interchange/Runtime/Source/Import/Public/Audio/Formats/InterchangeAudioTranslator_OPUS.h
@@ -0,0 +1,18 @@
+// Copyright Epic Games, Inc. All Rights Reserved.
+#pragma once
+#include "Audio/Formats/InterchangeAudioTranslatorBase.h"
+
+#include "InterchangeAudioTranslator_OPUS.generated.h"
+
+#define UE_API INTERCHANGEIMPORT_API
+
+UCLASS(BlueprintType)
+class UE_API UInterchangeAudioTranslator_OPUS : public UInterchangeAudioTranslatorBase
+{
+	GENERATED_BODY()
+	
+public:
+	virtual TArray<FString> GetSupportedFormats() const override;
+};
+
+#undef UE_API
\ No newline at end of file
diff --git a/Engine/Plugins/Interchange/Runtime/Source/Import/Public/Audio/Formats/InterchangeAudioTranslator_WAV.h b/Engine/Plugins/Interchange/Runtime/Source/Import/Public/Audio/Formats/InterchangeAudioTranslator_WAV.h
new file mode 100644
--- /dev/null
+++ b/Engine/Plugins/Interchange/Runtime/Source/Import/Public/Audio/Formats/InterchangeAudioTranslator_WAV.h
@@ -0,0 +1,18 @@
+// Copyright Epic Games, Inc. All Rights Reserved.
+#pragma once
+#include "Audio/Formats/InterchangeAudioTranslatorBase.h"
+
+#include "InterchangeAudioTranslator_WAV.generated.h"
+
+#define UE_API INTERCHANGEIMPORT_API
+
+UCLASS(BlueprintType)
+class UE_API UInterchangeAudioTranslator_WAV : public UInterchangeAudioTranslatorBase
+{
+	GENERATED_BODY()
+	
+public:
+	virtual TArray<FString> GetSupportedFormats() const override;
+};
+
+#undef UE_API
\ No newline at end of file
diff --git a/Engine/Plugins/Interchange/Runtime/Source/Import/Public/Audio/InterchangeGenericAudioTranslator.h b/Engine/Plugins/Interchange/Runtime/Source/Import/Public/Audio/InterchangeGenericAudioTranslator.h
deleted file mode 100644
