[Interchange Audio] Namespace fixes

#rnx

diff --git a/Engine/Plugins/Interchange/Runtime/Source/Import/Private/Audio/Formats/InterchangeAudioFormatsPrivate.cpp b/Engine/Plugins/Interchange/Runtime/Source/Import/Private/Audio/Formats/InterchangeAudioFormatsPrivate.cpp
--- a/Engine/Plugins/Interchange/Runtime/Source/Import/Private/Audio/Formats/InterchangeAudioFormatsPrivate.cpp
+++ b/Engine/Plugins/Interchange/Runtime/Source/Import/Private/Audio/Formats/InterchangeAudioFormatsPrivate.cpp
@@ -2,7 +2,7 @@
 #include "Audio/Formats/InterchangeAudioFormatsPrivate.h"
 #include "HAL/IConsoleManager.h"
 
-namespace Interchange::Audio
+namespace UE::Interchange::Audio
 {
 	const FString& FAudioFormatCVarNames::GetAudioFormatCVarName_AIF()
 	{ 
@@ -49,49 +49,49 @@ namespace Interchange::Audio
 
 static bool GInterchangeEnableAIFImport = false;
 static FAutoConsoleVariableRef CCvarInterchangeEnableAIFImport(
-	*Interchange::Audio::FAudioFormatCVarNames::GetAudioFormatCVarName_AIF(),
+	*UE::Interchange::Audio::FAudioFormatCVarNames::GetAudioFormatCVarName_AIF(),
 	GInterchangeEnableAIFImport,
 	TEXT("Whether AIF support is enabled."),
 	ECVF_Default);
 
 static bool GInterchangeEnableAIFFImport = false;
 static FAutoConsoleVariableRef CCvarInterchangeEnableAIFFImport(
-	*Interchange::Audio::FAudioFormatCVarNames::GetAudioFormatCVarName_AIFF(),
+	*UE::Interchange::Audio::FAudioFormatCVarNames::GetAudioFormatCVarName_AIFF(),
 	GInterchangeEnableAIFFImport,
 	TEXT("Whether AIFF support is enabled."),
 	ECVF_Default);
 
 static bool GInterchangeEnableOGGImport = false;
 static FAutoConsoleVariableRef CCvarInterchangeEnableOGGImport(
-	*Interchange::Audio::FAudioFormatCVarNames::GetAudioFormatCVarName_OGG(),
+	*UE::Interchange::Audio::FAudioFormatCVarNames::GetAudioFormatCVarName_OGG(),
 	GInterchangeEnableOGGImport,
 	TEXT("Whether OGG support is enabled."),
 	ECVF_Default);
 
 static bool GInterchangeEnableFLACImport = false;
 static FAutoConsoleVariableRef CCvarInterchangeEnableFLACImport(
-	*Interchange::Audio::FAudioFormatCVarNames::GetAudioFormatCVarName_FLAC(),
+	*UE::Interchange::Audio::FAudioFormatCVarNames::GetAudioFormatCVarName_FLAC(),
 	GInterchangeEnableFLACImport,
 	TEXT("Whether FLAC support is enabled."),
 	ECVF_Default);
 
 static bool GInterchangeEnableOPUSImport = false;
 static FAutoConsoleVariableRef CCvarInterchangeEnableOPUSImport(
-	*Interchange::Audio::FAudioFormatCVarNames::GetAudioFormatCVarName_OPUS(),
+	*UE::Interchange::Audio::FAudioFormatCVarNames::GetAudioFormatCVarName_OPUS(),
 	GInterchangeEnableOPUSImport,
 	TEXT("Whether OPUS support is enabled."),
 	ECVF_Default);
 
 static bool GInterchangeEnableMP3Import = false;
 static FAutoConsoleVariableRef CCvarInterchangeEnableMP3Import(
-	*Interchange::Audio::FAudioFormatCVarNames::GetAudioFormatCVarName_MP3(),
+	*UE::Interchange::Audio::FAudioFormatCVarNames::GetAudioFormatCVarName_MP3(),
 	GInterchangeEnableMP3Import,
 	TEXT("Whether MP3 support is enabled."),
 	ECVF_Default);
 
 static bool GInterchangeEnableWAVImport = false;
 static FAutoConsoleVariableRef CCvarInterchangeEnableWAVImport(
-	*Interchange::Audio::FAudioFormatCVarNames::GetAudioFormatCVarName_WAV(),
+	*UE::Interchange::Audio::FAudioFormatCVarNames::GetAudioFormatCVarName_WAV(),
 	GInterchangeEnableWAVImport,
 	TEXT("Whether WAV support is enabled."),
 	ECVF_Default);
\ No newline at end of file
diff --git a/Engine/Plugins/Interchange/Runtime/Source/Import/Private/Audio/Formats/InterchangeAudioFormatsPrivate.h b/Engine/Plugins/Interchange/Runtime/Source/Import/Private/Audio/Formats/InterchangeAudioFormatsPrivate.h
--- a/Engine/Plugins/Interchange/Runtime/Source/Import/Private/Audio/Formats/InterchangeAudioFormatsPrivate.h
+++ b/Engine/Plugins/Interchange/Runtime/Source/Import/Private/Audio/Formats/InterchangeAudioFormatsPrivate.h
@@ -8,7 +8,7 @@
 	#define WITH_SNDFILE_IO 0
 #endif
 
-namespace Interchange::Audio
+namespace UE::Interchange::Audio
 {
 	struct FAudioFormatCVarNames
 	{
diff --git a/Engine/Plugins/Interchange/Runtime/Source/Import/Private/Audio/Formats/InterchangeAudioTranslator_AIF.cpp b/Engine/Plugins/Interchange/Runtime/Source/Import/Private/Audio/Formats/InterchangeAudioTranslator_AIF.cpp
--- a/Engine/Plugins/Interchange/Runtime/Source/Import/Private/Audio/Formats/InterchangeAudioTranslator_AIF.cpp
+++ b/Engine/Plugins/Interchange/Runtime/Source/Import/Private/Audio/Formats/InterchangeAudioTranslator_AIF.cpp
@@ -6,7 +6,7 @@
 
 TArray<FString> UInterchangeAudioTranslator_AIF::GetSupportedFormats() const
 {
-	using namespace Interchange::Audio;
+	using namespace UE::Interchange::Audio;
 
 	TArray<FString> Formats;
 #if WITH_SNDFILE_IO
diff --git a/Engine/Plugins/Interchange/Runtime/Source/Import/Private/Audio/Formats/InterchangeAudioTranslator_AIFF.cpp b/Engine/Plugins/Interchange/Runtime/Source/Import/Private/Audio/Formats/InterchangeAudioTranslator_AIFF.cpp
--- a/Engine/Plugins/Interchange/Runtime/Source/Import/Private/Audio/Formats/InterchangeAudioTranslator_AIFF.cpp
+++ b/Engine/Plugins/Interchange/Runtime/Source/Import/Private/Audio/Formats/InterchangeAudioTranslator_AIFF.cpp
@@ -6,7 +6,7 @@
 
 TArray<FString> UInterchangeAudioTranslator_AIFF::GetSupportedFormats() const
 {
-	using namespace Interchange::Audio;
+	using namespace UE::Interchange::Audio;
 
 	TArray<FString> Formats;
 #if WITH_SNDFILE_IO
diff --git a/Engine/Plugins/Interchange/Runtime/Source/Import/Private/Audio/Formats/InterchangeAudioTranslator_FLAC.cpp b/Engine/Plugins/Interchange/Runtime/Source/Import/Private/Audio/Formats/InterchangeAudioTranslator_FLAC.cpp
--- a/Engine/Plugins/Interchange/Runtime/Source/Import/Private/Audio/Formats/InterchangeAudioTranslator_FLAC.cpp
+++ b/Engine/Plugins/Interchange/Runtime/Source/Import/Private/Audio/Formats/InterchangeAudioTranslator_FLAC.cpp
@@ -6,7 +6,7 @@
 
 TArray<FString> UInterchangeAudioTranslator_FLAC::GetSupportedFormats() const
 {
-	using namespace Interchange::Audio;
+	using namespace UE::Interchange::Audio;
 
 	TArray<FString> Formats;
 #if WITH_SNDFILE_IO
diff --git a/Engine/Plugins/Interchange/Runtime/Source/Import/Private/Audio/Formats/InterchangeAudioTranslator_MP3.cpp b/Engine/Plugins/Interchange/Runtime/Source/Import/Private/Audio/Formats/InterchangeAudioTranslator_MP3.cpp
--- a/Engine/Plugins/Interchange/Runtime/Source/Import/Private/Audio/Formats/InterchangeAudioTranslator_MP3.cpp
+++ b/Engine/Plugins/Interchange/Runtime/Source/Import/Private/Audio/Formats/InterchangeAudioTranslator_MP3.cpp
@@ -6,7 +6,7 @@
 
 TArray<FString> UInterchangeAudioTranslator_MP3::GetSupportedFormats() const
 {
-	using namespace Interchange::Audio;
+	using namespace UE::Interchange::Audio;
 
 	TArray<FString> Formats;
 #if WITH_SNDFILE_IO
diff --git a/Engine/Plugins/Interchange/Runtime/Source/Import/Private/Audio/Formats/InterchangeAudioTranslator_OGG.cpp b/Engine/Plugins/Interchange/Runtime/Source/Import/Private/Audio/Formats/InterchangeAudioTranslator_OGG.cpp
--- a/Engine/Plugins/Interchange/Runtime/Source/Import/Private/Audio/Formats/InterchangeAudioTranslator_OGG.cpp
+++ b/Engine/Plugins/Interchange/Runtime/Source/Import/Private/Audio/Formats/InterchangeAudioTranslator_OGG.cpp
@@ -6,7 +6,7 @@
 
 TArray<FString> UInterchangeAudioTranslator_OGG::GetSupportedFormats() const
 {
-	using namespace Interchange::Audio;
+	using namespace UE::Interchange::Audio;
 
 	TArray<FString> Formats;
 #if WITH_SNDFILE_IO
diff --git a/Engine/Plugins/Interchange/Runtime/Source/Import/Private/Audio/Formats/InterchangeAudioTranslator_OPUS.cpp b/Engine/Plugins/Interchange/Runtime/Source/Import/Private/Audio/Formats/InterchangeAudioTranslator_OPUS.cpp
--- a/Engine/Plugins/Interchange/Runtime/Source/Import/Private/Audio/Formats/InterchangeAudioTranslator_OPUS.cpp
+++ b/Engine/Plugins/Interchange/Runtime/Source/Import/Private/Audio/Formats/InterchangeAudioTranslator_OPUS.cpp
@@ -6,7 +6,7 @@
 
 TArray<FString> UInterchangeAudioTranslator_OPUS::GetSupportedFormats() const
 {
-	using namespace Interchange::Audio;
+	using namespace UE::Interchange::Audio;
 
 	TArray<FString> Formats;
 #if WITH_SNDFILE_IO
diff --git a/Engine/Plugins/Interchange/Runtime/Source/Import/Private/Audio/Formats/InterchangeAudioTranslator_WAV.cpp b/Engine/Plugins/Interchange/Runtime/Source/Import/Private/Audio/Formats/InterchangeAudioTranslator_WAV.cpp
--- a/Engine/Plugins/Interchange/Runtime/Source/Import/Private/Audio/Formats/InterchangeAudioTranslator_WAV.cpp
+++ b/Engine/Plugins/Interchange/Runtime/Source/Import/Private/Audio/Formats/InterchangeAudioTranslator_WAV.cpp
@@ -6,7 +6,7 @@
 
 TArray<FString> UInterchangeAudioTranslator_WAV::GetSupportedFormats() const
 {
-	using namespace Interchange::Audio;
+	using namespace UE::Interchange::Audio;
 
 	TArray<FString> Formats;
 	const IConsoleVariable* CVarFormatEnabled_WAV = IConsoleManager::Get().FindConsoleVariable(*FAudioFormatCVarNames::GetAudioFormatCVarName_WAV());
diff --git a/Engine/Plugins/Interchange/Runtime/Source/Import/Private/Audio/InterchangeAudioPayloadData.cpp b/Engine/Plugins/Interchange/Runtime/Source/Import/Private/Audio/InterchangeAudioPayloadData.cpp
--- a/Engine/Plugins/Interchange/Runtime/Source/Import/Private/Audio/InterchangeAudioPayloadData.cpp
+++ b/Engine/Plugins/Interchange/Runtime/Source/Import/Private/Audio/InterchangeAudioPayloadData.cpp
@@ -24,9 +24,9 @@ namespace UE::Interchange::Private
 		RawAudioData.AddUninitialized(RawWaveData.Num());
 		FMemory::Memcpy(RawAudioData.GetData(), RawWaveData.GetData(), RawWaveData.Num());
 		
-		const uint32 OrigNumSamples = Audio::SoundFileUtils::GetNumSamples(RawAudioData);
+		const uint32 OrigNumSamples = ::Audio::SoundFileUtils::GetNumSamples(RawAudioData);
 
-		if (Audio::SoundFileUtils::ConvertAudioToWav(RawAudioData, ConvertedRawWaveData))
+		if (::Audio::SoundFileUtils::ConvertAudioToWav(RawAudioData, ConvertedRawWaveData))
 		{
 			WaveInfo = FWaveModInfo();
 			if (!WaveInfo.ReadWaveInfo(ConvertedRawWaveData.GetData(), ConvertedRawWaveData.Num(), &ErrorMessage))
