Interchange:
- Fix Crash

diff --git a/Engine/Plugins/Interchange/Runtime/Source/Import/Private/Material/InterchangeMaterialFactory.cpp b/Engine/Plugins/Interchange/Runtime/Source/Import/Private/Material/InterchangeMaterialFactory.cpp
--- a/Engine/Plugins/Interchange/Runtime/Source/Import/Private/Material/InterchangeMaterialFactory.cpp
+++ b/Engine/Plugins/Interchange/Runtime/Source/Import/Private/Material/InterchangeMaterialFactory.cpp
@@ -1223,16 +1223,25 @@ void UInterchangeMaterialFactory::SetupObject_GameThread(const FSetupObjectParam
 		{
 			for (struct FTextureParameterValue& TextureParameterValue : ImportedMaterial->TextureParameterValues)
 			{
-				if (TextureParameterValue.ParameterValue)
+				if (!TextureParameterValue.ParameterValue)
 				{
-					UTexture* DefaultTexture = nullptr;
-					if (ImportedMaterial->GetTextureParameterDefaultValue(TextureParameterValue.ParameterInfo.Name, DefaultTexture))
-					{
-						if (UTexture* ConvertedTexture = UE::Interchange::MaterialFactory::Internal::GetVirtualTextureStreamingMatchedTexture(DefaultTexture->VirtualTextureStreaming, TextureParameterValue.ParameterValue))
-						{
-							TextureParameterValue.ParameterValue = ConvertedTexture;
-						}
-					}
+					continue;
+				}
+
+				UTexture* DefaultTexture = nullptr;
+				if (!ImportedMaterial->GetTextureParameterDefaultValue(TextureParameterValue.ParameterInfo.Name, DefaultTexture))
+				{
+					continue;
+				}
+
+				if (!DefaultTexture)
+				{
+					continue;
+				}
+
+				if (UTexture* ConvertedTexture = UE::Interchange::MaterialFactory::Internal::GetVirtualTextureStreamingMatchedTexture(DefaultTexture->VirtualTextureStreaming, TextureParameterValue.ParameterValue))
+				{
+					TextureParameterValue.ParameterValue = ConvertedTexture;
 				}
 			}
 		}
