Interchange: Respect MINIMUM_ANIMATION_LENGTH when baking animations from different payload types.

diff --git a/Engine/Plugins/Interchange/Runtime/Source/Parsers/CommonParser/Private/InterchangeCommonAnimationPayload.cpp b/Engine/Plugins/Interchange/Runtime/Source/Parsers/CommonParser/Private/InterchangeCommonAnimationPayload.cpp
--- a/Engine/Plugins/Interchange/Runtime/Source/Parsers/CommonParser/Private/InterchangeCommonAnimationPayload.cpp
+++ b/Engine/Plugins/Interchange/Runtime/Source/Parsers/CommonParser/Private/InterchangeCommonAnimationPayload.cpp
@@ -2,6 +2,7 @@
 
 #include "InterchangeCommonAnimationPayload.h"
 
+#include "Animation/AnimTypes.h"
 #include "CoreMinimal.h"
 
 #include "InterchangeHelper.h"
@@ -94,7 +95,7 @@ namespace UE::Interchange
 				RangeEndTime = 0;
 			}
 
-			const double BakeInterval = 1.0 / BakeFrequency;
+			const double BakeInterval = FMath::Max<double>(1.0 / BakeFrequency, MINIMUM_ANIMATION_LENGTH);
 			const double SequenceLength = FMath::Max<double>(RangeEndTime - RangeStartTime, BakeInterval);
 			int32 BakeKeyCount = FMath::RoundToInt32(SequenceLength * BakeFrequency) + 1;
 
