Interchange: Fix InterchangeWorker broken build when not building against the engine.

Horde issue 1042792

diff --git a/Engine/Plugins/Interchange/Runtime/Source/Parsers/CommonParser/Private/InterchangeCommonAnimationPayload.cpp b/Engine/Plugins/Interchange/Runtime/Source/Parsers/CommonParser/Private/InterchangeCommonAnimationPayload.cpp
--- a/Engine/Plugins/Interchange/Runtime/Source/Parsers/CommonParser/Private/InterchangeCommonAnimationPayload.cpp
+++ b/Engine/Plugins/Interchange/Runtime/Source/Parsers/CommonParser/Private/InterchangeCommonAnimationPayload.cpp
@@ -2,8 +2,10 @@
 
 #include "InterchangeCommonAnimationPayload.h"
 
-#include "Animation/AnimTypes.h"
 #include "CoreMinimal.h"
+#if WITH_ENGINE
+#include "Animation/AnimTypes.h"
+#endif
 
 #include "InterchangeHelper.h"
 
