[Interchange] Added option to import a static mesh as a skeletal mesh with top node as the only bone.

This CL also fixes UE-317102 (Unable to import a StaticMesh with "Force All Meshes" to SkeletalMesh and "Import Meshes in Bone Hierarchy" to false)

An option "Single Bone Skeleton" was added that is grayed out if "Force All Mesh as Type" is not set to "Skeletal Mesh".

This CL only affects Static Meshes that were converted to Skeletal Mesh and will simply generate a one bone skeleton instead of generating one with the model hierarchy.

No other path was changed.

[RN] Added an option to import a static mesh as a skeletal mesh with top node as the only bone.
[RN] Fixed imported skeletal meshes being empty when importing a static mesh with "Force All Mesh as Type" set to "Skeletal Mesh" and "Import Meshes in Bone Hierarchy" to false

diff --git a/Engine/Plugins/Interchange/Runtime/Source/Pipelines/Private/InterchangeGenericSkeletalMeshPipeline.cpp b/Engine/Plugins/Interchange/Runtime/Source/Pipelines/Private/InterchangeGenericSkeletalMeshPipeline.cpp
--- a/Engine/Plugins/Interchange/Runtime/Source/Pipelines/Private/InterchangeGenericSkeletalMeshPipeline.cpp
+++ b/Engine/Plugins/Interchange/Runtime/Source/Pipelines/Private/InterchangeGenericSkeletalMeshPipeline.cpp
@@ -232,9 +232,9 @@ void UInterchangeGenericMeshPipeline::ExecutePreImportPipelineSkeletalMesh()
 #endif
 		};
 
-
-		PipelineMeshesUtilities->GetCombinedSkinnedMeshInstances(MeshUidsPerSkeletonRootUid);
-		bool bUseMeshInstance = true;
+		const bool bUseSingleBoneForConvertedMeshes = CommonMeshesProperties->ForceAllMeshAsType == EInterchangeForceMeshType::IFMT_SkeletalMesh && CommonMeshesProperties->bSingleBoneSkeleton;
+		PipelineMeshesUtilities->GetCombinedSkinnedMeshInstances(MeshUidsPerSkeletonRootUid, bUseSingleBoneForConvertedMeshes);
+		constexpr bool bUseMeshInstance = true;
 		CreatePerSkeletonRootUidCombinedSkinnedMesh(bUseMeshInstance);
 	}
 	else
@@ -574,7 +574,12 @@ void UInterchangeGenericMeshPipeline::AddLodDataToSkeletalMesh(const UInterchang
 		TArray<FString> NodeUids = NodeUidsPerLodIndex.Contains(LodIndex) ? NodeUidsPerLodIndex.FindChecked(LodIndex) : EmptyLodData;
 		if (!CommonSkeletalMeshesAndAnimationsProperties->bImportMeshesInBoneHierarchy)
 		{
-			UE::Interchange::SkeletalMeshGenericPipeline::RemoveNestedMeshNodes(BaseNodeContainer, SkeletonFactoryNode, NodeUids);
+			//Unless ForceAllMeshAsType is set to SkeletalMesh (UE-317102)
+			//In this case we need to keep nested static meshes since they might be the only meshes used to generate this skeletal mesh.
+			if (CommonMeshesProperties->ForceAllMeshAsType != EInterchangeForceMeshType::IFMT_SkeletalMesh)
+			{
+				UE::Interchange::SkeletalMeshGenericPipeline::RemoveNestedMeshNodes(BaseNodeContainer, SkeletonFactoryNode, NodeUids);
+			}
 		}
 
 		//Create a lod data node with all the meshes for this LOD
diff --git a/Engine/Plugins/Interchange/Runtime/Source/Pipelines/Private/InterchangePipelineMeshesUtilities.cpp b/Engine/Plugins/Interchange/Runtime/Source/Pipelines/Private/InterchangePipelineMeshesUtilities.cpp
--- a/Engine/Plugins/Interchange/Runtime/Source/Pipelines/Private/InterchangePipelineMeshesUtilities.cpp
+++ b/Engine/Plugins/Interchange/Runtime/Source/Pipelines/Private/InterchangePipelineMeshesUtilities.cpp
@@ -905,7 +905,8 @@ void UInterchangePipelineMeshesUtilities::IterateAllMeshInstanceUsingMeshGeometr
 	}
 }
 
-void UInterchangePipelineMeshesUtilities::GetCombinedSkinnedMeshInstances(TMap<FString, TArray<FString>>& OutMeshInstanceUidsPerSkeletonRootUid) const
+void UInterchangePipelineMeshesUtilities::GetCombinedSkinnedMeshInstances(TMap<FString, TArray<FString>>& OutMeshInstanceUidsPerSkeletonRootUid,
+	bool bUseSingleBoneForConvertedMeshes /*= false*/) const
 {
 	check(BaseNodeContainer);
 
@@ -918,10 +919,11 @@ void UInterchangePipelineMeshesUtilities::GetCombinedSkinnedMeshInstances(TMap<F
 		{
 			continue;
 		}
+
 		bool bIsStaticMesh = !MeshInstance.bReferenceSkinnedMesh;
 		//Find the root skeleton for this MeshInstance
 		FString SkeletonRootUid; // = CurrentDataContext.FindSkeletalMeshSkeletonRootFromMeshInstance(MeshInstance, BaseNodeContainer);
-		
+
 		for (const FString& MeshGeometryUid : MeshInstance.ReferencingMeshGeometryUids)
 		{
 			if (const FString* SkeletonRootUidPtr = SkeletonRootUidPerMeshUid.Find(MeshGeometryUid))
@@ -952,7 +954,7 @@ void UInterchangePipelineMeshesUtilities::GetCombinedSkinnedMeshInstances(TMap<F
 						{
 							if (const UInterchangeSceneNode* ParentNode = Cast<UInterchangeSceneNode>(BaseNodeContainer->GetNode(ParentUid)))
 							{
-								if(ParentNode->IsSpecializedTypeContains(UE::Interchange::FSceneNodeStaticData::GetJointSpecializeTypeString()))
+								if (ParentNode->IsSpecializedTypeContains(UE::Interchange::FSceneNodeStaticData::GetJointSpecializeTypeString()))
 								{
 									SkeletonRootUid = ParentUid;
 								}
@@ -965,6 +967,24 @@ void UInterchangePipelineMeshesUtilities::GetCombinedSkinnedMeshInstances(TMap<F
 						{
 							SkeletonRootUid = LastSceneNodeUid;
 						}
+
+						//If we want to use a single bone skeleton for the static mesh, we need create or reuse a dedicated scene node for that
+						//Make sure to ignore nested static meshes so we don't affect "Import Meshes In Bone Hierarchy" option behaviour
+						if (!bIsNestedIntoSkeleton && bUseSingleBoneForConvertedMeshes)
+						{
+							//Make sure to use the SkeletonRootUid found above so all meshes under the same root joint will use the same single bone skeleton too.
+							const FString NodeUid = SkeletonRootUid + "_SingleBone_RootNode";
+
+							const UInterchangeBaseNode* SingleBoneRootNode = BaseNodeContainer->GetNode(NodeUid);
+							if (!SingleBoneRootNode)
+							{
+								UInterchangeSceneNode* SkeletonRootNode = NewObject<UInterchangeSceneNode>(BaseNodeContainer);
+								BaseNodeContainer->SetupNode(SkeletonRootNode, NodeUid, "Root", EInterchangeNodeContainerType::TranslatedAsset);
+								SkeletonRootNode->SetCustomLocalTransform(BaseNodeContainer, FTransform::Identity);
+							}
+
+							SkeletonRootUid = NodeUid;
+						}
 					}
 				}
 			}
@@ -979,6 +999,7 @@ void UInterchangePipelineMeshesUtilities::GetCombinedSkinnedMeshInstances(TMap<F
 			//Skip this MeshInstance
 			continue;
 		}
+
 		TArray<FString>& MeshInstanceUids = OutMeshInstanceUidsPerSkeletonRootUid.FindOrAdd(SkeletonRootUid);
 		MeshInstanceUids.Add(MeshInstanceUidAndMeshInstance.Key);
 	}
diff --git a/Engine/Plugins/Interchange/Runtime/Source/Pipelines/Public/InterchangeGenericAssetsPipelineSharedSettings.h b/Engine/Plugins/Interchange/Runtime/Source/Pipelines/Public/InterchangeGenericAssetsPipelineSharedSettings.h
--- a/Engine/Plugins/Interchange/Runtime/Source/Pipelines/Public/InterchangeGenericAssetsPipelineSharedSettings.h
+++ b/Engine/Plugins/Interchange/Runtime/Source/Pipelines/Public/InterchangeGenericAssetsPipelineSharedSettings.h
@@ -69,11 +69,19 @@ public:
  	UPROPERTY(EditAnywhere, BlueprintReadWrite, Category = "Common Meshes")
  	EInterchangeForceMeshType ForceAllMeshAsType = EInterchangeForceMeshType::IFMT_None;
 
+	/**
+	* If enabled, converted static meshes to skeletal meshes will use a skeleton with a unique bone.
+	* Otherwise, the bone structure is determined by the model hierarchy.
+	* This setting is only used if Force All Meshes As Type setting is set to "Skeletal Mesh".
+	*/
+	UPROPERTY(EditAnywhere, BlueprintReadWrite, Category = "Common Meshes", meta = (editcondition = "ForceAllMeshAsType == EInterchangeForceMeshType::IFMT_SkeletalMesh"))
+	bool bSingleBoneSkeleton = false;
+
 	/**
 	 * If enabled, and some static mesh transforms are animated, the pipeline will convert the static mesh into a rigid skeletal mesh.
 	 * This setting is only used if the Force All Meshes As Type setting is set to "None".
 	 */
-	UPROPERTY(EditAnywhere, BlueprintReadWrite, Category = "Common Meshes")
+	UPROPERTY(EditAnywhere, BlueprintReadWrite, Category = "Common Meshes", meta = (editcondition = "ForceAllMeshAsType == EInterchangeForceMeshType::IFMT_None"))
 	bool bAutoDetectMeshType = true;
 
 	/** If enabled, any existing LODs for meshes are imported. This setting is only used if the Bake Meshes setting is also enabled. */
@@ -145,6 +153,7 @@ public:
 		static const TSet<FName> NeedRefreshProperties =
 		{
 			GET_MEMBER_NAME_CHECKED(UInterchangeGenericCommonMeshesProperties, ForceAllMeshAsType),
+			GET_MEMBER_NAME_CHECKED(UInterchangeGenericCommonMeshesProperties, bSingleBoneSkeleton),
 			GET_MEMBER_NAME_CHECKED(UInterchangeGenericCommonMeshesProperties, bAutoDetectMeshType)
 		};
 
diff --git a/Engine/Plugins/Interchange/Runtime/Source/Pipelines/Public/InterchangePipelineMeshesUtilities.h b/Engine/Plugins/Interchange/Runtime/Source/Pipelines/Public/InterchangePipelineMeshesUtilities.h
--- a/Engine/Plugins/Interchange/Runtime/Source/Pipelines/Public/InterchangePipelineMeshesUtilities.h
+++ b/Engine/Plugins/Interchange/Runtime/Source/Pipelines/Public/InterchangePipelineMeshesUtilities.h
@@ -336,7 +336,7 @@ public:
 	* Return a list of skinned FInterchangeMeshInstance UIDs that can be combined together.
 	* We cannot create a skinned mesh that has multiple skeleton root nodes. This function returns a combined MeshInstance per skeleton roots.
 	*/
-	UE_API void GetCombinedSkinnedMeshInstances(TMap<FString, TArray<FString>>& OutMeshInstanceUidsPerSkeletonRootUid) const;
+	UE_API void GetCombinedSkinnedMeshInstances(TMap<FString, TArray<FString>>& OutMeshInstanceUidsPerSkeletonRootUid, bool bUseSingleBoneForConvertedMeshes = false) const;
 	
 	/**
 	* Return the skeleton root node UID. This is the UID for a UInterchangeSceneNode that has a "Joint" specialized type.
