[Interchange] Fixed Reimport not showing skeletal conflicts

The transient base node used to compute the skeleton conflicts was not duplicated with the children cache, causing the skeleton hierarchy to not be constructed properly.

[RN] Fixed Reimport not showing skeletal conflicts

diff --git a/Engine/Plugins/Interchange/Runtime/Source/Import/Private/Mesh/InterchangeSkeletonHelper.cpp b/Engine/Plugins/Interchange/Runtime/Source/Import/Private/Mesh/InterchangeSkeletonHelper.cpp
--- a/Engine/Plugins/Interchange/Runtime/Source/Import/Private/Mesh/InterchangeSkeletonHelper.cpp
+++ b/Engine/Plugins/Interchange/Runtime/Source/Import/Private/Mesh/InterchangeSkeletonHelper.cpp
@@ -85,6 +85,7 @@ namespace UE::Interchange::Private
 		return true;
 	}
 
+	// This function should match USkeleton::IsCompatibleMesh semantic for consistency with the "Failed to merge bones" dialog that shows up when skeletons are incompatible.
 	bool FSkeletonHelper::IsCompatibleSkeleton(const USkeleton* Skeleton
 		, const FString RootJoinUid
 		, const UInterchangeBaseNodeContainer* BaseNodeContainer
diff --git a/Engine/Plugins/Interchange/Runtime/Source/Pipelines/Private/InterchangeGenericAssetsPipeline.cpp b/Engine/Plugins/Interchange/Runtime/Source/Pipelines/Private/InterchangeGenericAssetsPipeline.cpp
--- a/Engine/Plugins/Interchange/Runtime/Source/Pipelines/Private/InterchangeGenericAssetsPipeline.cpp
+++ b/Engine/Plugins/Interchange/Runtime/Source/Pipelines/Private/InterchangeGenericAssetsPipeline.cpp
@@ -1066,6 +1066,7 @@ TArray<FInterchangeConflictInfo> UInterchangeGenericAssetsPipeline::GetConflictI
 	}
 
 	UInterchangeBaseNodeContainer* TransientBaseNodeContainer = DuplicateObject<UInterchangeBaseNodeContainer>(InBaseNodeContainer, GetTransientPackage());
+    TransientBaseNodeContainer->SetChildrenCache(InBaseNodeContainer->GetChildrenCache());
 
 	TArray<UInterchangeSourceData*> SourceDatas;
 	SourceDatas.Add(SourceData);
diff --git a/Engine/Plugins/Interchange/Runtime/Source/Pipelines/Private/InterchangePipelineMeshesUtilities.cpp b/Engine/Plugins/Interchange/Runtime/Source/Pipelines/Private/InterchangePipelineMeshesUtilities.cpp
--- a/Engine/Plugins/Interchange/Runtime/Source/Pipelines/Private/InterchangePipelineMeshesUtilities.cpp
+++ b/Engine/Plugins/Interchange/Runtime/Source/Pipelines/Private/InterchangePipelineMeshesUtilities.cpp
@@ -979,7 +979,7 @@ void UInterchangePipelineMeshesUtilities::GetCombinedSkinnedMeshInstances(TMap<F
 							if (!SingleBoneRootNode)
 							{
 								UInterchangeSceneNode* SkeletonRootNode = NewObject<UInterchangeSceneNode>(BaseNodeContainer);
-								BaseNodeContainer->SetupNode(SkeletonRootNode, NodeUid, "Root", EInterchangeNodeContainerType::TranslatedAsset);
+								BaseNodeContainer->SetupNode(SkeletonRootNode, NodeUid, "Root", EInterchangeNodeContainerType::TranslatedScene);
 								SkeletonRootNode->SetCustomLocalTransform(BaseNodeContainer, FTransform::Identity);
 							}
 
diff --git a/Engine/Plugins/Interchange/Runtime/Source/Pipelines/Public/InterchangePipelineMeshesUtilities.h b/Engine/Plugins/Interchange/Runtime/Source/Pipelines/Public/InterchangePipelineMeshesUtilities.h
--- a/Engine/Plugins/Interchange/Runtime/Source/Pipelines/Public/InterchangePipelineMeshesUtilities.h
+++ b/Engine/Plugins/Interchange/Runtime/Source/Pipelines/Public/InterchangePipelineMeshesUtilities.h
@@ -335,6 +335,8 @@ public:
 	/**
 	* Return a list of skinned FInterchangeMeshInstance UIDs that can be combined together.
 	* We cannot create a skinned mesh that has multiple skeleton root nodes. This function returns a combined MeshInstance per skeleton roots.
+	* 
+	* @Param bUseSingleBoneForConvertedMeshes - If true, the skeleton roots assigned to non-skinned meshes will be using a single bone. Otherwise, the skeleton will follow the hierarchy scene node hierarchy.
 	*/
 	UE_API void GetCombinedSkinnedMeshInstances(TMap<FString, TArray<FString>>& OutMeshInstanceUidsPerSkeletonRootUid, bool bUseSingleBoneForConvertedMeshes = false) const;
 	
