[Interchange-FBX] Fixing Force Front X Axis option.

#rnx

diff --git a/Engine/Plugins/Interchange/Runtime/Source/Import/Private/Mesh/InterchangeSkeletalMeshFactory.cpp b/Engine/Plugins/Interchange/Runtime/Source/Import/Private/Mesh/InterchangeSkeletalMeshFactory.cpp
--- a/Engine/Plugins/Interchange/Runtime/Source/Import/Private/Mesh/InterchangeSkeletalMeshFactory.cpp
+++ b/Engine/Plugins/Interchange/Runtime/Source/Import/Private/Mesh/InterchangeSkeletalMeshFactory.cpp
@@ -288,14 +288,17 @@ namespace UE
 							JointBindPoseMatrix = JointBindPose.ToMatrixWithScale();
 						}
 					
-						FTransform JointT0;
-						if (!JointNode->GetCustomTimeZeroGlobalTransform(NodeContainer, FTransform::Identity, JointT0))
+						if (!JointNode->GetCustomGlobalMatrixForT0Rebinding(JointT0Matrix))
 						{
-							//If there is no time zero global transform we cannot set the bind pose to time zero.
-							//We must skip this joint.
-							continue;
+							FTransform JointT0;
+							if (!JointNode->GetCustomTimeZeroGlobalTransform(NodeContainer, FTransform::Identity, JointT0))
+							{
+								//If there is no time zero global transform we cannot set the bind pose to time zero.
+								//We must skip this joint.
+								continue;
+							}
+							JointT0Matrix = JointT0.ToMatrixWithScale();
 						}
-						JointT0Matrix = JointT0.ToMatrixWithScale();
 
 						if (!JointNode->GetGlobalBindPoseReferenceForMeshUID(MeshNode->GetUniqueID(), JointReferenceMatrix))
 						{
@@ -306,7 +309,6 @@ namespace UE
 						//We must add the geometric node transform
 						JointReferenceMatrix = GeometricMatrix * JointReferenceMatrix;
 					}
-					
 					FMatrix ReferenceMatrixRelativeToBindPose = JointReferenceMatrix * JointBindPoseMatrix.Inverse();
 					FMatrix jointTimeZeroRelativeToMeshNodeMatrix = JointT0Matrix * SceneNodeMatrixInverse;
 					//Compute the transform to apply to vertex to change the bind pose from the existing one to use time zero pose instead
diff --git a/Engine/Plugins/Interchange/Runtime/Source/Nodes/Private/InterchangeSceneNode.cpp b/Engine/Plugins/Interchange/Runtime/Source/Nodes/Private/InterchangeSceneNode.cpp
--- a/Engine/Plugins/Interchange/Runtime/Source/Nodes/Private/InterchangeSceneNode.cpp
+++ b/Engine/Plugins/Interchange/Runtime/Source/Nodes/Private/InterchangeSceneNode.cpp
@@ -550,3 +550,13 @@ void UInterchangeSceneNode::GetComponentUids(TArray<FString>& OutComponentUids)
 {
 	ComponentUids.GetItems(OutComponentUids);
 }
+
+bool UInterchangeSceneNode::SetCustomGlobalMatrixForT0Rebinding(const FMatrix& AttributeValue)
+{
+	IMPLEMENT_NODE_ATTRIBUTE_SETTER_NODELEGATE(GlobalMatrixForT0Rebinding, FMatrix);
+}
+
+bool UInterchangeSceneNode::GetCustomGlobalMatrixForT0Rebinding(FMatrix& AttributeValue) const
+{
+	IMPLEMENT_NODE_ATTRIBUTE_GETTER(GlobalMatrixForT0Rebinding, FMatrix);
+}
diff --git a/Engine/Plugins/Interchange/Runtime/Source/Nodes/Public/InterchangeSceneNode.h b/Engine/Plugins/Interchange/Runtime/Source/Nodes/Public/InterchangeSceneNode.h
--- a/Engine/Plugins/Interchange/Runtime/Source/Nodes/Public/InterchangeSceneNode.h
+++ b/Engine/Plugins/Interchange/Runtime/Source/Nodes/Public/InterchangeSceneNode.h
@@ -289,6 +289,14 @@ public:
 	UFUNCTION(BlueprintCallable, Category = "Interchange | Node | Joint")
 	UE_API bool GetAnimationCurveTypeForCurveName(const FString& CurveName, EInterchangeAnimationPayLoadType& OutCurveAnimationType) const;
 
+	/** Sets the Global Transformation Matrix used for T0 rebinding.*/
+	UFUNCTION(BlueprintCallable, Category = "Interchange | Node | Joint")
+	UE_API bool SetCustomGlobalMatrixForT0Rebinding(const FMatrix& AttributeValue);
+
+	/** Gets the Global Transformation Matrix used for T0 rebinding.*/
+	UFUNCTION(BlueprintCallable, Category = "Interchange | Node | Joint")
+	UE_API bool GetCustomGlobalMatrixForT0Rebinding(FMatrix& AttributeValue) const;
+
 
 	UFUNCTION(BlueprintCallable, Category = "Interchange | Scene | Component")
 	UE_API bool AddComponentUid(const FString& ComponentUid);
@@ -357,6 +365,9 @@ private:
 	//A scene node can reference an animation asset on top of base asset:
 	const UE::Interchange::FAttributeKey Macro_CustomAnimationAssetUidToPlayKey = UE::Interchange::FAttributeKey(TEXT("AnimationAssetUidToPlay"));
 
+	//Global Transform Matrix specifically used for T0 re-binding.
+	const UE::Interchange::FAttributeKey Macro_CustomGlobalMatrixForT0RebindingKey = UE::Interchange::FAttributeKey(TEXT("GlobalMatrixForT0Rebinding"));
+
 	//mutable caches for global transforms
 	mutable TOptional<FTransform> CacheGlobalTransform;
 	mutable TOptional<FTransform> CacheBindPoseGlobalTransform;
diff --git a/Engine/Plugins/Interchange/Runtime/Source/Parsers/Fbx/Private/FbxAnimation.cpp b/Engine/Plugins/Interchange/Runtime/Source/Parsers/Fbx/Private/FbxAnimation.cpp
--- a/Engine/Plugins/Interchange/Runtime/Source/Parsers/Fbx/Private/FbxAnimation.cpp
+++ b/Engine/Plugins/Interchange/Runtime/Source/Parsers/Fbx/Private/FbxAnimation.cpp
@@ -540,7 +540,6 @@ namespace UE::Interchange::Private
 		else
 		{
 			FbxAMatrix LocalMatrix = FetchPayloadData.Node->EvaluateLocalTransform(Currenttime);
-			LocalMatrix = LocalMatrix * Parser.JointOrientationMatrix;
 
 			FbxVector4 NewLocalT = LocalMatrix.GetT();
 			FbxVector4 NewLocalS = LocalMatrix.GetS();
diff --git a/Engine/Plugins/Interchange/Runtime/Source/Parsers/Fbx/Private/FbxMesh.cpp b/Engine/Plugins/Interchange/Runtime/Source/Parsers/Fbx/Private/FbxMesh.cpp
--- a/Engine/Plugins/Interchange/Runtime/Source/Parsers/Fbx/Private/FbxMesh.cpp
+++ b/Engine/Plugins/Interchange/Runtime/Source/Parsers/Fbx/Private/FbxMesh.cpp
@@ -2133,7 +2133,7 @@ bool RetrievePoseFromBindPose(FbxNode* Joint, FbxAMatrix& GlobalBindPoseJointMat
 
 bool FFbxMesh::GetGlobalJointBindPoseTransform(FFbxParser* Parser, FbxScene* SDKScene, FbxNode* Joint, FFbxJointMeshBindPoseGenerator& FbxJointMeshBindPoseGenerator,
 	FbxAMatrix& GlobalBindPoseJointMatrix, TMap<FString, FMatrix>& MeshIdToGlobalBindPoseJointMap, TMap<FString, FMatrix>& MeshIdToGlobalBindPoseReferenceMap, 
-	bool& bBadBindPoseMessageDisplay, const FbxAMatrix& JointOrientationMatrix)
+	bool& bBadBindPoseMessageDisplay, const bool bIsOriginalJoint, const FbxAMatrix& JointOrientationMatrix)
 {
 	FbxManager* SDKManager = SDKScene->GetFbxManager();
 	
@@ -2174,7 +2174,8 @@ bool FFbxMesh::GetGlobalJointBindPoseTransform(FFbxParser* Parser, FbxScene* SDK
 		bRetrieveBindPoseResult = RetrievePoseFromBindPose(Joint, GlobalBindPoseJointMatrix);
 	}
 
-	if (bRetrieveBindPoseResult)
+	//Only joints should have the JointOrientationMatrix applied here.
+	if (bRetrieveBindPoseResult && bIsOriginalJoint)
 	{
 		GlobalBindPoseJointMatrix = GlobalBindPoseJointMatrix * JointOrientationMatrix;
 	}
@@ -2314,19 +2315,17 @@ FFbxJointMeshBindPoseGenerator::FFbxJointMeshBindPoseGenerator(FbxScene* SDKScen
 				{
 					Cluster->GetTransformLinkMatrix(GlobalBindPoseJointMatrixPerGeo);
 				}
-				GlobalBindPoseJointMatrixPerGeo = GlobalBindPoseJointMatrixPerGeo * Parser.JointOrientationMatrix;
 
 				TMap<FString, FMatrix>& MeshIdToGlobalBindPoseJointMap = JointToMeshIdToGlobalBindPoseJointMap.FindOrAdd(Joint);
 				MeshIdToGlobalBindPoseJointMap.Add(MeshUniqueID, FFbxConvert::ConvertMatrix<FMatrix>(GlobalBindPoseJointMatrixPerGeo));
 
 				if (!JointToBindPoseMap.Contains(Joint))
 				{
-					JointToBindPoseMap.Add(Joint, GlobalBindPoseJointMatrixPerGeo);
+					JointToBindPoseMap.Add(Joint, GlobalBindPoseJointMatrixPerGeo * Parser.JointOrientationMatrix);
 				}
 
 				FbxAMatrix GlobalBindPoseReferenceMatrix;
 				Cluster->GetTransformMatrix(GlobalBindPoseReferenceMatrix);
-				GlobalBindPoseReferenceMatrix = GlobalBindPoseReferenceMatrix * Parser.JointOrientationMatrix;
 
 				TMap<FString, FMatrix>& MeshIdToGlobalBindPoseReferenceMap = JointToMeshIdToGlobalBindPoseReferenceMap.FindOrAdd(Joint);
 				MeshIdToGlobalBindPoseReferenceMap.Add(MeshUniqueID, FFbxConvert::ConvertMatrix<FMatrix>(GlobalBindPoseReferenceMatrix));
diff --git a/Engine/Plugins/Interchange/Runtime/Source/Parsers/Fbx/Private/FbxMesh.h b/Engine/Plugins/Interchange/Runtime/Source/Parsers/Fbx/Private/FbxMesh.h
--- a/Engine/Plugins/Interchange/Runtime/Source/Parsers/Fbx/Private/FbxMesh.h
+++ b/Engine/Plugins/Interchange/Runtime/Source/Parsers/Fbx/Private/FbxMesh.h
@@ -212,7 +212,7 @@ namespace UE
 				void AddAllMeshes(FbxScene* SDKScene, FbxGeometryConverter* SDKGeometryConverter, UInterchangeBaseNodeContainer& NodeContainer, TMap<FString, TSharedPtr<FPayloadContextBase>>& PayloadContexts);
 				static bool GetGlobalJointBindPoseTransform(FFbxParser* Parser, FbxScene* SDKScene, FbxNode* Joint, FFbxJointMeshBindPoseGenerator& FbxJointMeshBindPoseGenerator,
 					FbxAMatrix& GlobalBindPoseJointMatrix, TMap<FString, FMatrix>& MeshIdToGlobalBindPoseJointMap, TMap<FString, FMatrix>& MeshIdToGlobalBindPoseReferenceMap, 
-					bool& bBadBindPoseMessageDisplay, const FbxAMatrix& JointOrientationMatrix);
+					bool& bBadBindPoseMessageDisplay, const bool bIsOriginalJoint, const FbxAMatrix& JointOrientationMatrix);
 
 				TArray<FMorphTargetAnimationBuildingData>& GetMorphTargetAnimationsBuildingData() { return MorphTargetAnimationsBuildingData; }
 			protected:
diff --git a/Engine/Plugins/Interchange/Runtime/Source/Parsers/Fbx/Private/FbxScene.cpp b/Engine/Plugins/Interchange/Runtime/Source/Parsers/Fbx/Private/FbxScene.cpp
--- a/Engine/Plugins/Interchange/Runtime/Source/Parsers/Fbx/Private/FbxScene.cpp
+++ b/Engine/Plugins/Interchange/Runtime/Source/Parsers/Fbx/Private/FbxScene.cpp
@@ -174,25 +174,7 @@ namespace UE
 					return Transform;
 				};
 
-				//Set the node default transform
-				{
-					FbxAMatrix GlobalFbxMatrix = Node->EvaluateGlobalTransform();
-					FTransform GlobalTransform = GetConvertedTransform(GlobalFbxMatrix);
-					if (FbxNode* ParentNode = Node->GetParent())
-					{
-						FbxAMatrix GlobalFbxParentMatrix = ParentNode->EvaluateGlobalTransform();
-						FbxAMatrix	LocalFbxMatrix = GlobalFbxParentMatrix.Inverse() * GlobalFbxMatrix;
-						FTransform LocalTransform = GetConvertedTransform(LocalFbxMatrix);
-						UnrealNode->SetCustomLocalTransform(&NodeContainer, LocalTransform, bResetCache);
-					}
-					else
-					{
-						//No parent, set the same matrix has the global
-						UnrealNode->SetCustomLocalTransform(&NodeContainer, GlobalTransform, bResetCache);
-					}
-				}
-
-				auto ApplySkeletonAttribute = [this, &SDKScene, &UnrealNode, &Node, &NodeContainer, &bResetCache, &GetConvertedTransform, &bBadBindPoseMessageDisplay, &FbxJointMeshBindPoseGenerator]()
+				auto ApplySkeletonAttribute = [this, &SDKScene, &UnrealNode, &Node, &NodeContainer, &bResetCache, &GetConvertedTransform, &bBadBindPoseMessageDisplay, &FbxJointMeshBindPoseGenerator](bool bIsOriginalJoint)
 				{
 					if (FRootJointInfo* RootJointInfo = CommonJointRootNodes.Find(Node))
 					{
@@ -213,40 +195,43 @@ namespace UE
 
 					//Add the joint specialized type
 					UnrealNode->AddSpecializedType(FSceneNodeStaticData::GetJointSpecializeTypeString());
-					//Get the bind pose transform for this joint
-					FbxAMatrix GlobalBindPoseJointMatrix = SDKScene->GetAnimationEvaluator()->GetNodeGlobalTransform(Node, 0);
-					TMap<FString, FMatrix> MeshIdToGlobalBindPoseReferenceMap;
-					TMap<FString, FMatrix> MeshIdToGlobalBindPoseJointMap;
-					FFbxMesh::GetGlobalJointBindPoseTransform(&Parser, SDKScene, Node, FbxJointMeshBindPoseGenerator, GlobalBindPoseJointMatrix, MeshIdToGlobalBindPoseJointMap, MeshIdToGlobalBindPoseReferenceMap, bBadBindPoseMessageDisplay, Parser.JointOrientationMatrix);
-
-					FTransform GlobalBindPoseJointTransform = GetConvertedTransform(GlobalBindPoseJointMatrix);
-					UnrealNode->SetGlobalBindPoseReferenceForMeshUIDs(MeshIdToGlobalBindPoseReferenceMap);
-					//Add Transform attribute per Geometry
-					//Temporary (to fix in 5.6 main branch): We need to make an API like Get/Set GlobalBindPoseReferenceForMeshUIDs for the joint one
-					for (const TPair<FString, FMatrix>& MeshIdAndBindPoseJointMatrix : MeshIdToGlobalBindPoseJointMap)
-					{
-						FString AttributeKey = TEXT("JointBindPosePerMesh_") + MeshIdAndBindPoseJointMatrix.Key;
-						UnrealNode->RegisterAttribute<FMatrix>(FAttributeKey(AttributeKey), MeshIdAndBindPoseJointMatrix.Value);
-					}
 
 					FbxNode* ParentNode = Node->GetParent();
 					
-					if (ParentNode != nullptr)
 					{
-						FbxAMatrix GlobalFbxParentMatrix = SDKScene->GetAnimationEvaluator()->GetNodeGlobalTransform(ParentNode, 0);
-						TMap<FString, FMatrix> ParentMeshIdToGlobalBindPoseReferenceMap;
-						TMap<FString, FMatrix> ParentMeshIdToGlobalBindPoseJointMap;
-						FFbxMesh::GetGlobalJointBindPoseTransform(&Parser, SDKScene, ParentNode, FbxJointMeshBindPoseGenerator, GlobalFbxParentMatrix, ParentMeshIdToGlobalBindPoseJointMap, ParentMeshIdToGlobalBindPoseReferenceMap, bBadBindPoseMessageDisplay, bIsRootJoint ? FbxAMatrix() : Parser.JointOrientationMatrix);
+						//Get the bind pose transform for this joint
+						FbxAMatrix GlobalBindPoseJointMatrix = SDKScene->GetAnimationEvaluator()->GetNodeGlobalTransform(Node, 0);
+						TMap<FString, FMatrix> MeshIdToGlobalBindPoseReferenceMap;
+						TMap<FString, FMatrix> MeshIdToGlobalBindPoseJointMap;
+						FFbxMesh::GetGlobalJointBindPoseTransform(&Parser, SDKScene, Node, FbxJointMeshBindPoseGenerator, GlobalBindPoseJointMatrix, MeshIdToGlobalBindPoseJointMap, MeshIdToGlobalBindPoseReferenceMap, bBadBindPoseMessageDisplay, bIsOriginalJoint, Parser.JointOrientationMatrix);
+
+						FTransform GlobalBindPoseJointTransform = GetConvertedTransform(GlobalBindPoseJointMatrix);
+						UnrealNode->SetGlobalBindPoseReferenceForMeshUIDs(MeshIdToGlobalBindPoseReferenceMap);
+						//Add Transform attribute per Geometry
+						//Temporary (to fix in 5.6 main branch): We need to make an API like Get/Set GlobalBindPoseReferenceForMeshUIDs for the joint one
+						for (const TPair<FString, FMatrix>& MeshIdAndBindPoseJointMatrix : MeshIdToGlobalBindPoseJointMap)
+						{
+							FString AttributeKey = TEXT("JointBindPosePerMesh_") + MeshIdAndBindPoseJointMatrix.Key;
+							UnrealNode->RegisterAttribute<FMatrix>(FAttributeKey(AttributeKey), MeshIdAndBindPoseJointMatrix.Value);
+						}
+
+						if (ParentNode != nullptr)
+						{
+							FbxAMatrix GlobalFbxParentMatrix = SDKScene->GetAnimationEvaluator()->GetNodeGlobalTransform(ParentNode, 0);
+							TMap<FString, FMatrix> ParentMeshIdToGlobalBindPoseReferenceMap;
+							TMap<FString, FMatrix> ParentMeshIdToGlobalBindPoseJointMap;
+							FFbxMesh::GetGlobalJointBindPoseTransform(&Parser, SDKScene, ParentNode, FbxJointMeshBindPoseGenerator, GlobalFbxParentMatrix, ParentMeshIdToGlobalBindPoseJointMap, ParentMeshIdToGlobalBindPoseReferenceMap, bBadBindPoseMessageDisplay, bIsOriginalJoint, Parser.JointOrientationMatrix);
 						
-						FbxAMatrix LocalFbxMatrix = GlobalFbxParentMatrix.Inverse() * GlobalBindPoseJointMatrix;
-						FTransform LocalBindPoseJointTransform = GetConvertedTransform(LocalFbxMatrix);
+							FbxAMatrix LocalFbxMatrix = GlobalFbxParentMatrix.Inverse() * GlobalBindPoseJointMatrix;
+							FTransform LocalBindPoseJointTransform = GetConvertedTransform(LocalFbxMatrix);
 
-						UnrealNode->SetCustomBindPoseLocalTransform(&NodeContainer, LocalBindPoseJointTransform, bResetCache);
-					}
-					else
-					{
-						//No parent, set the same matrix has the global
-						UnrealNode->SetCustomBindPoseLocalTransform(&NodeContainer, GlobalBindPoseJointTransform, bResetCache);
+							UnrealNode->SetCustomBindPoseLocalTransform(&NodeContainer, LocalBindPoseJointTransform, bResetCache);
+						}
+						else
+						{
+							//No parent, set the same matrix has the global
+							UnrealNode->SetCustomBindPoseLocalTransform(&NodeContainer, GlobalBindPoseJointTransform, bResetCache);
+						}
 					}
 
 					//Get time Zero transform for this joint
@@ -261,6 +246,8 @@ namespace UE
 
 						//Set the global node transform
 						FbxAMatrix GlobalFbxMatrix = SDKScene->GetAnimationEvaluator()->GetNodeGlobalTransform(Node, 0);
+						//To be able to use the GlobalTimeZero as is (without JointOrientationMatrix application) for T0 Rebinding, we store the Calculated GlobalTimeZero here.
+						UnrealNode->SetCustomGlobalMatrixForT0Rebinding(FFbxConvert::ConvertMatrix<FMatrix>(GlobalFbxMatrix));
 						GlobalFbxMatrix = GlobalFbxMatrix * Parser.JointOrientationMatrix;
 
 						FTransform GlobalTransform = GetConvertedTransform(GlobalFbxMatrix);
@@ -285,6 +272,29 @@ namespace UE
 						}
 					}
 
+					{
+						//Set Default transforms with JointOrientation applied:
+						FbxAMatrix GlobalFbxMatrix = Node->EvaluateGlobalTransform() * Parser.JointOrientationMatrix;
+						if (ParentNode)
+						{
+							FbxAMatrix GlobalFbxParentMatrix = ParentNode->EvaluateGlobalTransform();
+							if (!bIsRootJoint)
+							{
+								GlobalFbxParentMatrix = GlobalFbxParentMatrix * Parser.JointOrientationMatrix;
+							}
+
+							FbxAMatrix	LocalFbxMatrix = GlobalFbxParentMatrix.Inverse() * GlobalFbxMatrix;
+							FTransform LocalTransform = GetConvertedTransform(LocalFbxMatrix);
+							UnrealNode->SetCustomLocalTransform(&NodeContainer, LocalTransform, bResetCache);
+						}
+						else
+						{
+							FTransform GlobalTransform = GetConvertedTransform(GlobalFbxMatrix);
+							//No parent, set the same matrix has the global
+							UnrealNode->SetCustomLocalTransform(&NodeContainer, GlobalTransform, bResetCache);
+						}
+					}
+
 					FString JointNodeName = Parser.GetFbxHelper()->GetFbxObjectName(Node, true);
 					UnrealNode->SetDisplayLabel(JointNodeName);
 				};
@@ -329,7 +339,7 @@ namespace UE
 						//No break since the eNull act has a skeleton if possible
 						case FbxNodeAttribute::eSkeleton:
 						{
-							ApplySkeletonAttribute();
+							ApplySkeletonAttribute(true);
 							bIsNodeContainJointAttribute = true;
 							break;
 						}
@@ -392,14 +402,38 @@ namespace UE
 					if (ForceJointNodes.Contains(Node))
 					{
 						UnrealNode->AddSpecializedType(FSceneNodeStaticData::GetTransformSpecializeTypeString());
-						ApplySkeletonAttribute();
+						ApplySkeletonAttribute(false);
 					}
 					else if (!bIsRootNode && IsNodeUnderCommonJointRootNode(Node, CommonJointRootNodes))
 					{
+						FString Name = Node->GetName();
 						UnrealNode->AddSpecializedType(FSceneNodeStaticData::GetTransformSpecializeTypeString());
-						ApplySkeletonAttribute();
+						ApplySkeletonAttribute(false);
+					}
+					else
+					{
+						//ApplySkeletonAttribute() will set the default transforms.
+
+						//Set the node default transform
+						{
+							FbxAMatrix GlobalFbxMatrix = Node->EvaluateGlobalTransform();
+							FTransform GlobalTransform = GetConvertedTransform(GlobalFbxMatrix);
+							if (FbxNode* ParentNode = Node->GetParent())
+							{
+								FbxAMatrix GlobalFbxParentMatrix = ParentNode->EvaluateGlobalTransform();
+								FbxAMatrix	LocalFbxMatrix = GlobalFbxParentMatrix.Inverse() * GlobalFbxMatrix;
+								FTransform LocalTransform = GetConvertedTransform(LocalFbxMatrix);
+								UnrealNode->SetCustomLocalTransform(&NodeContainer, LocalTransform, bResetCache);
+							}
+							else
+							{
+								//No parent, set the same matrix has the global
+								UnrealNode->SetCustomLocalTransform(&NodeContainer, GlobalTransform, bResetCache);
+							}
+						}
 					}
 				}
+
 				
 				auto AddAnimationTrackNode = [&](EInterchangePropertyTracks PropertyTrack, const FString& CurveNodeName, const FString& PayloadKey, EInterchangeAnimationPayLoadType PayloadType)
 				{
