[Interchange][FBX] UFBX failing FBX asset import tests (missing features)

- vertex colors
- metadata

diff --git a/Engine/Plugins/Interchange/Runtime/Source/Parsers/Fbx/Private/Ufbx/UfbxConvert.h b/Engine/Plugins/Interchange/Runtime/Source/Parsers/Fbx/Private/Ufbx/UfbxConvert.h
--- a/Engine/Plugins/Interchange/Runtime/Source/Parsers/Fbx/Private/Ufbx/UfbxConvert.h
+++ b/Engine/Plugins/Interchange/Runtime/Source/Parsers/Fbx/Private/Ufbx/UfbxConvert.h
@@ -32,7 +32,7 @@ namespace Convert
 
 	struct FNodeNameAndUid
 	{
-		// todo: currently replicates FBX SDK's behavior to name root node "RootNode".
+		// #ufbx_todo: currently replicates FBX SDK's behavior to name root node "RootNode".
 		//		But the root node doesn't have actual name, in fact. Also it might be better to name this node after the file name?
 		FNodeNameAndUid(FUfbxParser& Parser, UInterchangeSceneNode* ParentSceneNode, const ufbx_node& Node);
 
@@ -65,6 +65,16 @@ namespace Convert
 		return FQuat(Q.x, Q.y, Q.z, Q.w);
 	}
 
+	inline FLinearColor ConvertColor(const ufbx_vec4& Color) 
+	{
+		// FBX colors are sRGB
+		FVector4d ConvertedColor = Convert::ConvertVec4(Color
+			).ComponentMax(FVector4d::Zero()
+			).ComponentMin(FVector4d::One())
+			*255.f;
+		return FLinearColor(FColor{uint8(ConvertedColor.X), uint8(ConvertedColor.Y), uint8(ConvertedColor.Z), uint8(ConvertedColor.W)});
+	}
+
 	inline FTransform ConvertTransform(const ufbx_transform& InTransform)
 	{
 		FTransform Result;
diff --git a/Engine/Plugins/Interchange/Runtime/Source/Parsers/Fbx/Private/Ufbx/UfbxMesh.cpp b/Engine/Plugins/Interchange/Runtime/Source/Parsers/Fbx/Private/Ufbx/UfbxMesh.cpp
--- a/Engine/Plugins/Interchange/Runtime/Source/Parsers/Fbx/Private/Ufbx/UfbxMesh.cpp
+++ b/Engine/Plugins/Interchange/Runtime/Source/Parsers/Fbx/Private/Ufbx/UfbxMesh.cpp
@@ -72,6 +72,11 @@ namespace UE::Interchange::Private
 			return Convert::ConvertVec3(ufbx_get_vertex_vec3(&Mesh.vertex_normal, TriangulatedIndex));
 		}
 
+		FLinearColor ConvertVertexColor(uint32_t TriangulatedIndex) const
+		{
+			return Convert::ConvertColor(ufbx_get_vertex_vec4(&Mesh.vertex_color, TriangulatedIndex));
+		}
+
 		FVector ConvertVertexTangent(uint32_t TriangulatedIndex) const
 		{
 			return Convert::ConvertVec3(ufbx_get_vertex_vec3(&Mesh.vertex_tangent, TriangulatedIndex));
@@ -111,10 +116,9 @@ namespace UE::Interchange::Private
 					FVertexInstanceID VertexInstanceID = Builder.AppendInstance(VertexMapping[SrcVertexIndices[Vertex]]);
 					Triangle[Vertex] = VertexInstanceID;
 
+					const uint32_t TriangulatedIndex = TriangulatedIndices[Vertex];
 					if (Mesh.vertex_normal.exists)
 					{
-						const uint32_t TriangulatedIndex = TriangulatedIndices[Vertex];
-
 						FVector Normal = ConvertVertexNormal(TriangulatedIndex);
 
 						if (Mesh.vertex_tangent.exists && Mesh.vertex_bitangent.exists)
@@ -130,7 +134,12 @@ namespace UE::Interchange::Private
 						{
 							Builder.SetInstanceNormal(VertexInstanceID, Normal);
 						}
+					}
 
+					if (Mesh.vertex_color.exists)
+					{
+						FLinearColor Color(ConvertVertexColor(TriangulatedIndex));
+						Builder.SetInstanceColor(VertexInstanceID, FVector4f(Color.R, Color.G, Color.B, Color.A));
 					}
 				}
 
@@ -220,7 +229,7 @@ namespace UE::Interchange::Private
 			bool bMeshHasVertexBinormal = Mesh->vertex_bitangent.exists && Mesh->vertex_tangent.exists;
 			bool bMeshHasVertexTangent = bMeshHasVertexBinormal;
 			bool bMeshHasSmoothGroup = true;
-			bool bMeshHasVertexColor = false; // #ufbx_todo: 
+			bool bMeshHasVertexColor = Mesh->vertex_color.exists;
 			int32 MeshUVCount = Mesh->uv_sets.count;
 			
 			// Why are these needed?
diff --git a/Engine/Plugins/Interchange/Runtime/Source/Parsers/Fbx/Private/Ufbx/UfbxScene.cpp b/Engine/Plugins/Interchange/Runtime/Source/Parsers/Fbx/Private/Ufbx/UfbxScene.cpp
--- a/Engine/Plugins/Interchange/Runtime/Source/Parsers/Fbx/Private/Ufbx/UfbxScene.cpp
+++ b/Engine/Plugins/Interchange/Runtime/Source/Parsers/Fbx/Private/Ufbx/UfbxScene.cpp
@@ -6,6 +6,7 @@
 
 #include "InterchangeMeshNode.h"
 #include "Nodes/InterchangeBaseNodeContainer.h"
+#include "Nodes/InterchangeUserDefinedAttribute.h"
 
 #define LOCTEXT_NAMESPACE "InterchangeFbxParser"
 
@@ -144,6 +145,88 @@ namespace UE::Interchange::Private
 		{
 			SceneNode->AddSpecializedType(FSceneNodeStaticData::GetLodGroupSpecializeTypeString());
 		}
+
+		for (const ufbx_prop& Prop : InNode.props.props)
+		{
+			if (Prop.flags & UFBX_PROP_FLAG_USER_DEFINED)
+			{
+				FString PropertyName = Convert::ToUnrealString(Prop.name);
+				TOptional<FString> PayloadKey;
+				// #ufbx_todo: animated prop; PayloadKey will be added then  
+
+				switch (Prop.type)
+				{
+				case UFBX_PROP_UNKNOWN:
+					break;
+				case UFBX_PROP_BOOLEAN:
+					{
+						const bool PropertyValue = Prop.value_int != 0;
+						UInterchangeUserDefinedAttributesAPI::CreateUserDefinedAttribute(SceneNode, PropertyName, PropertyValue, PayloadKey);
+					}
+					break;
+				case UFBX_PROP_INTEGER:
+					{
+						const int32 PropertyValue = Prop.value_int;
+						UInterchangeUserDefinedAttributesAPI::CreateUserDefinedAttribute(SceneNode, PropertyName, PropertyValue, PayloadKey);
+					}
+					break;
+				case UFBX_PROP_NUMBER:
+					{
+						const float PropertyValue = Prop.value_real;
+						UInterchangeUserDefinedAttributesAPI::CreateUserDefinedAttribute(SceneNode, PropertyName, PropertyValue, PayloadKey);
+					}
+					break;
+				case UFBX_PROP_VECTOR:
+					{
+						const FVector4d PropertyValue = Convert::ConvertVec4(Prop.value_vec4);
+						UInterchangeUserDefinedAttributesAPI::CreateUserDefinedAttribute(SceneNode, PropertyName, PropertyValue, PayloadKey);
+					}
+					break;
+				case UFBX_PROP_COLOR:
+					{
+						const FVector PropertyValue = FVector(Convert::ConvertColor(Prop.value_vec4));
+						UInterchangeUserDefinedAttributesAPI::CreateUserDefinedAttribute(SceneNode, PropertyName, PropertyValue, PayloadKey);
+					}
+					break;
+				case UFBX_PROP_COLOR_WITH_ALPHA:
+					{
+						const FLinearColor PropertyValue = Convert::ConvertColor(Prop.value_vec4);
+						UInterchangeUserDefinedAttributesAPI::CreateUserDefinedAttribute(SceneNode, PropertyName, PropertyValue, PayloadKey);
+					}
+					break;
+				case UFBX_PROP_STRING:
+					{
+						const FString PropertyValue = Convert::ToUnrealString(Prop.value_str);
+						UInterchangeUserDefinedAttributesAPI::CreateUserDefinedAttribute(SceneNode, PropertyName, PropertyValue, PayloadKey);
+					}
+					break;
+				case UFBX_PROP_DATE_TIME:
+					{
+						const FDateTime PropertyValue(Prop.value_int);
+						UInterchangeUserDefinedAttributesAPI::CreateUserDefinedAttribute(SceneNode, PropertyName, PropertyValue, PayloadKey);
+					}
+					break;
+				// #ufbx_todo: support other useful custom attributes among these types
+				case UFBX_PROP_TRANSLATION:
+					break;
+				case UFBX_PROP_ROTATION:
+					break;
+				case UFBX_PROP_SCALING:
+					break;
+				case UFBX_PROP_DISTANCE:
+					break;
+				case UFBX_PROP_COMPOUND:
+					break;
+				case UFBX_PROP_BLOB:
+					break;
+				case UFBX_PROP_REFERENCE:
+					break;
+				case UFBX_PROP_TYPE_FORCE_32BIT:
+					break;
+				}
+			}
+		}
+
 	}
 
 	void FUfbxScene::ProcessNodes(UInterchangeBaseNodeContainer& NodeContainer)
