Interchange:
- Add a method to determine if Substrate Adaptive GBuffer is enabled

MaterialX:
- Log a warning for Standard Surface and OpenPBR materials if Adaptive GBuffer format is not selected

diff --git a/Engine/Plugins/Interchange/Runtime/Source/Import/Private/InterchangeImportModule.cpp b/Engine/Plugins/Interchange/Runtime/Source/Import/Private/InterchangeImportModule.cpp
--- a/Engine/Plugins/Interchange/Runtime/Source/Import/Private/InterchangeImportModule.cpp
+++ b/Engine/Plugins/Interchange/Runtime/Source/Import/Private/InterchangeImportModule.cpp
@@ -71,10 +71,15 @@ class FInterchangeImportModule : public IInterchangeImportModule
 	{
 		return bIsSubstrateEnabled && GInterchangeEnableSubstrate;
 	}
+	virtual bool IsSubstrateAdaptiveGBufferEnabled() const override
+	{
+		return bIsSubstrateAdaptiveGBufferEnabled;
+	}
 
 private:
 
 	bool bIsSubstrateEnabled = false;
+	bool bIsSubstrateAdaptiveGBufferEnabled = false;
 };
 
 IMPLEMENT_MODULE(FInterchangeImportModule, InterchangeImport)
@@ -83,7 +88,11 @@ IMPLEMENT_MODULE(FInterchangeImportModule, InterchangeImport)
 
 void FInterchangeImportModule::StartupModule()
 {
-	bIsSubstrateEnabled = GetDefault<URendererSettings>() ? GetDefault<URendererSettings>()->bEnableSubstrate : false;
+	if (const URendererSettings* RendererSettings = GetDefault<URendererSettings>())
+	{
+		bIsSubstrateEnabled = RendererSettings->bEnableSubstrate;
+		bIsSubstrateAdaptiveGBufferEnabled = RendererSettings->SubstrateGBufferFormat == ESubstrateStorageFormat::Type::AdaptiveBuffer;
+	}
 
 	FInterchangeImportMaterialAsyncHelper& InterchangeMaterialAsyncHelper = FInterchangeImportMaterialAsyncHelper::GetInstance();
 
diff --git a/Engine/Plugins/Interchange/Runtime/Source/Import/Private/MaterialX/MaterialXUtils/MaterialXBase.cpp b/Engine/Plugins/Interchange/Runtime/Source/Import/Private/MaterialX/MaterialXUtils/MaterialXBase.cpp
--- a/Engine/Plugins/Interchange/Runtime/Source/Import/Private/MaterialX/MaterialXUtils/MaterialXBase.cpp
+++ b/Engine/Plugins/Interchange/Runtime/Source/Import/Private/MaterialX/MaterialXUtils/MaterialXBase.cpp
@@ -11,6 +11,7 @@ namespace mx = MaterialX;
 FMaterialXBase::FMaterialXBase(UInterchangeBaseNodeContainer& BaseNodeContainer)
 	: NodeContainer(BaseNodeContainer)
 	, bIsSubstrateEnabled(FMaterialXManager::GetInstance().IsSubstrateEnabled())
+	, bIsSubstrateAdaptiveGBufferEnabled(FMaterialXManager::GetInstance().IsSubstrateAdaptiveGBufferEnabled())
 {}
 
 void FMaterialXBase::UpdateDocumentRecursively(MaterialX::GraphElementPtr Graph)
diff --git a/Engine/Plugins/Interchange/Runtime/Source/Import/Private/MaterialX/MaterialXUtils/MaterialXManager.cpp b/Engine/Plugins/Interchange/Runtime/Source/Import/Private/MaterialX/MaterialXUtils/MaterialXManager.cpp
--- a/Engine/Plugins/Interchange/Runtime/Source/Import/Private/MaterialX/MaterialXUtils/MaterialXManager.cpp
+++ b/Engine/Plugins/Interchange/Runtime/Source/Import/Private/MaterialX/MaterialXUtils/MaterialXManager.cpp
@@ -351,6 +351,7 @@ FMaterialXManager::FMaterialXManager()
 		{mx::Category::NormalMap, {"normal", "tangent", "bitangent"}}, // FlattenNormalMap Material Function doesn't have normal/tangent inputs, we just remove them to avoid unnecessary connections
 	}
 	, bIsSubstrateEnabled{ IInterchangeImportModule::IsAvailable() ? IInterchangeImportModule::Get().IsSubstrateEnabled() : false }
+	, bIsSubstrateAdaptiveGBufferEnabled{ IInterchangeImportModule::IsAvailable() ? IInterchangeImportModule::Get().IsSubstrateAdaptiveGBufferEnabled() : false }
 {
 	// Load in the Game thread
 	MaterialFunctionPackage(UE::MaterialFunctions::Path::MxAnd);
@@ -675,6 +676,11 @@ bool FMaterialXManager::IsSubstrateEnabled() const
 	return bIsSubstrateEnabled;
 }
 
+bool FMaterialXManager::IsSubstrateAdaptiveGBufferEnabled() const
+{
+	return bIsSubstrateAdaptiveGBufferEnabled;
+}
+
 bool FMaterialXManager::FilterNodeGraph(MaterialX::NodePtr Node) const
 {
 	// the test seems counterintuitive, but MaterialX check is "!filter" in the flattenSubgraphs functions
diff --git a/Engine/Plugins/Interchange/Runtime/Source/Import/Private/MaterialX/MaterialXUtils/MaterialXOpenPBRSurfaceShader.cpp b/Engine/Plugins/Interchange/Runtime/Source/Import/Private/MaterialX/MaterialXUtils/MaterialXOpenPBRSurfaceShader.cpp
--- a/Engine/Plugins/Interchange/Runtime/Source/Import/Private/MaterialX/MaterialXUtils/MaterialXOpenPBRSurfaceShader.cpp
+++ b/Engine/Plugins/Interchange/Runtime/Source/Import/Private/MaterialX/MaterialXUtils/MaterialXOpenPBRSurfaceShader.cpp
@@ -48,6 +48,11 @@ UInterchangeBaseNode* FMaterialXOpenPBRSurfaceShader::Translate(MaterialX::NodeP
 
 	if(bIsSubstrateEnabled)
 	{
+		if (!bIsSubstrateAdaptiveGBufferEnabled)
+		{
+			MTLX_LOG("MaterialXOpenPBRSurfaceShader", "OpenPBR material rendering might be wrong. Please select Substrate Adaptive GBuffer format in the Project settings.");
+		}
+
 		auto CreateShaderNodeInternal = [this]<class T>(mx::NodePtr OpenPBRSurfaceNode, const FString & PathOrShaderType, const FString & DisplayLabel)
 		{
 			T* Node;
diff --git a/Engine/Plugins/Interchange/Runtime/Source/Import/Private/MaterialX/MaterialXUtils/MaterialXStandardSurfaceShader.cpp b/Engine/Plugins/Interchange/Runtime/Source/Import/Private/MaterialX/MaterialXUtils/MaterialXStandardSurfaceShader.cpp
--- a/Engine/Plugins/Interchange/Runtime/Source/Import/Private/MaterialX/MaterialXUtils/MaterialXStandardSurfaceShader.cpp
+++ b/Engine/Plugins/Interchange/Runtime/Source/Import/Private/MaterialX/MaterialXUtils/MaterialXStandardSurfaceShader.cpp
@@ -47,6 +47,11 @@ UInterchangeBaseNode* FMaterialXStandardSurfaceShader::Translate(mx::NodePtr Sta
 	// Outputs
 	if(bIsSubstrateEnabled)
 	{
+		if (!bIsSubstrateAdaptiveGBufferEnabled)
+		{
+			MTLX_LOG("MaterialXStandardSurfaceShader", "Standard Surface material rendering might be wrong. Please select Substrate Adaptive GBuffer format in the Project settings.");
+		}
+
 		auto CreateShaderNodeInternal = [this]<class T>(mx::NodePtr StandardSurfaceNode, const FString& PathOrShaderType, const FString& DisplayLabel)
 		{
 			T* Node;
diff --git a/Engine/Plugins/Interchange/Runtime/Source/Import/Public/InterchangeImportModule.h b/Engine/Plugins/Interchange/Runtime/Source/Import/Public/InterchangeImportModule.h
--- a/Engine/Plugins/Interchange/Runtime/Source/Import/Public/InterchangeImportModule.h
+++ b/Engine/Plugins/Interchange/Runtime/Source/Import/Public/InterchangeImportModule.h
@@ -34,12 +34,19 @@ public:
 	{
 		return FModuleManager::Get().IsModuleLoaded( "InterchangeImport" );
 	}
-
+	
 	/**
 	 * Checks to see if Substrate support is enabled both in the Project Settings and for Interchange
 	 * 
 	 * @return True if Substrate is enabled
 	 */
 	virtual bool IsSubstrateEnabled() const { return false; };
+
+	/**
+	 * Checks to see if Substrate Adaptive GBuffer is enabled in the Project Settings
+	 *
+	 * @return True if Adaptive GBuffer is enabled
+	 */
+	virtual bool IsSubstrateAdaptiveGBufferEnabled() const { return false; }
 };
 
diff --git a/Engine/Plugins/Interchange/Runtime/Source/Import/Public/MaterialX/MaterialXUtils/MaterialXBase.h b/Engine/Plugins/Interchange/Runtime/Source/Import/Public/MaterialX/MaterialXUtils/MaterialXBase.h
--- a/Engine/Plugins/Interchange/Runtime/Source/Import/Public/MaterialX/MaterialXUtils/MaterialXBase.h
+++ b/Engine/Plugins/Interchange/Runtime/Source/Import/Public/MaterialX/MaterialXUtils/MaterialXBase.h
@@ -171,5 +171,7 @@ protected:
 	const char* NodeDefinition; 
 
 	bool bIsSubstrateEnabled;
+
+	bool bIsSubstrateAdaptiveGBufferEnabled;
 };
 #endif
diff --git a/Engine/Plugins/Interchange/Runtime/Source/Import/Public/MaterialX/MaterialXUtils/MaterialXManager.h b/Engine/Plugins/Interchange/Runtime/Source/Import/Public/MaterialX/MaterialXUtils/MaterialXManager.h
--- a/Engine/Plugins/Interchange/Runtime/Source/Import/Public/MaterialX/MaterialXUtils/MaterialXManager.h
+++ b/Engine/Plugins/Interchange/Runtime/Source/Import/Public/MaterialX/MaterialXUtils/MaterialXManager.h
@@ -66,6 +66,8 @@ public:
 	
 	INTERCHANGEIMPORT_API bool IsSubstrateEnabled() const;
 
+	INTERCHANGEIMPORT_API bool IsSubstrateAdaptiveGBufferEnabled() const;
+
 	/** Return true if the node category should be filtered out during the flattening of the subgraphs*/
 	INTERCHANGEIMPORT_API bool FilterNodeGraph(MaterialX::NodePtr Node) const;
 
@@ -170,6 +172,8 @@ private:
 
 	bool bIsSubstrateEnabled;
 
+	bool bIsSubstrateAdaptiveGBufferEnabled;
+
 	friend bool UE::Interchange::MaterialX::AreMaterialFunctionPackagesLoaded();
 };
 #endif // WITH_EDITOR
\ No newline at end of file
