Interchange: Use hash of node UIDs to guarantee unique bone names when needed.

diff --git a/Engine/Plugins/Interchange/Runtime/Source/Import/Private/Mesh/InterchangeSkeletalMeshFactory.cpp b/Engine/Plugins/Interchange/Runtime/Source/Import/Private/Mesh/InterchangeSkeletalMeshFactory.cpp
--- a/Engine/Plugins/Interchange/Runtime/Source/Import/Private/Mesh/InterchangeSkeletalMeshFactory.cpp
+++ b/Engine/Plugins/Interchange/Runtime/Source/Import/Private/Mesh/InterchangeSkeletalMeshFactory.cpp
@@ -698,11 +698,13 @@ namespace UE
 				ImportAssetObjectLODDataRef.bDiffPose = false;
 				TArray <Private::FJointInfo> JointInfos;
 				TArray<FString> BoneNotBindNames;
+				TSet<FName> UsedBoneNames;
 				Private::FSkeletonHelper::RecursiveAddBones(Arguments.NodeContainer
 					, RootJointNodeId
 					, JointInfos
 					, INDEX_NONE
 					, ImportAssetObjectLODDataRef.RefBonesBinary
+					, UsedBoneNames
 					, ImportAssetObjectLODDataRef.bUseTimeZeroAsBindPose
 					, ImportAssetObjectLODDataRef.bDiffPose
 					, BoneNotBindNames
diff --git a/Engine/Plugins/Interchange/Runtime/Source/Import/Private/Mesh/InterchangeSkeletonHelper.cpp b/Engine/Plugins/Interchange/Runtime/Source/Import/Private/Mesh/InterchangeSkeletonHelper.cpp
--- a/Engine/Plugins/Interchange/Runtime/Source/Import/Private/Mesh/InterchangeSkeletonHelper.cpp
+++ b/Engine/Plugins/Interchange/Runtime/Source/Import/Private/Mesh/InterchangeSkeletonHelper.cpp
@@ -44,7 +44,8 @@ namespace UE::Interchange::Private
 		FReferenceSkeletonModifier RefSkelModifier(RefSkeleton, SkeletonAsset);
 		TArray <FJointInfo> JointInfos;
 		TArray<FString> BoneNotBindNames;
-		RecursiveAddBones(NodeContainer, RootJointNodeId, JointInfos, INDEX_NONE, RefBonesBinary, bUseTimeZeroAsBindPose, bOutDiffPose, BoneNotBindNames, bImportSockets);
+		TSet<FName> UsedBoneNames;
+		RecursiveAddBones(NodeContainer, RootJointNodeId, JointInfos, INDEX_NONE, RefBonesBinary, UsedBoneNames, bUseTimeZeroAsBindPose, bOutDiffPose, BoneNotBindNames, bImportSockets);
 		if (bOutDiffPose)
 		{
 			//bOutDiffPose can only be true if the user ask to bind with time zero transform.
@@ -245,11 +246,28 @@ namespace UE::Interchange::Private
 		, TArray <FJointInfo>& JointInfos
 		, int32 ParentIndex
 		, TArray<SkeletalMeshImportData::FBone>& RefBonesBinary
+		, TSet<FName>& UsedBoneNames
 		, const bool bUseTimeZeroAsBindPose
 		, bool& bOutDiffPose
 		, TArray<FString>& OutBoneNotBindNames
 		, bool bImportSockets)
 	{
+		// Produces a unique bone name for a particular node, and stores it on UsedBoneNames
+		auto GetBoneNameForNode = [&UsedBoneNames](const UInterchangeBaseNode* Node)
+		{
+			FString Name = Node->GetDisplayLabel();
+
+			if (UsedBoneNames.Contains(*Name))
+			{
+				const FString UniqueId = Node->GetUniqueID();
+				Name += TEXT("_") + FMD5::HashAnsiString(*UniqueId);
+			}
+
+			ensure(!UsedBoneNames.Contains(*Name));
+			UsedBoneNames.Add(*Name);
+			return Name;
+		};
+
 		const UInterchangeBaseNode* BaseNode = NodeContainer->GetNode(JointNodeId);
 		const UInterchangeSceneNode* JointNode = Cast<UInterchangeSceneNode>(BaseNode);
 		if (!JointNode)
@@ -257,7 +275,7 @@ namespace UE::Interchange::Private
 			if (JointInfos.Num() == 0 && BaseNode)
 			{
 				FJointInfo& Info = JointInfos.AddZeroed_GetRef();
-				Info.Name = BaseNode->GetDisplayLabel();
+				Info.Name = GetBoneNameForNode(BaseNode);
 				Info.ParentIndex = INDEX_NONE;
 				Info.LocalTransform = FTransform::Identity;
 
@@ -280,7 +298,7 @@ namespace UE::Interchange::Private
 
 		int32 JointInfoIndex = JointInfos.Num();
 		FJointInfo& Info = JointInfos.AddZeroed_GetRef();
-		Info.Name = JointNode->GetDisplayLabel();
+		Info.Name = GetBoneNameForNode(JointNode);
 
 		FTransform LocalTransform;
 		FTransform TimeZeroLocalTransform;
@@ -355,7 +373,7 @@ namespace UE::Interchange::Private
 					//We dont want a socket to be a skeleton joint
 					continue;
 				}
-				RecursiveAddBones(NodeContainer, ChildrenIds[ChildIndex], JointInfos, JointInfoIndex, RefBonesBinary, bUseTimeZeroAsBindPose, bOutDiffPose, OutBoneNotBindNames, bImportSockets);
+				RecursiveAddBones(NodeContainer, ChildrenIds[ChildIndex], JointInfos, JointInfoIndex, RefBonesBinary, UsedBoneNames, bUseTimeZeroAsBindPose, bOutDiffPose, OutBoneNotBindNames, bImportSockets);
 			}
 		}
 	}
diff --git a/Engine/Plugins/Interchange/Runtime/Source/Import/Public/InterchangeSkeletonHelper.h b/Engine/Plugins/Interchange/Runtime/Source/Import/Public/InterchangeSkeletonHelper.h
--- a/Engine/Plugins/Interchange/Runtime/Source/Import/Public/InterchangeSkeletonHelper.h
+++ b/Engine/Plugins/Interchange/Runtime/Source/Import/Public/InterchangeSkeletonHelper.h
@@ -56,6 +56,7 @@ namespace UE::Interchange::Private
 			, TArray <FJointInfo>& JointInfos
 			, int32 ParentIndex
 			, TArray<SkeletalMeshImportData::FBone>& RefBonesBinary
+			, TSet<FName>& UsedBoneNames
 			, const bool bUseTimeZeroAsBindPose
 			, bool& bOutDiffPose
 			, TArray<FString>& OutBoneNotBindNames
