Interchange: Tweak previous fix to ensure minimum animation length in the correct way.

diff --git a/Engine/Plugins/Interchange/Runtime/Source/Parsers/CommonParser/Private/InterchangeCommonAnimationPayload.cpp b/Engine/Plugins/Interchange/Runtime/Source/Parsers/CommonParser/Private/InterchangeCommonAnimationPayload.cpp
--- a/Engine/Plugins/Interchange/Runtime/Source/Parsers/CommonParser/Private/InterchangeCommonAnimationPayload.cpp
+++ b/Engine/Plugins/Interchange/Runtime/Source/Parsers/CommonParser/Private/InterchangeCommonAnimationPayload.cpp
@@ -97,8 +97,8 @@ namespace UE::Interchange
 				RangeEndTime = 0;
 			}
 
-			const double BakeInterval = FMath::Max<double>(1.0 / BakeFrequency, MINIMUM_ANIMATION_LENGTH);
-			const double SequenceLength = FMath::Max<double>(RangeEndTime - RangeStartTime, BakeInterval);
+			const double BakeInterval = 1.0 / (BakeFrequency != 0.0 ? BakeFrequency : DEFAULT_SAMPLERATE);
+			const double SequenceLength = FMath::Max<double>(FMath::Max<double>(RangeEndTime - RangeStartTime, BakeInterval), MINIMUM_ANIMATION_LENGTH);
 			int32 BakeKeyCount = FMath::RoundToInt32(SequenceLength * BakeFrequency) + 1;
 
 			auto EvaluateCurve = [this](const int32& CurveIndex, double CurrentTime, float DefaultValue)
