[Interchange] MorphTarget and Curve name fixes.

#rnx

diff --git a/Engine/Plugins/Interchange/Runtime/Source/Import/Private/Animation/InterchangeAnimSequenceFactory.cpp b/Engine/Plugins/Interchange/Runtime/Source/Import/Private/Animation/InterchangeAnimSequenceFactory.cpp
--- a/Engine/Plugins/Interchange/Runtime/Source/Import/Private/Animation/InterchangeAnimSequenceFactory.cpp
+++ b/Engine/Plugins/Interchange/Runtime/Source/Import/Private/Animation/InterchangeAnimSequenceFactory.cpp
@@ -8,6 +8,7 @@
 #include "InterchangeAssetImportData.h"
 #include "InterchangeCommonAnimationPayload.h"
 #include "InterchangeCommonPipelineDataFactoryNode.h"
+#include "InterchangeHelper.h"
 #include "InterchangeImportCommon.h"
 #include "InterchangeImportLog.h"
 #include "InterchangeMeshNode.h"
@@ -107,7 +108,7 @@ namespace UE::Interchange::Private
 		for (int32 CurveIndex = 0; CurveIndex < Curves.Num(); ++CurveIndex)
 		{
 			FRichCurve& Curve = Curves[CurveIndex];
-			FName Name = *CurveNames[CurveIndex];
+			FName Name(UE::Interchange::ReplaceNoneWithNull(CurveNames[CurveIndex]));
 
 			if (bDoNotImportCurveWithZero)
 			{
diff --git a/Engine/Plugins/Interchange/Runtime/Source/Import/Private/Mesh/InterchangeMeshHelper.cpp b/Engine/Plugins/Interchange/Runtime/Source/Import/Private/Mesh/InterchangeMeshHelper.cpp
--- a/Engine/Plugins/Interchange/Runtime/Source/Import/Private/Mesh/InterchangeMeshHelper.cpp
+++ b/Engine/Plugins/Interchange/Runtime/Source/Import/Private/Mesh/InterchangeMeshHelper.cpp
@@ -1386,28 +1386,23 @@ namespace UE::Interchange::Private::MeshHelper
 		for (TPair<FString, TArray<FString>>& NamePerUIDs : KeysPerName)
 		{
 			const FString& MorphTargetNameBase = NamePerUIDs.Key;
-			FString MorphTargetName(MorphTargetNameBase);
+			FName MorphTargetName(UE::Interchange::ReplaceNoneWithNull(MorphTargetNameBase));
 
-			if (MorphTargetName.IsEmpty())
+			for (int32 Suffix = 1; UsedMorphTargetNames.Contains(MorphTargetName); Suffix++)
 			{
-				MorphTargetName = TEXT("__Empty__");
+				MorphTargetName = FName(FString::Printf(TEXT("%s_%d"), *MorphTargetNameBase, Suffix));
 			}
-
-			for (int32 Suffix = 1; UsedMorphTargetNames.Contains(*MorphTargetName); Suffix++)
-			{
-				MorphTargetName = FString::Printf(TEXT("%s_%d"), *MorphTargetNameBase, Suffix);
-			}
-			UsedMorphTargetNames.Add(*MorphTargetName);
+			UsedMorphTargetNames.Add(MorphTargetName);
 
 			if (MorphTargetName != MorphTargetNameBase)
 			{
-				UE_LOG(LogInterchangeImport, Warning, TEXT("Duplicate morph target '%s' found, renamed to '%s'."), *MorphTargetNameBase, *MorphTargetName);
+				UE_LOG(LogInterchangeImport, Warning, TEXT("Duplicate morph target '%s' found, renamed to '%s'."), *MorphTargetNameBase, *MorphTargetName.ToString());
 			}
 
-			DestinationMeshAttributes.RegisterMorphTargetAttribute(*MorphTargetName, bUseMorphTargetNormals);
+			DestinationMeshAttributes.RegisterMorphTargetAttribute(MorphTargetName, bUseMorphTargetNormals);
 
-			TVertexAttributesRef<FVector3f> DestinationMeshMorphPosDeltas = DestinationMeshAttributes.GetVertexMorphPositionDelta(*MorphTargetName);
-			TVertexInstanceAttributesRef<FVector3f> DestinationMeshMorphNormals = DestinationMeshAttributes.GetVertexInstanceMorphNormalDelta(*MorphTargetName);
+			TVertexAttributesRef<FVector3f> DestinationMeshMorphPosDeltas = DestinationMeshAttributes.GetVertexMorphPositionDelta(MorphTargetName);
+			TVertexInstanceAttributesRef<FVector3f> DestinationMeshMorphNormals = DestinationMeshAttributes.GetVertexInstanceMorphNormalDelta(MorphTargetName);
 
 			for (const FString& MorphTargetKey : NamePerUIDs.Value)
 			{
@@ -1484,7 +1479,7 @@ namespace UE::Interchange::Private::MeshHelper
 				}
 			}
 
-			SkeletonMorphCurveMetadataNames.Add(MorphTargetName);
+			SkeletonMorphCurveMetadataNames.Add(MorphTargetName.ToString());
 		}
 	}
 
