[Interchange] Adds core support for skeletal and static mesh Nanite assemblies (Nanite Foliage)

Highlights:

* UInterchangeMeshNode can now register and track assembly dependencies (AddAssemblyPartDependencyUid, GetAssemblyPartDependencies, etc.).
Translators must discover and create the required dependencies between base and part meshes.

* UE::Interchange::FMeshPayloadData has a new TOptionalUE::Interchange::FNaniteAssemblyDescription member.
In addition to setting the FMeshDescription during payload retrieval, translators can optionally choose to populate the assembly description as well.

* InterchangePipelineMeshUtilities now tracks the UIDs of meshes that have assembly dependencies (a base mesh) or is the target of a dependency (a part mesh).
This allows the factory node to factory node tracking to short-circuit for the common case of no Nanite assemblies.

* UInterchangeGenericMeshPipeline can now register factory node dependencies between assembly meshes.

* Added FInterchangeNaniteAssemblyBuilder utility class for skeletal and static meshes, which creates Nanite assembly nodes and parts from one or more FNaniteAssemblyDescription items.

* Updated InterchangeSkeletalMeshFactory and InterchangeStaticMeshFactory to build assembly data from mesh payload data using the  FInterchangeNaniteAssemblyBuilder utility.
This happens in SetupObject_GameThread so that all the required base and part mesh objects are ready.

Notes & Caveats:

Only non-combined static mesh assemblies are currently implemented. The combined static mesh mode currently tries to combine assembly part meshes as well as regular
meshes and base meshes. Part meshes must remain separate obviously, and therefore excluded from the pipeline pass that moves all the mesh descriptions under a single
factory node.

Because payload data is now potentially needed during SetupObject_GameThread (for InterchangeSkeletalMeshFactory and InterchangeStaticMeshFactory) the data must
persist and not be deleted in some prior step, such as ImportAsset_Async. Payload data is now cleared at the end of SetupObject_GameThread which does mean it is
retained a little longer than before.

The skeletal and static mesh factories must now loop over every LOD0 payload looking for nanite assembly data during SetupObject_GameThread. If the various TMap
lookups to acquire the data turn out to be overly costly, we may be able to cache the existence of assembly data in an earlier factory phase that is already dealing with
payload data. Also, there's nothing warning the user right now if assembly data has been attached to a non-LOD0 payload (which will just get silently dropped).

diff --git a/Engine/Plugins/Interchange/Runtime/Source/FactoryNodes/Private/InterchangeMeshFactoryNode.cpp b/Engine/Plugins/Interchange/Runtime/Source/FactoryNodes/Private/InterchangeMeshFactoryNode.cpp
--- a/Engine/Plugins/Interchange/Runtime/Source/FactoryNodes/Private/InterchangeMeshFactoryNode.cpp
+++ b/Engine/Plugins/Interchange/Runtime/Source/FactoryNodes/Private/InterchangeMeshFactoryNode.cpp
@@ -85,6 +85,12 @@ namespace UE::Interchange
 		return SlotMaterialDependency_BaseKey;
 	}
 
+	const FAttributeKey& FMeshFactoryNodeStaticData::GetAssemblyPartDependencyBaseKey()
+	{
+		static FAttributeKey SlotMaterialDependency_BaseKey(TEXT("__AssemblyPartDependency__"));
+		return SlotMaterialDependency_BaseKey;
+	}
+
 } // namespace UE::Interchange
 
 
@@ -92,6 +98,7 @@ UInterchangeMeshFactoryNode::UInterchangeMeshFactoryNode()
 {
 	LodDependencies.Initialize(Attributes, UE::Interchange::FMeshFactoryNodeStaticData::GetLodDependenciesBaseKey().ToString());
 	SlotMaterialDependencies.Initialize(Attributes.ToSharedRef(), UE::Interchange::FMeshFactoryNodeStaticData::GetSlotMaterialDependencyBaseKey().ToString());
+	AssemblyPartDependencies.Initialize(Attributes.ToSharedRef(), UE::Interchange::FMeshFactoryNodeStaticData::GetAssemblyPartDependencyBaseKey().ToString());
 }
 
 #if WITH_EDITOR
@@ -144,6 +151,21 @@ FString UInterchangeMeshFactoryNode::GetKeyDisplayName(const UE::Interchange::FA
 		}
 		return KeyDisplayName;
 	}
+	else if (NodeAttributeKey == UE::Interchange::FMeshFactoryNodeStaticData::GetAssemblyPartDependencyBaseKey())
+	{
+		KeyDisplayName = TEXT("Nanite Assembly Part Dependencies Count");
+	}
+	else if (NodeAttributeKeyString.StartsWith(UE::Interchange::FMeshFactoryNodeStaticData::GetAssemblyPartDependencyBaseKey().ToString()))
+	{
+		KeyDisplayName = TEXT("Nanite Assembly Part Dependencies Index ");
+		const FString IndexKey = UE::Interchange::TArrayAttributeHelper<FString>::IndexKey();
+		int32 IndexPosition = NodeAttributeKeyString.Find(IndexKey) + IndexKey.Len();
+		if (IndexPosition < NodeAttributeKeyString.Len())
+		{
+			KeyDisplayName += NodeAttributeKeyString.RightChop(IndexPosition);
+		}
+		return KeyDisplayName;
+	}
 	return Super::GetKeyDisplayName(NodeAttributeKey);
 }
 
@@ -157,6 +179,10 @@ FString UInterchangeMeshFactoryNode::GetAttributeCategory(const UE::Interchange:
 	{
 		return TEXT("Materials Dependencies");
 	}
+	else if (NodeAttributeKey.ToString().StartsWith(UE::Interchange::FMeshFactoryNodeStaticData::GetAssemblyPartDependencyBaseKey().ToString()))
+	{
+		return TEXT("Nanite Assembly Part Dependencies");
+	}
 	else
 	{
 		return Super::GetAttributeCategory(NodeAttributeKey);
@@ -302,6 +328,38 @@ bool UInterchangeMeshFactoryNode::ResetSlotMaterialDependencies()
 	return true;
 }
 
+int32 UInterchangeMeshFactoryNode::GetAssemblyPartDependenciesCount() const
+{
+	return AssemblyPartDependencies.GetCount();
+}
+
+void UInterchangeMeshFactoryNode::GetAssemblyPartDependencies(TMap<FString, FString>& OutAssemblyPartDependencies) const
+{
+	OutAssemblyPartDependencies = AssemblyPartDependencies.ToMap();
+}
+
+bool UInterchangeMeshFactoryNode::SetAssemblyPartDependencyUid(const FString& MeshUid, const FString& AssemblyPartDependencyUid)
+{
+	return AssemblyPartDependencies.SetKeyValue(MeshUid, AssemblyPartDependencyUid);
+}
+
+bool UInterchangeMeshFactoryNode::RemoveAssemblyPartDependencyUid(const FString& MeshUid)
+{
+	if (AssemblyPartDependencies.RemoveKey(MeshUid))
+	{
+		AssemblyPartDependencies.RebuildCache();
+		return true;
+	}
+	return false;
+}
+
+bool UInterchangeMeshFactoryNode::ResetAssemblyDependencies()
+{
+	AssemblyPartDependencies.Empty();
+	AssemblyPartDependencies.RebuildCache();
+	return true;
+}
+
 bool UInterchangeMeshFactoryNode::GetCustomLODGroup(FName& AttributeValue) const
 {
 	IMPLEMENT_NODE_ATTRIBUTE_GETTER(LODGroup, FName)
diff --git a/Engine/Plugins/Interchange/Runtime/Source/FactoryNodes/Public/InterchangeMeshFactoryNode.h b/Engine/Plugins/Interchange/Runtime/Source/FactoryNodes/Public/InterchangeMeshFactoryNode.h
--- a/Engine/Plugins/Interchange/Runtime/Source/FactoryNodes/Public/InterchangeMeshFactoryNode.h
+++ b/Engine/Plugins/Interchange/Runtime/Source/FactoryNodes/Public/InterchangeMeshFactoryNode.h
@@ -16,6 +16,7 @@ namespace UE::Interchange
 	{
 		static UE_API const FAttributeKey& GetLodDependenciesBaseKey();
 		static UE_API const FAttributeKey& GetSlotMaterialDependencyBaseKey();
+		static UE_API const FAttributeKey& GetAssemblyPartDependencyBaseKey();
 	};
 } // namespace Interchange
 
@@ -120,13 +121,33 @@ public:
 	UFUNCTION(BlueprintCallable, Category = "Interchange | FactoryNode | Mesh")
 	UE_API bool SetSlotMaterialDependencyUid(const FString& SlotName, const FString& MaterialDependencyUid);
 
-	/** Remove the Material dependency associated with the specfied slot name of this object. */
+	/** Remove the Material dependency associated with the specified slot name of this object. */
 	UFUNCTION(BlueprintCallable, Category = "Interchange | FactoryNode | Mesh")
 	UE_API bool RemoveSlotMaterialDependencyUid(const FString& SlotName);
 
 	/** Reset all the material dependencies. */
 	UFUNCTION(BlueprintCallable, Category = "Interchange | FactoryNode | Mesh")
 	UE_API bool ResetSlotMaterialDependencies();
+	
+	/** Retrieve the number of Nanite assembly part dependencies for this object. */
+	UFUNCTION(BlueprintCallable, Category = "Interchange | FactoryNode | Mesh")
+	UE_API int32 GetAssemblyPartDependenciesCount() const;
+
+	/** Retrieve the Nanite assembly part dependencies for this object. */
+	UFUNCTION(BlueprintCallable, Category = "Interchange | FactoryNode | Mesh")
+	UE_API void GetAssemblyPartDependencies(TMap<FString, FString>& OutAssemblyPartDependencies) const;
+
+	/** Add a Nanite assembly part dependency to this object. */
+	UFUNCTION(BlueprintCallable, Category = "Interchange | FactoryNode | Mesh")
+	UE_API bool SetAssemblyPartDependencyUid(const FString& MeshUid, const FString& AssemblyPartDependencyUid);
+
+	/** Remove a Nanite assembly part dependency associated with this object */
+	UFUNCTION(BlueprintCallable, Category = "Interchange | FactoryNode | Mesh")
+	UE_API bool RemoveAssemblyPartDependencyUid(const FString& MeshUid);
+
+	/** Reset all the Nanite assembly dependencies. */
+	UFUNCTION(BlueprintCallable, Category = "Interchange | FactoryNode | Mesh")
+	UE_API bool ResetAssemblyDependencies();
 
 	/** Query whether a custom LOD group is set for the mesh. */
 	UFUNCTION(BlueprintCallable, Category = "Interchange | FactoryNode | Mesh")
@@ -263,6 +284,7 @@ private:
 
 	UE::Interchange::TArrayAttributeHelper<FString> LodDependencies;
 	UE::Interchange::TMapAttributeHelper<FString, FString> SlotMaterialDependencies;
+	UE::Interchange::TMapAttributeHelper<FString, FString> AssemblyPartDependencies;
 
 protected:
 	virtual void FillAssetClassFromAttribute() PURE_VIRTUAL("FillAssetClassFromAttribute");
diff --git a/Engine/Plugins/Interchange/Runtime/Source/Import/Private/Mesh/InterchangeMeshHelper.cpp b/Engine/Plugins/Interchange/Runtime/Source/Import/Private/Mesh/InterchangeMeshHelper.cpp
--- a/Engine/Plugins/Interchange/Runtime/Source/Import/Private/Mesh/InterchangeMeshHelper.cpp
+++ b/Engine/Plugins/Interchange/Runtime/Source/Import/Private/Mesh/InterchangeMeshHelper.cpp
@@ -5,6 +5,7 @@
 #if WITH_EDITOR
 #include "BSPOps.h"
 #endif
+#include "Engine/NaniteAssemblyData.h"
 #include "Engine/Polys.h"
 #include "Engine/SkeletalMesh.h"
 #include "Engine/SkinnedAssetCommon.h"
@@ -18,7 +19,9 @@
 #include "Mesh/InterchangeMeshPayload.h"
 #include "MeshDescription.h"
 #include "MeshUVChannelInfo.h"
+#include "Misc/PackageName.h"
 #include "Model.h"
+#include "NaniteAssemblyDataBuilder.h"
 #include "Nodes/InterchangeBaseNodeContainer.h"
 #include "PhysicsEngine/AggregateGeom.h"
 #include "PhysicsEngine/BodySetup.h"
@@ -39,6 +42,8 @@
 #include "InterchangeStaticMeshLodDataNode.h"
 #include "Nodes/InterchangeSourceNode.h"
 
+#include <type_traits>
+
 namespace UE::Interchange::Private::MeshHelper
 {
 	bool AddConvexGeomFromVertices(
@@ -1626,4 +1631,474 @@ namespace UE::Interchange::Private::MeshHelper
 	}
 #endif
 
+#if WITH_EDITOR
+	/** Nanite assembly builder abstract base class (to hide underlying USkeletalMesh/UStaticMesh templated type) */
+
+	class FInterchangeNaniteAssemblyBuilder::ImplBase
+	{
+	public:
+
+		ImplBase() = default;
+		virtual ~ImplBase() = default;
+		UE_NONCOPYABLE(ImplBase);
+
+		virtual void AddDescription(const TOptional<UE::Interchange::FNaniteAssemblyDescription>& Description) = 0;
+	};
+
+	/** Nanite assembly builder implementation templated by either USkeletalMesh or UStaticMesh. The main responsibility of this class
+	  * is to translate instance data from one or more interchange FNaniteAssemblyDescription items, to node data on the underlying
+	  * FNaniteAssemblyDataBuilder (nanite assembly API refers to an instance as a 'node', and a prototype mesh as a 'part').
+	  * 
+	  * Note that in addition to the typical transform + prototype/part id per-instance data, skeletal mesh nanite assembly nodes also
+	  * require one or more FNaniteAssemblyBoneInfluence bindings (basically joint index and a weight). This data tells nanite where/how
+	  * to attach the root bone from the skeletal mesh part being instanced onto the main skeletal mesh skeleton.
+	  */
+
+	template<typename MeshType>
+	class FInternalNaniteAssemblyBuilder : public FInterchangeNaniteAssemblyBuilder::ImplBase
+	{
+		static_assert(
+			std::is_same_v<UStaticMesh, MeshType> || std::is_same_v<USkeletalMesh, MeshType>,
+			"MeshType must be either UStaticMesh or USkeletalMesh"
+		);
+
+		friend class FInterchangeNaniteAssemblyBuilder;
+
+		using FNaniteAssemblyDescription = UE::Interchange::FNaniteAssemblyDescription;
+
+		// List of USkeletalMesh/UStaticMesh objects per part/prototype index (most commonly just a single object)  
+		using FPartObjectList = TArray<TArray<MeshType*, TInlineAllocator<1>>>;
+
+		// The assembly builder instance used to create the assembly nodes and parts.
+		FNaniteAssemblyDataBuilder Builder;
+		
+		// The base USkeletalMesh/UStaticMesh object we're enabling nanite assembly on.
+		MeshType* Mesh = nullptr;
+
+		// The UInterchangeMeshFactoryNode for this mesh that holds the part dependency information.
+		const UInterchangeMeshFactoryNode* MeshFactoryNode = nullptr;
+
+		// The base node container
+		const UInterchangeBaseNodeContainer* NodeContainer = nullptr;
+
+		bool bCtorIsValid = false;
+
+		FInternalNaniteAssemblyBuilder(MeshType* InMesh
+			, const UInterchangeMeshFactoryNode* InMeshFactoryNode
+			, const UInterchangeBaseNodeContainer* InNodeContainer)
+		  : Mesh(InMesh)
+		  , MeshFactoryNode(InMeshFactoryNode)
+		  , NodeContainer(InNodeContainer)
+		  , bCtorIsValid(Mesh && MeshFactoryNode && NodeContainer)
+		{
+			if (!ensure(bCtorIsValid))
+			{
+				return;
+			}
+
+			// Add the initial base mesh materials to the builder
+			ForEachMeshMaterial(
+				Mesh,
+				[this](int32 LocalMaterialIndex, const TObjectPtr<UMaterialInterface>& Material)
+				{
+					Builder.AddMaterialSlot(Material);
+				}
+			);
+		}
+
+		static constexpr bool IsSkeletalMesh()
+		{
+			return std::is_same_v<USkeletalMesh, MeshType>;
+		}
+
+		static constexpr bool IsStaticMesh()
+		{
+			return std::is_same_v<UStaticMesh, MeshType>;
+		}
+
+		~FInternalNaniteAssemblyBuilder()
+		{
+			EnableNaniteAndFinalize();
+		}
+
+		/** Main entry point to add an interchange FNaniteAssemblyDescription to the builder */
+		void AddDescription(const TOptional<FNaniteAssemblyDescription>& InDescription) override
+		{
+			if (!bCtorIsValid || !InDescription.IsSet())
+			{
+				return;
+			}
+
+			const FNaniteAssemblyDescription& Description = InDescription.GetValue();
+
+			// Verify the description data is well-formed. This check will also ensure that iterating over the
+			// bone influence data (based on the description's supplied stride data) later on is safe. 
+			if (FString Reason; !Description.IsValid(&Reason))
+			{
+				UE_LOG(LogInterchangeImport, Warning 
+					,TEXT("Invalid description while building nanite assembly data for mesh '%s' - %s")
+					,*MeshFactoryNode->GetDisplayLabel()
+					,*Reason
+				);
+				return;
+			}
+
+			// Skeletal mesh assemblies must have binding data.
+			if (IsSkeletalMesh() && Description.BoneInfluences.IsEmpty())
+			{
+				UE_LOG(LogInterchangeImport, Warning
+					,TEXT("Invalid description while building nanite assembly data for mesh '%s' - Skeletal "
+						  "meshes must have at least one bone influence per instance.")
+					,*MeshFactoryNode->GetDisplayLabel()
+				);
+				return;
+			}
+
+			// Get valid UStaticMesh/USkeletal objects from the descriptions part list.
+			FPartObjectList PartObjects;
+			if (!GetPartObjects(Description, PartObjects))
+			{
+				return;
+			}
+
+			// Setup our binding data. Each instance can have a variable number of bindings as specified by the
+			// NumInfluencesPerInstance field on the description, and so the pointer must be advanced accordingly.
+			// Note that the call to Description.IsValid() above will have checked that the binding array is 
+			// appropriately sized. Note also that when adding static mesh nodes we must still supply an empty binding.
+
+			FNaniteAssemblyBoneInfluence* BoneInfluenceData = nullptr;
+			TArray<FNaniteAssemblyBoneInfluence> NaniteBoneInfluences;
+			NaniteBoneInfluences.Reserve(Description.BoneInfluences.Num());
+			if constexpr (IsSkeletalMesh())
+			{
+				for (const FNaniteAssemblyDescription::FBoneInfluence& Influence : Description.BoneInfluences)
+				{
+					NaniteBoneInfluences.Emplace(Influence.BoneIndex, Influence.BoneWeight);
+				}
+				BoneInfluenceData = NaniteBoneInfluences.GetData();
+			}
+
+			TArrayView<const FNaniteAssemblyBoneInfluence> Binding{};
+
+			// Loop over every instance and add a node for each part object.
+			const int32 NumInstances = Description.Transforms.Num();
+			int32 NumInstancesWithInvalidPartIndex = 0;
+			for (int32 Index = 0; Index < NumInstances; ++Index)
+			{
+				// Update the skeletal mesh binding view
+				if constexpr (IsSkeletalMesh())
+				{
+					const int32 NumInfluences = Description.NumInfluencesPerInstance[Index];
+					Binding = MakeArrayView(BoneInfluenceData, NumInfluences);
+					BoneInfluenceData += NumInfluences;
+				}
+
+				const int32 PartIndex = Description.PartIndices[Index];
+				if (PartObjects.IsValidIndex(PartIndex))
+				{
+					for (MeshType* PartObject : PartObjects[PartIndex])
+					{
+						AddNodeAndBindMaterials(PartObject, Description.Transforms[Index], Binding);
+					}
+				}
+				else
+				{
+					NumInstancesWithInvalidPartIndex++;
+				}
+			}
+
+			// Report invalid instance indices and return false if we failed to add any nodes at all.
+			if (NumInstancesWithInvalidPartIndex > 0)
+			{
+				UE_LOG(LogInterchangeImport, Warning
+					, TEXT("Nanite assembly description part index data for mesh '%s' contains one or more invalid indices.")
+					, *MeshFactoryNode->GetDisplayLabel()
+				);
+
+				if (NumInstancesWithInvalidPartIndex == NumInstances)
+				{
+					return;
+				}
+			}
+
+			return;
+		}
+
+		/** Enable nanite and finalize/apply the FNaniteAssemblyDataBuilder data to the mesh. If the finalize process 
+		  * fails for any reason, revert the nanite enabled flag to the original value.
+		  */
+		void EnableNaniteAndFinalize()
+		{
+			if (!bCtorIsValid)
+			{
+				return;
+			}
+
+			if (!Builder.GetData().IsValid())
+			{
+				UE_LOG(LogInterchangeImport, Warning
+					,TEXT("Nanite assembly builder data for mesh '%s' is invalid, skipping.")
+					, *MeshFactoryNode->GetDisplayLabel()
+				);
+				return;
+			}
+
+			// Enable nanite, apply the builder data and re-cache references ...
+
+			const bool bOriginalNaniteEnabledValue = SetNaniteEnabled(true);
+
+			bool ApplyResult = false;
+
+			if constexpr (IsSkeletalMesh())
+			{
+				ApplyResult = Builder.ApplyToSkeletalMesh(*Mesh);
+			}
+			else if constexpr (IsStaticMesh())
+			{
+				ApplyResult = Builder.ApplyToStaticMesh(*Mesh);
+			}
+
+			if (!ApplyResult)
+			{
+				SetNaniteEnabled(bOriginalNaniteEnabledValue);
+
+				UE_LOG(LogInterchangeImport, Warning
+					, TEXT("Failed to apply assembly builder data to mesh '%s', skipping.")
+					, * MeshFactoryNode->GetDisplayLabel()
+				);
+				return;
+			}
+
+			Mesh->RecacheNaniteAssemblyReferences();
+
+			return;
+		}
+
+		/** Add a node to the FNaniteAssemblyDataBuilder. Also add a new part if this is the first time the part
+		  * the first time the part mesh has been encountered.
+		  * 
+		  * Note that the bool return value is not success/failure, but whether or not a new part was created.
+		  */
+		bool AddNode(const MeshType* PartMesh
+			, const FTransform& LocalTransform
+			, const TArrayView<const FNaniteAssemblyBoneInfluence>& Bindings
+			, int32& OutPartIndex)
+		{
+			const FSoftObjectPath MeshPath(PartMesh);
+			OutPartIndex = Builder.FindPart(MeshPath);
+
+			const bool bNewPart = (OutPartIndex == INDEX_NONE);
+			if (bNewPart)
+			{
+				OutPartIndex = Builder.AddPart(MeshPath);
+			}
+
+			Builder.AddNode(OutPartIndex, FTransform3f(LocalTransform), ENaniteAssemblyNodeTransformSpace::Local, Bindings);
+
+			return bNewPart;
+		}
+
+		/** Add a node and materials from the given part mesh to the builder */
+		void AddNodeAndBindMaterials(const MeshType* PartMesh
+			, const FTransform& LocalTransform
+			, const TArrayView<const FNaniteAssemblyBoneInfluence>& Bindings)
+		{
+			int32 PartIndex;
+			if (AddNode(PartMesh, LocalTransform, Bindings, PartIndex))
+			{
+				ForEachMeshMaterial(
+					PartMesh,
+					[this, PartIndex](int32 LocalMaterialIndex, const TObjectPtr<UMaterialInterface>& PartMaterial)
+					{
+						int32 MaterialIndex = INDEX_NONE;
+						
+						// For now we always merge identical materials. There are however more advanced workflows 
+						// that require exact control over material and slot management apparently, and so more
+						// flexibility may be required in the future.
+						constexpr bool bMergeIdenticalMaterialSlots = true;
+						if (bMergeIdenticalMaterialSlots)
+						{
+							MaterialIndex = Builder.GetMaterialSlots().IndexOfByPredicate(
+								[&PartMaterial](const FNaniteAssemblyDataBuilder::FMaterialSlot& ExistingMaterial)
+								{
+									return ExistingMaterial.Material == PartMaterial;
+								}
+							);
+						}
+
+						if (MaterialIndex == INDEX_NONE)
+						{
+							MaterialIndex = Builder.AddMaterialSlot(PartMaterial);
+						}
+
+						Builder.RemapPartMaterial(PartIndex, LocalMaterialIndex, MaterialIndex);
+					}
+				);
+			}
+		}
+
+		/* Resolve the string mesh uids and/or top-level asset references from the given description to valid
+         * USkeletalMesh/UStaticMesh objects. Note that invalid (null) objects will be omitted from the inner array.
+		 * The outer array length is guaranteed to match the number of parts in the input description however, so that 
+		 * each per-instance part index indexes into the correct list of part objects. 
+		 * 
+		 * Returns false if every uid was invalid and so we effectively have no parts to work with.
+		 * 
+		 */
+		bool GetPartObjects(const UE::Interchange::FNaniteAssemblyDescription& Description, FPartObjectList& OutPartObjects)
+		{
+			TSet<FString> MissingPartUids;
+			bool bFoundAtLeastOnePartObject = false;
+
+			TMap<FString, FString> PartUidToFactoryNodeUidTable;
+			MeshFactoryNode->GetAssemblyPartDependencies(PartUidToFactoryNodeUidTable);
+
+			// Note the length of output part array must match the length of the incoming part uids from the description.
+			OutPartObjects.SetNum(Description.PartUids.Num());
+
+			for (int Index = 0; Index < Description.PartUids.Num(); ++Index)
+			{
+				const TArray<FString>& PartUids = Description.PartUids[Index];
+				for (const FString& PartUid : PartUids)
+				{
+					if (PartUid.IsEmpty())
+					{
+						MissingPartUids.Add(TEXT(""));
+						continue;
+					}
+
+					// First check if this part relates to a factory node (most common case) and if so get the UObject from
+					// from the node. If that fails, try loading the part as an already existing asset.
+
+					bool bFoundPartObject = false;
+					if (const FString* FactoryNodeDependencyUid = PartUidToFactoryNodeUidTable.Find(PartUid))
+					{
+						if (const UInterchangeMeshFactoryNode* PartMeshFactoryNode = Cast<const UInterchangeMeshFactoryNode>(NodeContainer->GetFactoryNode(*FactoryNodeDependencyUid)))
+						{
+							FSoftObjectPath PartReferenceObject;
+							if (PartMeshFactoryNode->GetCustomReferenceObject(PartReferenceObject))
+							{
+								if (MeshType* MeshObject = Cast<MeshType>(PartReferenceObject.TryLoad()))
+								{
+									OutPartObjects[Index].Add(MeshObject);
+									bFoundPartObject = true;
+								}
+							}
+						}
+					}
+					else if (FPackageName::IsValidLongPackageName(PartUid) || FPackageName::IsValidObjectPath(PartUid))
+					{
+						FTopLevelAssetPath AssetPath(PartUid);
+						if (AssetPath.IsValid())
+						{
+							if (UObject* MeshObject = StaticLoadAsset(MeshType::StaticClass(), AssetPath))
+							{
+								OutPartObjects[Index].Add(Cast<MeshType>(MeshObject));
+								bFoundPartObject = true;
+							}
+						}
+					}
+
+					if (bFoundPartObject)
+					{
+						bFoundAtLeastOnePartObject = true;
+					}
+					else
+					{
+						MissingPartUids.Add(PartUid);
+					}
+				}
+			}
+
+			if (!MissingPartUids.IsEmpty())
+			{
+				FString Message = FString::Printf(
+					TEXT("Failed to resolve (%d) parts to valid objects while processing nanite assembly description for mesh '%s' : \n")
+					, MissingPartUids.Num()
+					, *MeshFactoryNode->GetDisplayLabel()
+				);
+				for (const FString& PartUid : MissingPartUids)
+				{
+					Message.Append(FString::Printf(TEXT("... %s"), *PartUid));
+				}
+				UE_LOG(LogInterchangeImport, Warning, TEXT("%s"), *Message);
+			}
+
+			return bFoundAtLeastOnePartObject;
+		}
+
+		/** Set the nanite enabled flag to the given value. Note the bool return value here is the previous enabled 
+		  * on/off state (since we may need to revert it), not whether the new setting succeeded/failed. 
+		  */
+		bool SetNaniteEnabled(bool Value)
+		{
+			FMeshNaniteSettings Settings = Mesh->GetNaniteSettings();
+			const bool bOriginalValue = Settings.bEnabled;
+			if (Value != bOriginalValue)
+			{
+				Settings.bEnabled = Value;
+				Mesh->SetNaniteSettings(Settings);
+			}
+			return bOriginalValue;
+		}
+
+		/** Loop over the materials for the given mesh object, and call Func for each MaterialInterface */
+		template <typename TFunc>
+		static void ForEachMeshMaterial(const UObject* MeshObject, TFunc&& Func)
+		{
+			const MeshType* MeshTypeObject = Cast<MeshType>(MeshObject);
+			if (!ensure(MeshTypeObject))
+			{
+				return;
+			}
+
+			int32 LocalMaterialIndex = 0;
+			if constexpr (IsSkeletalMesh())
+			{
+				for (const FSkeletalMaterial& SkeletalMaterial : MeshTypeObject->GetMaterials())
+				{
+					Func(LocalMaterialIndex++, SkeletalMaterial.MaterialInterface);
+				}
+			}
+			else if constexpr (IsStaticMesh())
+			{
+				for (const FStaticMaterial& StaticMaterial : MeshTypeObject->GetStaticMaterials())
+				{
+					Func(LocalMaterialIndex++, StaticMaterial.MaterialInterface);
+				}
+			}
+		}
+	};
+
+	/** Nanite assembly builder public interface */
+
+	FInterchangeNaniteAssemblyBuilder FInterchangeNaniteAssemblyBuilder::Create(USkeletalMesh* SkeletalMesh
+		, const UInterchangeMeshFactoryNode* MeshFactoryNode
+		, const UInterchangeBaseNodeContainer* NodeContainer)
+	{
+		return FInterchangeNaniteAssemblyBuilder(TUniquePtr<ImplBase>(new FInternalNaniteAssemblyBuilder<USkeletalMesh>(SkeletalMesh, MeshFactoryNode, NodeContainer)));
+	}
+
+	FInterchangeNaniteAssemblyBuilder FInterchangeNaniteAssemblyBuilder::Create(UStaticMesh* StaticMesh
+		, const UInterchangeMeshFactoryNode* MeshFactoryNode
+		, const UInterchangeBaseNodeContainer* NodeContainer)
+	{
+		return FInterchangeNaniteAssemblyBuilder(TUniquePtr<ImplBase>(new FInternalNaniteAssemblyBuilder<UStaticMesh>(StaticMesh, MeshFactoryNode, NodeContainer)));
+	}
+
+	void FInterchangeNaniteAssemblyBuilder::AddDescription(const TOptional<UE::Interchange::FNaniteAssemblyDescription>& InDescription)
+	{
+		Impl->AddDescription(InDescription);
+	}
+
+	/** Nanite assembly builder ctor/dtor */
+
+	FInterchangeNaniteAssemblyBuilder::FInterchangeNaniteAssemblyBuilder(TUniquePtr<ImplBase> InImpl)
+		: Impl(MoveTemp(InImpl))
+	{
+
+	}
+
+	FInterchangeNaniteAssemblyBuilder::~FInterchangeNaniteAssemblyBuilder() = default;
+#endif // WITH_EDITOR
 } //ns UE::Interchange::Private::MeshHelper
\ No newline at end of file
diff --git a/Engine/Plugins/Interchange/Runtime/Source/Import/Private/Mesh/InterchangeMeshPayload.cpp b/Engine/Plugins/Interchange/Runtime/Source/Import/Private/Mesh/InterchangeMeshPayload.cpp
new file mode 100644
--- /dev/null
+++ b/Engine/Plugins/Interchange/Runtime/Source/Import/Private/Mesh/InterchangeMeshPayload.cpp
@@ -0,0 +1,51 @@
+// Copyright Epic Games, Inc. All Rights Reserved.
+
+#include "Mesh/InterchangeMeshPayload.h"
+
+#include "Algo/Accumulate.h" 
+
+namespace UE::Interchange
+{
+	bool FNaniteAssemblyDescription::IsValid(FString* OutReason) const
+	{
+		// Transforms and PartIndices must not be empty and should be of equal length. There must also be at least one part. 
+		// Note we don't check for invalid PartIndices indexing into PartUids here - those will be skipped by the assembly builder.
+		if (Transforms.IsEmpty() || PartIndices.IsEmpty() || (Transforms.Num() != PartIndices.Num()))
+		{
+			if (OutReason)
+			{
+				*OutReason = FString::Printf(TEXT("Transforms and part-indices arrays are either empty or different lengths (%d vs %d)"),
+					Transforms.Num(), PartIndices.Num());
+			}
+			return false;
+		}
+
+		if (!BoneInfluences.IsEmpty())
+		{
+			// The stride data should match the number of transforms.
+			if (NumInfluencesPerInstance.Num() != Transforms.Num())
+			{
+				if (OutReason)
+				{
+					*OutReason = FString::Printf(TEXT("Transforms and influences-per-instance arrays are different lengths (%d vs %d)"),
+						Transforms.Num(), NumInfluencesPerInstance.Num());
+				}
+				return false;
+			}
+
+			// Lastly, sum the stride data and make sure it agrees with the number of influences.
+			const int32 NumTotalInflueces = Algo::Accumulate(NumInfluencesPerInstance, 0);
+			if (NumTotalInflueces != BoneInfluences.Num())
+			{
+				if (OutReason)
+				{
+					*OutReason = FString::Printf(TEXT("Bone influences array length (%d) does not match influences-per-instance total (%d)"),
+						BoneInfluences.Num(), NumTotalInflueces);
+				}
+				return false;
+			}
+		}
+
+		return true;
+	}
+} //ns UE::Interchange
\ No newline at end of file
diff --git a/Engine/Plugins/Interchange/Runtime/Source/Import/Private/Mesh/InterchangeSkeletalMeshFactory.cpp b/Engine/Plugins/Interchange/Runtime/Source/Import/Private/Mesh/InterchangeSkeletalMeshFactory.cpp
--- a/Engine/Plugins/Interchange/Runtime/Source/Import/Private/Mesh/InterchangeSkeletalMeshFactory.cpp
+++ b/Engine/Plugins/Interchange/Runtime/Source/Import/Private/Mesh/InterchangeSkeletalMeshFactory.cpp
@@ -9,6 +9,7 @@
 #include "ClothingAsset.h"
 #include "Components.h"
 #include "CoreGlobals.h"
+#include "Engine/NaniteAssemblyData.h"
 #include "Engine/SkeletalMesh.h"
 #include "Engine/SkeletalMeshSocket.h"
 #include "Engine/SkinnedAssetAsyncCompileUtils.h"
@@ -843,6 +844,67 @@ namespace UE
 				}
 				return true;
 			}
+
+			TMap<FString, int32> GetSkeletonBoneIndexByBoneNameTable(const USkeleton& Skeleton)
+			{
+				TMap<FString, int32> BoneIndexByBoneNameTable;
+
+				const FReferenceSkeleton& ReferenceSkeleton = Skeleton.GetReferenceSkeleton();
+				const int32 NumBones = ReferenceSkeleton.GetNum();
+				for (int32 BoneIndex = 0; BoneIndex < NumBones; ++BoneIndex)
+				{
+					// Get bone name by index
+					const FName BoneName = ReferenceSkeleton.GetBoneName(BoneIndex);
+					BoneIndexByBoneNameTable.Add(BoneName.ToString(), BoneIndex);
+				}
+
+				return BoneIndexByBoneNameTable;
+			}
+
+			bool RemapNaniteAssemblyBoneIndices(UE::Interchange::FMeshPayloadData& MeshPayloadData
+				, const TMap<FString, int32>& NewBoneIndexByBoneNameMap
+			    , const FString& AssetName)
+			{
+				using FNaniteAssemblyDescription = UE::Interchange::FNaniteAssemblyDescription;
+
+				if (!MeshPayloadData.NaniteAssemblyDescription.IsSet())
+				{
+					return false;
+				}
+
+				FNaniteAssemblyDescription& Description = MeshPayloadData.NaniteAssemblyDescription.GetValue();
+
+				// Remap the bone influence bone indices to our actual skeleton.  Note that Nanite assembly 
+				// bone indices are uint32, while the indices held in the table are int32.
+
+				const TArray<FString>& JointNames = MeshPayloadData.JointNames;
+				int32 bNumRemapped = 0;
+				for (FNaniteAssemblyDescription::FBoneInfluence& BoneInfluence : Description.BoneInfluences)
+				{
+					const uint32 OldBoneIndex = BoneInfluence.BoneIndex;
+					if (IntFitsIn<uint32, int32>(OldBoneIndex) && JointNames.IsValidIndex(static_cast<int32>(OldBoneIndex)))
+					{
+						const FString& OriginalName = JointNames[OldBoneIndex];
+						if (const int32* RemappedIndexPtr = NewBoneIndexByBoneNameMap.Find(OriginalName))
+						{
+							BoneInfluence.BoneIndex = static_cast<uint32>(*RemappedIndexPtr);
+							bNumRemapped++;
+						}
+					}
+				}
+
+				if (bNumRemapped != Description.BoneInfluences.Num())
+				{
+					UE_LOG(LogInterchangeImport, Warning
+						, TEXT("Failed to remap one or more Nanite Assembly Bone Influence's while processing skeletal mesh '%s'")
+						, *AssetName
+					);
+					return false;
+				}
+
+				return true;
+			}
+
 		} //Namespace Private
 	} //namespace Interchange
 } //namespace UE
@@ -1586,7 +1648,7 @@ UInterchangeFactoryBase::FImportAssetResult UInterchangeSkeletalMeshFactory::Imp
 			MeshDescription = *OriginalMeshDescription;
 		}
 
-		//Store the existing material import data before updating it so we can remap properly the material  and section data
+		//Store the existing material import data before updating it so we can remap properly the material and section data
 		TArray<FName> ExistingOriginalPerSectionMaterialImportName;
 		if (ImportAssetObjectData.bIsReImport)
 		{
@@ -1660,9 +1722,13 @@ UInterchangeFactoryBase::FImportAssetResult UInterchangeSkeletalMeshFactory::Imp
 				ImportAssetObjectLODData.ImportedMaterials.Add(Material);
 			}
 		}
-		
-		PayloadsPerLodIndex.Remove(LodIndex);
-		
+
+		// TODO: reinstate this by checking and retaining only the payloads that have Nanite assembly
+		// data. Also, even if we can't remove the payload entirely, we should be ok to call Empty()  
+		// on the mesh description to free up resources most of the resources.
+		// 
+		// // PayloadsPerLodIndex.Remove(LodIndex);
+
 		CurrentLodIndex++;
 	}
 
@@ -2178,6 +2244,62 @@ void UInterchangeSkeletalMeshFactory::SetupObject_GameThread(const FSetupObjectP
 			}
 
 #if WITH_EDITOR
+			// Build Nanite Assembly Data
+			if (SkeletalMeshFactoryNode && !ImportAssetObjectData.LodDatas.IsEmpty())
+			{
+				if (USkeleton* SkeletonReference = ImportAssetObjectData.SkeletonReference)
+				{
+					auto ForEachLOD0MeshPayload = [this](TFunctionRef<void(UE::Interchange::FMeshPayloadData* MeshPayloadData)> IterationLambda)
+						{
+							constexpr int32 LOD0 = 0;
+							FImportAssetObjectLODData& ImportAssetObjectLODData = ImportAssetObjectData.LodDatas[LOD0];
+							if (UInterchangeSkeletalMeshFactory::FLodPayloads* LodPayloads = PayloadsPerLodIndex.Find(LOD0))
+							{
+								for (const UE::Interchange::FMeshNodeContext& MeshNodeContext : ImportAssetObjectLODData.MeshNodeContexts)
+								{
+									const FInterchangeMeshPayLoadKey Key = MeshNodeContext.GetTranslatorAndTransformPayloadKey();
+									if (TOptional<UE::Interchange::FMeshPayloadData>* MeshPayload = LodPayloads->MeshPayloadPerKey.Find(Key))
+									{
+										if (UE::Interchange::FMeshPayloadData* MeshPayloadData = MeshPayload->GetPtrOrNull())
+										{
+											IterationLambda(MeshPayloadData);
+										}
+									}
+								}	
+							}
+						};
+
+					// Two passes here to avoid Nanite assembly builder construction for the common case of zero nanite assemblies.
+
+					int32 NumActiveNaniteAssemblyPayloads = 0;
+					ForEachLOD0MeshPayload([&NumActiveNaniteAssemblyPayloads](UE::Interchange::FMeshPayloadData* MeshPayloadData)
+						{
+							if (MeshPayloadData->NaniteAssemblyDescription.IsSet())
+							{
+								NumActiveNaniteAssemblyPayloads++;
+							}
+						});
+
+					if (NumActiveNaniteAssemblyPayloads > 0)
+					{
+						using namespace UE::Interchange::Private::MeshHelper;
+
+						const TMap<FString, int32> NewBoneIndexByBoneNameTable = UE::Interchange::Private::GetSkeletonBoneIndexByBoneNameTable(*SkeletonReference);
+						FInterchangeNaniteAssemblyBuilder Builder = FInterchangeNaniteAssemblyBuilder::Create(SkeletalMesh, SkeletalMeshFactoryNode, Arguments.NodeContainer);
+
+						ForEachLOD0MeshPayload([&Builder, &NewBoneIndexByBoneNameTable, &Arguments](UE::Interchange::FMeshPayloadData* MeshPayloadData)
+							{
+								// The bone influence bone indices data in the payload likely needs remapping since translation may have added extra bones and/or 
+								// changed the indices. Note the remapping here may as well be done in-place since the current (invalid) indices are of little use.
+								if (UE::Interchange::Private::RemapNaniteAssemblyBoneIndices(*MeshPayloadData, NewBoneIndexByBoneNameTable, Arguments.NodeUniqueID))
+								{
+									Builder.AddDescription(MeshPayloadData->NaniteAssemblyDescription);
+								}
+							});
+					}
+				}
+			}
+
 			//Re-apply the alternate skinning data
 			if (ImportAssetObjectData.ExistingSkinWeightProfileInfos.Num() > 0)
 			{
@@ -2342,6 +2464,9 @@ void UInterchangeSkeletalMeshFactory::SetupObject_GameThread(const FSetupObjectP
 			}
 		}
 	}
+
+	// Payloads have been processed, so remove
+	PayloadsPerLodIndex.Empty();
 }
 
 void UInterchangeSkeletalMeshFactory::BuildObject_GameThread(const FSetupObjectParams& Arguments, bool& OutPostEditchangeCalled)
diff --git a/Engine/Plugins/Interchange/Runtime/Source/Import/Private/Mesh/InterchangeStaticMeshFactory.cpp b/Engine/Plugins/Interchange/Runtime/Source/Import/Private/Mesh/InterchangeStaticMeshFactory.cpp
--- a/Engine/Plugins/Interchange/Runtime/Source/Import/Private/Mesh/InterchangeStaticMeshFactory.cpp
+++ b/Engine/Plugins/Interchange/Runtime/Source/Import/Private/Mesh/InterchangeStaticMeshFactory.cpp
@@ -566,10 +566,6 @@ UInterchangeFactoryBase::FImportAssetResult UInterchangeStaticMeshFactory::Impor
 			}
 		}
 
-		// Payload has been processed and moved to ImportAssetObjectData.LodMeshDescriptions.
-		// Remove it to make sure it's not reused.
-		PayloadsPerLodIndex.Remove(LodIndex); 
-
 		CurrentLodIndex++;
 	}
 
@@ -1098,8 +1094,56 @@ void UInterchangeStaticMeshFactory::SetupObject_GameThread(const FSetupObjectPar
 		UE::Interchange::FFactoryCommon::FUpdateImportAssetDataParameters UpdateImportAssetDataParameters(StaticMesh, ImportDataPtr, Arguments.SourceData, Arguments.NodeUniqueID, Arguments.NodeContainer, Arguments.OriginalPipelines, Arguments.Translator);
 		ImportDataPtr = UE::Interchange::FFactoryCommon::UpdateImportAssetData(UpdateImportAssetDataParameters);
 		StaticMesh->SetAssetImportData(ImportDataPtr);
+
+#if WITH_EDITOR
+		// Create nanite assembly from all LOD0 payloads 
+		if (const UInterchangeStaticMeshFactoryNode* StaticMeshFactoryNode = Cast<const UInterchangeStaticMeshFactoryNode>(Arguments.NodeContainer->GetFactoryNode(Arguments.NodeUniqueID)))
+		{
+			auto ForEachLOD0MeshPayload = [this](TFunctionRef<void(const UE::Interchange::FMeshPayloadData* MeshPayloadData)> IterationLambda)
+				{
+					constexpr int32 LOD0 = 0;
+					if (const FLodPayloads* LodPayloads = PayloadsPerLodIndex.Find(LOD0))
+					{
+						for (const TPair<FInterchangeMeshPayLoadKey, UE::Interchange::FMeshPayload>& KeyAndPayload : LodPayloads->MeshPayloadPerKey)
+						{
+							const TOptional<UE::Interchange::FMeshPayloadData>& MeshPayload = KeyAndPayload.Value.PayloadData;
+							if (const UE::Interchange::FMeshPayloadData* MeshPayloadData = MeshPayload.GetPtrOrNull())
+							{
+								IterationLambda(MeshPayloadData);
+							}
+						}
+					}
+				};
+
+			// Two passes here to avoid Nanite assembly builder construction for the common case of zero nanite assemblies.
+
+			int32 NumActiveNaniteAssemblyPayloads = 0;
+			ForEachLOD0MeshPayload([&NumActiveNaniteAssemblyPayloads](const UE::Interchange::FMeshPayloadData* const MeshPayloadData)
+				{
+					if (MeshPayloadData->NaniteAssemblyDescription.IsSet())
+					{
+						NumActiveNaniteAssemblyPayloads++;
+					}
+				});
+
+			if (NumActiveNaniteAssemblyPayloads > 0)
+			{
+				using namespace UE::Interchange::Private::MeshHelper;
+
+				FInterchangeNaniteAssemblyBuilder Builder = FInterchangeNaniteAssemblyBuilder::Create(StaticMesh, StaticMeshFactoryNode, Arguments.NodeContainer);
+
+				ForEachLOD0MeshPayload([&Builder](const UE::Interchange::FMeshPayloadData* const MeshPayloadData)
+					{
+						Builder.AddDescription(MeshPayloadData->NaniteAssemblyDescription);
+					});
+			}
+		}
+#endif //WITH_EDITOR
+
+		// Payloads have been processed, so remove
+		PayloadsPerLodIndex.Empty();
 	}
-#endif
+#endif // WITH_EDITORONLY_DATA
 }
 
 void UInterchangeStaticMeshFactory::BuildObject_GameThread(const FSetupObjectParams& Arguments, bool& OutPostEditchangeCalled)
diff --git a/Engine/Plugins/Interchange/Runtime/Source/Import/Public/Mesh/InterchangeMeshHelper.h b/Engine/Plugins/Interchange/Runtime/Source/Import/Public/Mesh/InterchangeMeshHelper.h
--- a/Engine/Plugins/Interchange/Runtime/Source/Import/Public/Mesh/InterchangeMeshHelper.h
+++ b/Engine/Plugins/Interchange/Runtime/Source/Import/Public/Mesh/InterchangeMeshHelper.h
@@ -10,6 +10,7 @@
 class UBodySetup;
 class UInterchangeBaseNodeContainer;
 class UInterchangeSceneNode;
+class UInterchangeMeshFactoryNode;
 class UInterchangeStaticMeshFactoryNode;
 class UInterchangeStaticMeshLodDataNode;
 class UMaterialInterface;
@@ -24,6 +25,7 @@ namespace UE::Interchange
 {
 	struct FMeshPayloadData;
 	struct FMeshPayload;
+	struct FNaniteAssemblyDescription;
 }
 
 namespace UE::Interchange::Private::MeshHelper
@@ -128,4 +130,35 @@ namespace UE::Interchange::Private::MeshHelper
 	INTERCHANGEIMPORT_API void RemapSkeletalMeshVertexColorToMeshDescription(const USkeletalMesh* SkeletalMesh, const int32 LODIndex, FMeshDescription& MeshDescription);
 #endif
 
+#if WITH_EDITOR
+	/** Nanite assembly builder helper for static and skeletal meshes. */
+	class FInterchangeNaniteAssemblyBuilder
+	{
+	public:
+		~FInterchangeNaniteAssemblyBuilder();
+
+		/**
+		 * Create a builder for the given skeletal mesh. The given mesh factory node must hold the mesh part dependency information. 
+		 */
+		static FInterchangeNaniteAssemblyBuilder Create(USkeletalMesh* SkeletalMesh, const UInterchangeMeshFactoryNode* MeshFactoryNode, const UInterchangeBaseNodeContainer* NodeContainer);
+
+		/**
+		 * Create a builder for the given static mesh. The given mesh factory node must hold the mesh part dependency information. 
+		 */
+		static FInterchangeNaniteAssemblyBuilder Create(UStaticMesh* StaticMesh, const UInterchangeMeshFactoryNode* MeshFactoryNode, const UInterchangeBaseNodeContainer* NodeContainer);
+
+		/**
+		 * Add the given FNaniteAssemblyDescription data to the builder. This can be called more than
+		 * once in order to merge multiple descriptions into a single skeletal/static mesh assembly.
+		 */
+		INTERCHANGEIMPORT_API void AddDescription(const TOptional<UE::Interchange::FNaniteAssemblyDescription>& Description);
+
+		class ImplBase;
+
+	private:
+
+		TUniquePtr<ImplBase> Impl;
+		explicit FInterchangeNaniteAssemblyBuilder(TUniquePtr<ImplBase> Impl);
+	};
+#endif
 } //ns UE::Interchange::Private::MeshHelper
diff --git a/Engine/Plugins/Interchange/Runtime/Source/Import/Public/Mesh/InterchangeMeshPayload.h b/Engine/Plugins/Interchange/Runtime/Source/Import/Public/Mesh/InterchangeMeshPayload.h
--- a/Engine/Plugins/Interchange/Runtime/Source/Import/Public/Mesh/InterchangeMeshPayload.h
+++ b/Engine/Plugins/Interchange/Runtime/Source/Import/Public/Mesh/InterchangeMeshPayload.h
@@ -5,10 +5,42 @@
 #include "CoreMinimal.h"
 #include "MeshDescription.h"
 
+struct FNaniteAssemblyBoneInfluence;
+
 namespace UE
 {
 	namespace Interchange
 	{
+		struct FNaniteAssemblyDescription
+		{
+			struct FBoneInfluence
+			{
+				uint32 BoneIndex = INDEX_NONE;
+
+				float BoneWeight = 1.0f;
+			};
+
+			/* Assembly Transforms */
+			TArray<FTransform> Transforms;
+
+			/* Assembly Part Indices, must be the same length as Transforms */
+			TArray<int32> PartIndices;
+
+			/* List of interchange mesh uids, or, UE asset paths per index */
+			TArray<TArray<FString>> PartUids;
+
+			/* List of bone influences. */
+			// For SkeletalMesh assemblies there should one or more influences per transform. 
+			// For StaticMesh assemblies this should be empty.
+			TArray<FBoneInfluence> BoneInfluences;
+
+			/* Number of bone influences per instance */
+			TArray<int32> NumInfluencesPerInstance;
+
+			/* Check that the above arrays have been authored correctly */
+			bool IsValid(FString* OutReason = nullptr) const;
+		};
+
 		struct FMeshPayloadData
 		{
 			/* MESH */
@@ -25,6 +57,10 @@ namespace UE
 			int32 VertexOffset;
 			//The name of the morph target
 			FString MorphTargetName;
+
+			/* NANITE ASSEMBLY */
+			// Optional Nanite Assembly data that should be attached to this skeletal or static mesh.
+			TOptional<FNaniteAssemblyDescription> NaniteAssemblyDescription;
 		};
 
 		struct FMeshPayload
diff --git a/Engine/Plugins/Interchange/Runtime/Source/Nodes/Private/InterchangeMeshNode.cpp b/Engine/Plugins/Interchange/Runtime/Source/Nodes/Private/InterchangeMeshNode.cpp
--- a/Engine/Plugins/Interchange/Runtime/Source/Nodes/Private/InterchangeMeshNode.cpp
+++ b/Engine/Plugins/Interchange/Runtime/Source/Nodes/Private/InterchangeMeshNode.cpp
@@ -63,6 +63,12 @@ namespace UE
 			return Dependencies_BaseKey;
 		}
 
+		const FAttributeKey& FMeshNodeStaticData::GetAssemblyPartDependenciesKey()
+		{
+			static FAttributeKey Dependencies_BaseKey(TEXT("__AssemblyPartDependencies__"));
+			return Dependencies_BaseKey;
+		}
+
 	}//ns Interchange
 }//ns UE
 
@@ -72,6 +78,7 @@ UInterchangeMeshNode::UInterchangeMeshNode()
 	MorphTargetDependencies.Initialize(Attributes, UE::Interchange::FMeshNodeStaticData::GetMorphTargetDependenciesKey().ToString());
 	SceneInstancesUids.Initialize(Attributes, UE::Interchange::FMeshNodeStaticData::GetSceneInstancesUidsKey().ToString());
 	SlotMaterialDependencies.Initialize(Attributes.ToSharedRef(), UE::Interchange::FMeshNodeStaticData::GetSlotMaterialDependenciesKey().ToString());
+	AssemblyPartDependencies.Initialize(Attributes, UE::Interchange::FMeshNodeStaticData::GetAssemblyPartDependenciesKey().ToString());
 }
 
 #if WITH_EDITOR
@@ -538,6 +545,31 @@ bool UInterchangeMeshNode::RemoveSlotMaterialDependencyUid(const FString& SlotNa
 	return SlotMaterialDependencies.RemoveKey(SlotName);
 }
 
+int32 UInterchangeMeshNode::GetAssemblyPartDependenciesCount() const
+{
+	return AssemblyPartDependencies.GetCount();
+}
+
+void UInterchangeMeshNode::GetAssemblyPartDependencies(TArray<FString>& OutDependencies) const
+{
+	AssemblyPartDependencies.GetItems(OutDependencies);
+}
+
+void UInterchangeMeshNode::GetAssemblyPartDependency(const int32 Index, FString& OutDependency) const
+{
+	AssemblyPartDependencies.GetItem(Index, OutDependency);
+}
+
+bool UInterchangeMeshNode::AddAssemblyPartDependencyUid(const FString& DependencyUid)
+{
+	return AssemblyPartDependencies.AddItem(DependencyUid);
+}
+
+bool UInterchangeMeshNode::RemoveAssemblyPartDependencyUid(const FString& DependencyUid)
+{
+	return AssemblyPartDependencies.RemoveItem(DependencyUid);
+}
+
 bool UInterchangeMeshNode::GetCustomCollisionType(EInterchangeMeshCollision& AttributeValue) const
 {
 	IMPLEMENT_NODE_ATTRIBUTE_GETTER(CollisionType, EInterchangeMeshCollision);
diff --git a/Engine/Plugins/Interchange/Runtime/Source/Nodes/Public/InterchangeMeshNode.h b/Engine/Plugins/Interchange/Runtime/Source/Nodes/Public/InterchangeMeshNode.h
--- a/Engine/Plugins/Interchange/Runtime/Source/Nodes/Public/InterchangeMeshNode.h
+++ b/Engine/Plugins/Interchange/Runtime/Source/Nodes/Public/InterchangeMeshNode.h
@@ -27,6 +27,7 @@ namespace UE
 			static UE_API const FAttributeKey& GetMorphTargetDependenciesKey();
 			static UE_API const FAttributeKey& GetSceneInstancesUidsKey();
 			static UE_API const FAttributeKey& GetSlotMaterialDependenciesKey();
+			static UE_API const FAttributeKey& GetAssemblyPartDependenciesKey();
 		};
 
 	}//ns Interchange
@@ -360,6 +361,36 @@ public:
 	UFUNCTION(BlueprintCallable, Category = "Interchange | Node | StaticMesh")
 	UE_API bool RemoveSlotMaterialDependencyUid(const FString& SlotName);
 
+	/**
+    * Retrieve the number of Nanite assembly part dependencies for this object.
+    */
+	UFUNCTION(BlueprintCallable, Category = "Interchange | Node | Mesh")
+	UE_API int32 GetAssemblyPartDependenciesCount() const;
+
+	/**
+	 * Retrieve the Nanite assembly part dependencies for this object.
+	 */
+	UFUNCTION(BlueprintCallable, Category = "Interchange | Node | Mesh")
+	UE_API void GetAssemblyPartDependencies(TArray<FString>& OutDependencies) const;
+
+	/**
+	 * Retrieve the specified Nanite assembly part dependency for this object.
+	 */
+	UFUNCTION(BlueprintCallable, Category = "Interchange | Node | Mesh")
+	UE_API void GetAssemblyPartDependency(const int32 Index, FString& OutDependency) const;
+
+	/**
+	 * Add the specified Nanite assembly part dependency to this object.
+	 */
+	UFUNCTION(BlueprintCallable, Category = "Interchange | Node | Mesh")
+	UE_API bool AddAssemblyPartDependencyUid(const FString& DependencyUid);
+
+	/**
+	 * Remove the specified Nanite assembly part dependency from this object.
+	 */
+	UFUNCTION(BlueprintCallable, Category = "Interchange | Node | Mesh")
+	UE_API bool RemoveAssemblyPartDependencyUid(const FString& DependencyUid);
+
 	/**
 	 * Get the type of collision shapes we should generate from this mesh.
 	 * Note: This is a separate mechanism from the FBX-style collision shape name prefixes. For now, these
@@ -387,6 +418,7 @@ private:
 	UE::Interchange::TArrayAttributeHelper<FString> SkeletonDependencies;
 	UE::Interchange::TArrayAttributeHelper<FString> MaterialDependencies;
 	UE::Interchange::TArrayAttributeHelper<FString> MorphTargetDependencies;
+	UE::Interchange::TArrayAttributeHelper<FString> AssemblyPartDependencies;
 	UE::Interchange::TArrayAttributeHelper<FString> SceneInstancesUids;
 
 	UE::Interchange::TMapAttributeHelper<FString, FString> SlotMaterialDependencies;
diff --git a/Engine/Plugins/Interchange/Runtime/Source/Pipelines/Private/InterchangeGenericMeshPipeline.cpp b/Engine/Plugins/Interchange/Runtime/Source/Pipelines/Private/InterchangeGenericMeshPipeline.cpp
--- a/Engine/Plugins/Interchange/Runtime/Source/Pipelines/Private/InterchangeGenericMeshPipeline.cpp
+++ b/Engine/Plugins/Interchange/Runtime/Source/Pipelines/Private/InterchangeGenericMeshPipeline.cpp
@@ -602,4 +602,141 @@ bool UInterchangeGenericMeshPipeline::IsImpactingAnyMeshesRecursive(const UInter
 		}
 	}
 	return false;
+}
+
+void UInterchangeGenericMeshPipeline::UpdateAssemblyPartDependencyTable(UInterchangeMeshFactoryNode* MeshFactoryNode, const TMap<int32, TArray<FString>>& NodeUidsPerLodIndex)
+{
+	// Early exit for the common case of no assemblies
+	if (!PipelineMeshesUtilities || !PipelineMeshesUtilities->HasAssemblyMeshDependencies() || !ensure(MeshFactoryNode))
+	{
+		return;
+	}
+
+	for (const TPair<int32, TArray<FString>>& LodIndexAndNodeUids : NodeUidsPerLodIndex)
+	{
+		const TArray<FString>& NodeUids = LodIndexAndNodeUids.Value;
+		for (const FString& NodeUid : NodeUids)
+		{
+			const UInterchangeBaseNode* BaseNode = BaseNodeContainer->GetNode(NodeUid);
+			if (const UInterchangeSceneNode* SceneNode = Cast<UInterchangeSceneNode>(BaseNode))
+			{
+				FString MeshDependency;
+				SceneNode->GetCustomAssetInstanceUid(MeshDependency);
+				if (BaseNodeContainer->IsNodeUidValid(MeshDependency))
+				{
+					if (const UInterchangeMeshNode* MeshDependencyNode = Cast<UInterchangeMeshNode>(BaseNodeContainer->GetNode(MeshDependency)))
+					{
+						if (PipelineMeshesUtilities->IsAssemblyMeshUid(MeshDependencyNode->GetUniqueID()))
+						{
+							AssemblyPartMeshUidToFactoryNodeTable.Add(MeshDependencyNode->GetUniqueID(), MeshFactoryNode);
+						}
+					}
+				}
+			}
+			else if (const UInterchangeMeshNode* MeshNode = Cast<UInterchangeMeshNode>(BaseNode))
+			{
+				if (PipelineMeshesUtilities->IsAssemblyMeshUid(MeshNode->GetUniqueID()))
+				{
+					AssemblyPartMeshUidToFactoryNodeTable.Add(MeshNode->GetUniqueID(), MeshFactoryNode);
+				}
+			}
+		}
+	}
+}
+
+void UInterchangeGenericMeshPipeline::CreateAssemblyPartDependencies()
+{
+	if (!PipelineMeshesUtilities)
+	{
+		return;
+	}
+
+	// Table to help identify invalid part dependencies. Nested assemblies - parts that themselves have assemblies that
+	// use more parts and so are both a base and part mesh at the same time are not supported right now.
+	TMap<FString, FString> GlobalPartUidToBaseUidTable;
+
+	for (const TPair<FString, UInterchangeMeshFactoryNode*>& MeshUidAndFactoryNode : AssemblyPartMeshUidToFactoryNodeTable)
+	{
+		UInterchangeMeshFactoryNode* CurrentFactoryNode = MeshUidAndFactoryNode.Value;
+		if (!CurrentFactoryNode)
+		{
+			continue;
+		}
+
+		const FString& CurrentMeshUid = MeshUidAndFactoryNode.Key;
+		if (!PipelineMeshesUtilities->IsValidMeshGeometryUid(CurrentMeshUid))
+		{
+			UE_LOG(LogInterchangePipeline, Warning
+				, TEXT("Failed to find mesh geometry '%s' while creating mesh-to-mesh factory node dependencies.")
+				, *CurrentMeshUid
+			);
+			continue;
+		}
+
+		const FInterchangeMeshGeometry& MeshGeometry = PipelineMeshesUtilities->GetMeshGeometryByUid(CurrentMeshUid);
+		if (!MeshGeometry.MeshNode || MeshGeometry.MeshNode->GetAssemblyPartDependenciesCount() <= 0)
+		{
+			continue;
+		}
+		
+		TArray<FString> AssemblyPartDependencies;
+		MeshGeometry.MeshNode->GetAssemblyPartDependencies(AssemblyPartDependencies);
+
+		// Check that the current mesh hasn't already been used as a part itself
+		if (const FString* BaseUid = GlobalPartUidToBaseUidTable.Find(CurrentMeshUid))
+		{
+			FString Message = FString::Printf(
+				TEXT("Mesh '%s' has already been used as an assembly part by mesh '%s' and therefore must not ")
+				TEXT("specify its own dependencies as well. The (%d) dependencies that will be ignored are : \n")
+				, *CurrentMeshUid
+				, **BaseUid
+				, AssemblyPartDependencies.Num()
+			);
+			for (const FString& AssemblyPartUid : AssemblyPartDependencies)
+			{
+				Message.Append(FString::Printf(TEXT("... %s"), *AssemblyPartUid));
+			}
+
+			UE_LOG(LogInterchangePipeline, Warning, TEXT("%s"), *Message);
+
+			continue;
+		}
+
+		// Assuming we can find the part, set up the required factory node dependency so that 
+		// the part mesh asset is imported and built before the base mesh asset.
+
+		TArray<FString> FactoryDependencies;
+		CurrentFactoryNode->GetFactoryDependencies(FactoryDependencies);
+
+		for (const FString& AssemblyPartUid : AssemblyPartDependencies)
+		{
+			if (AssemblyPartUid == CurrentMeshUid)
+			{
+				continue;
+			}
+
+			if (!AssemblyPartMeshUidToFactoryNodeTable.Contains(AssemblyPartUid))
+			{
+				UE_LOG(LogInterchangePipeline, Warning
+					, TEXT("Failed to create mesh-to-mesh factory node dependency '%s' -> '%s' - target not found.")
+					, *CurrentMeshUid
+					, *AssemblyPartUid
+				);
+				continue;
+			}
+			
+			const UInterchangeMeshFactoryNode* AssemblyPartFactoryNode = AssemblyPartMeshUidToFactoryNodeTable.FindChecked(AssemblyPartUid);
+			if (AssemblyPartFactoryNode)
+			{
+				const FString AssemblyPartFactoryNodeUid = AssemblyPartFactoryNode->GetUniqueID();
+				CurrentFactoryNode->SetAssemblyPartDependencyUid(AssemblyPartUid, AssemblyPartFactoryNodeUid);
+				// Create a factory dependency so the part asset is imported before the mesh asset
+				if (!FactoryDependencies.Contains(AssemblyPartFactoryNodeUid))
+				{
+					CurrentFactoryNode->AddFactoryDependencyUid(AssemblyPartFactoryNodeUid);
+				}
+				GlobalPartUidToBaseUidTable.Add(AssemblyPartUid, CurrentMeshUid);
+			}
+		}
+	}
 }
\ No newline at end of file
diff --git a/Engine/Plugins/Interchange/Runtime/Source/Pipelines/Private/InterchangeGenericSkeletalMeshPipeline.cpp b/Engine/Plugins/Interchange/Runtime/Source/Pipelines/Private/InterchangeGenericSkeletalMeshPipeline.cpp
--- a/Engine/Plugins/Interchange/Runtime/Source/Pipelines/Private/InterchangeGenericSkeletalMeshPipeline.cpp
+++ b/Engine/Plugins/Interchange/Runtime/Source/Pipelines/Private/InterchangeGenericSkeletalMeshPipeline.cpp
@@ -163,6 +163,7 @@ void UInterchangeGenericMeshPipeline::ExecutePreImportPipelineSkeletalMesh()
 					SetSkeletalMeshDependencies(SkeletonRootUid, SkeletalMeshFactoryNode);
 					SkeletonFactoryNodes.Add(SkeletonFactoryNode);
 					SkeletalMeshFactoryNodes.Add(SkeletalMeshFactoryNode);
+					UpdateAssemblyPartDependencyTable(SkeletalMeshFactoryNode, MeshUidsPerLodIndex);
 					bFoundInstances = true;
 				}
 			}
@@ -312,6 +313,8 @@ void UInterchangeGenericMeshPipeline::ExecutePreImportPipelineSkeletalMesh()
 		bool bUseMeshInstance = true;
 		CreatePerSkeletonRootUidSkinnedMesh(bUseMeshInstance);
 	}
+
+	CreateAssemblyPartDependencies();
 }
 
 UInterchangeSkeletalMeshFactoryNode* UInterchangeGenericMeshPipeline::CreateSkeletalMeshFactoryNode(const FString& RootJointUid, const TMap<int32, TArray<FString>>& MeshUidsPerLodIndex)
diff --git a/Engine/Plugins/Interchange/Runtime/Source/Pipelines/Private/InterchangeGenericStaticMeshPipeline.cpp b/Engine/Plugins/Interchange/Runtime/Source/Pipelines/Private/InterchangeGenericStaticMeshPipeline.cpp
--- a/Engine/Plugins/Interchange/Runtime/Source/Pipelines/Private/InterchangeGenericStaticMeshPipeline.cpp
+++ b/Engine/Plugins/Interchange/Runtime/Source/Pipelines/Private/InterchangeGenericStaticMeshPipeline.cpp
@@ -427,6 +427,7 @@ void UInterchangeGenericMeshPipeline::ExecutePreImportPipelineStaticMesh()
 							if (UInterchangeStaticMeshFactoryNode* StaticMeshFactoryNode = CreateStaticMeshFactoryNode(MeshUidsPerLodIndex, MeshUid))
 							{
 								StaticMeshFactoryNodes.Add(StaticMeshFactoryNode);
+								UpdateAssemblyPartDependencyTable(StaticMeshFactoryNode, MeshUidsPerLodIndex);
 								bFoundMeshes = true;
 							}
 						}
@@ -444,6 +445,8 @@ void UInterchangeGenericMeshPipeline::ExecutePreImportPipelineStaticMesh()
 			}
 		}
 	}
+
+	CreateAssemblyPartDependencies();
 }
 
 namespace UE::InterchangeStaticMeshFactory::Private
diff --git a/Engine/Plugins/Interchange/Runtime/Source/Pipelines/Private/InterchangePipelineMeshesUtilities.cpp b/Engine/Plugins/Interchange/Runtime/Source/Pipelines/Private/InterchangePipelineMeshesUtilities.cpp
--- a/Engine/Plugins/Interchange/Runtime/Source/Pipelines/Private/InterchangePipelineMeshesUtilities.cpp
+++ b/Engine/Plugins/Interchange/Runtime/Source/Pipelines/Private/InterchangePipelineMeshesUtilities.cpp
@@ -322,6 +322,14 @@ UInterchangePipelineMeshesUtilities* UInterchangePipelineMeshesUtilities::Create
 					FInterchangeMeshGeometry& MeshGeometry = PipelineMeshesUtilities->MeshGeometriesPerMeshUid.FindOrAdd(NodeUid);
 					MeshGeometry.MeshUid = NodeUid;
 					MeshGeometry.MeshNode = MeshNode;
+
+					if (MeshNode->GetAssemblyPartDependenciesCount() > 0)
+					{
+						PipelineMeshesUtilities->AssemblyMeshUids.Add(NodeUid);
+						TArray<FString> AssemblyPartDependencies;
+						MeshNode->GetAssemblyPartDependencies(AssemblyPartDependencies);
+						PipelineMeshesUtilities->AssemblyMeshUids.Append(AssemblyPartDependencies);
+					}
 				}
 			}
 		}
@@ -1061,3 +1069,13 @@ FString UInterchangePipelineMeshesUtilities::GetMeshGeometrySkeletonRootUid(cons
 	}
 	return SkeletonRootUid;
 }
+
+bool UInterchangePipelineMeshesUtilities::HasAssemblyMeshDependencies() const
+{
+	return !AssemblyMeshUids.IsEmpty();
+}
+
+bool UInterchangePipelineMeshesUtilities::IsAssemblyMeshUid(const FString& MeshUid) const
+{
+	return AssemblyMeshUids.Contains(MeshUid);
+}
diff --git a/Engine/Plugins/Interchange/Runtime/Source/Pipelines/Public/InterchangeGenericMeshPipeline.h b/Engine/Plugins/Interchange/Runtime/Source/Pipelines/Public/InterchangeGenericMeshPipeline.h
--- a/Engine/Plugins/Interchange/Runtime/Source/Pipelines/Public/InterchangeGenericMeshPipeline.h
+++ b/Engine/Plugins/Interchange/Runtime/Source/Pipelines/Public/InterchangeGenericMeshPipeline.h
@@ -352,11 +352,25 @@ private:
 	/* Meshes utilities, to parse the translated graph and extract the meshes informations. */
 	TObjectPtr<UInterchangePipelineMeshesUtilities> PipelineMeshesUtilities = nullptr;
 
+	/** Table to access the corresponding factory node for the given mesh uid*/
+	TMap<FString, UInterchangeMeshFactoryNode*> AssemblyPartMeshUidToFactoryNodeTable;
+
 	static UE_API bool IsImpactingAnyMeshesRecursive(const UInterchangeSceneNode* SceneNode
 		, const UInterchangeBaseNodeContainer* InBaseNodeContainer
 		, const TArray<FString>& StaticMeshNodeUids
 		, TMap<const UInterchangeSceneNode*, bool>& CacheProcessSceneNodes);
 
+	/**
+	* Update the MeshUid to MeshFactoryNode table with the given node uids
+	* @param NodeUidsPerLodIndex - The NodeUids can be a UInterchangeSceneNode or a UInterchangeMeshNode.
+	*/
+	UE_API void UpdateAssemblyPartDependencyTable(UInterchangeMeshFactoryNode* MeshFactoryNode, const TMap<int32, TArray<FString>>& NodeUidsPerLodIndex);
+
+	/**
+	* Creates the final dependencies between nanite assembly parts
+	*/
+	UE_API void CreateAssemblyPartDependencies();
+
 	/************************************************************************/
 	/* Skeletal mesh API BEGIN                                              */
 
diff --git a/Engine/Plugins/Interchange/Runtime/Source/Pipelines/Public/InterchangePipelineMeshesUtilities.h b/Engine/Plugins/Interchange/Runtime/Source/Pipelines/Public/InterchangePipelineMeshesUtilities.h
--- a/Engine/Plugins/Interchange/Runtime/Source/Pipelines/Public/InterchangePipelineMeshesUtilities.h
+++ b/Engine/Plugins/Interchange/Runtime/Source/Pipelines/Public/InterchangePipelineMeshesUtilities.h
@@ -358,6 +358,17 @@ public:
 
 	UE_API FString GetMeshGeometrySkeletonRootUid(const FInterchangeMeshGeometry& MeshGeometry) const;
 
+	/**
+    * Returns true if there are nanite assembly mesh-to-mesh dependencies specified by the input meshes.
+    */
+	UE_API bool HasAssemblyMeshDependencies() const;
+
+	/**
+	* Returns true if the given mesh uid is the base or part mesh of a nanite assembly.
+	*/
+	UE_API bool IsAssemblyMeshUid(const FString& MeshUid) const;
+
+
 	UFUNCTION(BlueprintCallable, Category = "Interchange | Pipeline | Meshes")
 	void SetContext(const FInterchangePipelineMeshesUtilitiesContext& Context) const
 	{
@@ -368,6 +379,7 @@ protected:
 	TMap<FString, FInterchangeMeshGeometry> MeshGeometriesPerMeshUid;
 	TMap<FString, FInterchangeMeshInstance> MeshInstancesPerMeshInstanceUid;
 	TMap<FString, FString> SkeletonRootUidPerMeshUid;
+	TSet<FString> AssemblyMeshUids;
 
 	UInterchangeBaseNodeContainer* BaseNodeContainer;
 
