Interchange:
- Fix a crash when we try to load an asset outside the Game Thread

diff --git a/Engine/Plugins/Interchange/Runtime/Source/Pipelines/Private/InterchangeGenericMaterialPipeline.cpp b/Engine/Plugins/Interchange/Runtime/Source/Pipelines/Private/InterchangeGenericMaterialPipeline.cpp
--- a/Engine/Plugins/Interchange/Runtime/Source/Pipelines/Private/InterchangeGenericMaterialPipeline.cpp
+++ b/Engine/Plugins/Interchange/Runtime/Source/Pipelines/Private/InterchangeGenericMaterialPipeline.cpp
@@ -184,6 +184,9 @@ namespace UE::Interchange::InterchangeGenericMaterialPipeline::Private
 		const FName MaterialFunctionMemberName = GET_MEMBER_NAME_CHECKED(UMaterialExpressionMaterialFunctionCall, MaterialFunction);
 		UClass* CustomExpressionClass = UMaterialExpressionMaterialFunctionCall::StaticClass();
 
+		// load the material function in the Game Thread to avoid any potential crash in the Material Factory
+		FSoftObjectPath{ MaterialFunctionPath }.TryLoad();
+
 		FunctionCallExpression.SetCustomExpressionClassName(CustomExpressionClass->GetName());
 
 		FunctionCallExpression.AddStringAttribute(MaterialFunctionMemberName.ToString(), MaterialFunctionPath);
