MaterialX:
- Enable volumetric FOG-Cloud material by default

diff --git a/Engine/Plugins/Interchange/Runtime/Source/Pipelines/Public/InterchangeMaterialXPipeline.h b/Engine/Plugins/Interchange/Runtime/Source/Pipelines/Public/InterchangeMaterialXPipeline.h
--- a/Engine/Plugins/Interchange/Runtime/Source/Pipelines/Public/InterchangeMaterialXPipeline.h
+++ b/Engine/Plugins/Interchange/Runtime/Source/Pipelines/Public/InterchangeMaterialXPipeline.h
@@ -104,7 +104,7 @@ public:
 
 	/** If true, handle Volume Distribution Functions as volumetric materials, i.e an internal Substrate Volumetric FOG-Cloud BSDF is created, otherwise a simple volume Slab BSDF is used. */
 	UPROPERTY(EditAnywhere, config, Category = "Volume")
-	bool bVolumetricMaterial = false;
+	bool bVolumetricMaterial = true;
 
 protected:
 	UE_API virtual void AdjustSettingsForContext(const FInterchangePipelineContextParams& ContextParams) override;
