[Interchange][FBX] UFBX failing FBX asset import tests (missing features)

- user defined attributes
- pivot transform

#FailedValidations Preflight

diff --git a/Engine/Plugins/Interchange/Runtime/Source/Parsers/Fbx/Private/Ufbx/UfbxParser.cpp b/Engine/Plugins/Interchange/Runtime/Source/Parsers/Fbx/Private/Ufbx/UfbxParser.cpp
--- a/Engine/Plugins/Interchange/Runtime/Source/Parsers/Fbx/Private/Ufbx/UfbxParser.cpp
+++ b/Engine/Plugins/Interchange/Runtime/Source/Parsers/Fbx/Private/Ufbx/UfbxParser.cpp
@@ -24,6 +24,7 @@
 #include "Serialization/MemoryWriter.h"
 
 #include "ufbx.c"
+#include "Nodes/InterchangeUserDefinedAttribute.h"
 
 
 #define LOCTEXT_NAMESPACE "InterchangeFbxParser"
@@ -82,11 +83,99 @@ FName FUfbxParser::GetMaterialName(const ufbx_material& Material) const
 	return *GetElementName(Material.element);
 }
 
+FString FUfbxParser::GetMaterialUid(const FString& MaterialName) const
+{
+	return TEXT("\\Material\\") + MaterialName;
+}
+
 FString FUfbxParser::GetMeshUid(const ufbx_element& Mesh) const
 {
 	return FString::FromInt(Mesh.element_id);
 }
 
+void FUfbxParser::ConvertProperty(const ufbx_prop& Prop, UInterchangeBaseNode* SceneNode)
+{
+	if (Prop.flags & UFBX_PROP_FLAG_USER_DEFINED)
+	{
+		FString PropertyName = Convert::ToUnrealString(Prop.name);
+		TOptional<FString> PayloadKey;
+		// #ufbx_todo: animated prop; PayloadKey will be added then  
+
+		switch (Prop.type)
+		{
+		case UFBX_PROP_UNKNOWN:
+			break;
+		case UFBX_PROP_BOOLEAN:
+			{
+				const bool PropertyValue = Prop.value_int != 0;
+				UInterchangeUserDefinedAttributesAPI::CreateUserDefinedAttribute(SceneNode, PropertyName, PropertyValue, PayloadKey);
+			}
+			break;
+		case UFBX_PROP_INTEGER:
+			{
+				const int32 PropertyValue = Prop.value_int;
+				UInterchangeUserDefinedAttributesAPI::CreateUserDefinedAttribute(SceneNode, PropertyName, PropertyValue, PayloadKey);
+			}
+			break;
+		case UFBX_PROP_NUMBER:
+			{
+				const float PropertyValue = Prop.value_real;
+				UInterchangeUserDefinedAttributesAPI::CreateUserDefinedAttribute(SceneNode, PropertyName, PropertyValue, PayloadKey);
+			}
+			break;
+		case UFBX_PROP_VECTOR:
+			{
+				const FVector4d PropertyValue = Convert::ConvertVec4(Prop.value_vec4);
+				UInterchangeUserDefinedAttributesAPI::CreateUserDefinedAttribute(SceneNode, PropertyName, PropertyValue, PayloadKey);
+			}
+			break;
+		case UFBX_PROP_COLOR:
+			{
+				const FVector PropertyValue = FVector(Convert::ConvertColor(Prop.value_vec4));
+				UInterchangeUserDefinedAttributesAPI::CreateUserDefinedAttribute(SceneNode, PropertyName, PropertyValue, PayloadKey);
+			}
+			break;
+		case UFBX_PROP_COLOR_WITH_ALPHA:
+			{
+				const FLinearColor PropertyValue = Convert::ConvertColor(Prop.value_vec4);
+				UInterchangeUserDefinedAttributesAPI::CreateUserDefinedAttribute(SceneNode, PropertyName, PropertyValue, PayloadKey);
+			}
+			break;
+		case UFBX_PROP_STRING:
+			{
+				const FString PropertyValue = Convert::ToUnrealString(Prop.value_str);
+				UInterchangeUserDefinedAttributesAPI::CreateUserDefinedAttribute(SceneNode, PropertyName, PropertyValue, PayloadKey);
+			}
+			break;
+		case UFBX_PROP_DATE_TIME:
+			{
+				const FDateTime PropertyValue(Prop.value_int);
+				UInterchangeUserDefinedAttributesAPI::CreateUserDefinedAttribute(SceneNode, PropertyName, PropertyValue, PayloadKey);
+			}
+			break;
+		// #ufbx_todo: support other useful custom attributes among these types
+		case UFBX_PROP_TRANSLATION:
+			break;
+		case UFBX_PROP_ROTATION:
+			break;
+		case UFBX_PROP_SCALING:
+			break;
+		case UFBX_PROP_DISTANCE:
+			break;
+		case UFBX_PROP_COMPOUND:
+			break;
+		case UFBX_PROP_BLOB:
+			break;
+		case UFBX_PROP_REFERENCE:
+			break;
+		case UFBX_PROP_TYPE_FORCE_32BIT:
+			break;
+		}
+	}
+}
+
+
+
 void FUfbxParser::SetResultContainer(UInterchangeResultsContainer* Result)
 {
 	ResultsContainer = Result;
@@ -334,10 +423,6 @@ bool FUfbxParser::FetchTexturePayload(const FString& PayloadKey, TOptional<TArra
 
 		AddMessage<UInterchangeResultError_Generic>()->Text = LOCTEXT("CannotRetrieveTexturePayloadNonTextureKey", "Cannot retrieve payload; FUfbxParser::FetchTexturePayload was called with non-texture payload key.");
 	}
-	else
-	{
-		AddMessage<UInterchangeResultError_Generic>()->Text = LOCTEXT("CannotRetrieveTexturePayloadNoContext", "Cannot retrieve payload; FUfbxParser::FetchTexturePayload - payload key doesn't have any context.");
-	}
 
 	return false;
 }
diff --git a/Engine/Plugins/Interchange/Runtime/Source/Parsers/Fbx/Private/Ufbx/UfbxParser.h b/Engine/Plugins/Interchange/Runtime/Source/Parsers/Fbx/Private/Ufbx/UfbxParser.h
--- a/Engine/Plugins/Interchange/Runtime/Source/Parsers/Fbx/Private/Ufbx/UfbxParser.h
+++ b/Engine/Plugins/Interchange/Runtime/Source/Parsers/Fbx/Private/Ufbx/UfbxParser.h
@@ -140,6 +140,9 @@ namespace UE::Interchange::Private
 		FString GetNodeUid(const ufbx_node& Node) const;
 		FString GetMeshUid(const ufbx_element& Node) const;
 		FName GetMaterialName(const ufbx_material& Material) const;
+		FString GetMaterialUid(const FString& Name) const;
+
+		static void ConvertProperty(const ufbx_prop& Prop, UInterchangeBaseNode* SceneNode);
 
 		ufbx_scene* Scene = nullptr;
 
diff --git a/Engine/Plugins/Interchange/Runtime/Source/Parsers/Fbx/Private/Ufbx/UfbxScene.cpp b/Engine/Plugins/Interchange/Runtime/Source/Parsers/Fbx/Private/Ufbx/UfbxScene.cpp
--- a/Engine/Plugins/Interchange/Runtime/Source/Parsers/Fbx/Private/Ufbx/UfbxScene.cpp
+++ b/Engine/Plugins/Interchange/Runtime/Source/Parsers/Fbx/Private/Ufbx/UfbxScene.cpp
@@ -94,6 +94,12 @@ namespace UE::Interchange::Private
 		// #ufbx_todo: Node->geometry_transform
 		// #ufbx_todo: possible to set SetCustomGeometricTransform, SetCustomPivotNodeTransform?
 
+		FTransform GeometryTransform = Convert::ConvertTransform(InNode.geometry_transform);
+		if (!GeometryTransform.Equals(FTransform::Identity))
+		{
+			SceneNode->SetCustomGeometricTransform(GeometryTransform);
+		}
+
 		if (InNode.mesh)
 		{
 			FString MeshUniqueID = Parser.GetMeshUid(InNode.mesh->element);
@@ -148,83 +154,7 @@ namespace UE::Interchange::Private
 
 		for (const ufbx_prop& Prop : InNode.props.props)
 		{
-			if (Prop.flags & UFBX_PROP_FLAG_USER_DEFINED)
-			{
-				FString PropertyName = Convert::ToUnrealString(Prop.name);
-				TOptional<FString> PayloadKey;
-				// #ufbx_todo: animated prop; PayloadKey will be added then  
-
-				switch (Prop.type)
-				{
-				case UFBX_PROP_UNKNOWN:
-					break;
-				case UFBX_PROP_BOOLEAN:
-					{
-						const bool PropertyValue = Prop.value_int != 0;
-						UInterchangeUserDefinedAttributesAPI::CreateUserDefinedAttribute(SceneNode, PropertyName, PropertyValue, PayloadKey);
-					}
-					break;
-				case UFBX_PROP_INTEGER:
-					{
-						const int32 PropertyValue = Prop.value_int;
-						UInterchangeUserDefinedAttributesAPI::CreateUserDefinedAttribute(SceneNode, PropertyName, PropertyValue, PayloadKey);
-					}
-					break;
-				case UFBX_PROP_NUMBER:
-					{
-						const float PropertyValue = Prop.value_real;
-						UInterchangeUserDefinedAttributesAPI::CreateUserDefinedAttribute(SceneNode, PropertyName, PropertyValue, PayloadKey);
-					}
-					break;
-				case UFBX_PROP_VECTOR:
-					{
-						const FVector4d PropertyValue = Convert::ConvertVec4(Prop.value_vec4);
-						UInterchangeUserDefinedAttributesAPI::CreateUserDefinedAttribute(SceneNode, PropertyName, PropertyValue, PayloadKey);
-					}
-					break;
-				case UFBX_PROP_COLOR:
-					{
-						const FVector PropertyValue = FVector(Convert::ConvertColor(Prop.value_vec4));
-						UInterchangeUserDefinedAttributesAPI::CreateUserDefinedAttribute(SceneNode, PropertyName, PropertyValue, PayloadKey);
-					}
-					break;
-				case UFBX_PROP_COLOR_WITH_ALPHA:
-					{
-						const FLinearColor PropertyValue = Convert::ConvertColor(Prop.value_vec4);
-						UInterchangeUserDefinedAttributesAPI::CreateUserDefinedAttribute(SceneNode, PropertyName, PropertyValue, PayloadKey);
-					}
-					break;
-				case UFBX_PROP_STRING:
-					{
-						const FString PropertyValue = Convert::ToUnrealString(Prop.value_str);
-						UInterchangeUserDefinedAttributesAPI::CreateUserDefinedAttribute(SceneNode, PropertyName, PropertyValue, PayloadKey);
-					}
-					break;
-				case UFBX_PROP_DATE_TIME:
-					{
-						const FDateTime PropertyValue(Prop.value_int);
-						UInterchangeUserDefinedAttributesAPI::CreateUserDefinedAttribute(SceneNode, PropertyName, PropertyValue, PayloadKey);
-					}
-					break;
-				// #ufbx_todo: support other useful custom attributes among these types
-				case UFBX_PROP_TRANSLATION:
-					break;
-				case UFBX_PROP_ROTATION:
-					break;
-				case UFBX_PROP_SCALING:
-					break;
-				case UFBX_PROP_DISTANCE:
-					break;
-				case UFBX_PROP_COMPOUND:
-					break;
-				case UFBX_PROP_BLOB:
-					break;
-				case UFBX_PROP_REFERENCE:
-					break;
-				case UFBX_PROP_TYPE_FORCE_32BIT:
-					break;
-				}
-			}
+			Parser.ConvertProperty(Prop, SceneNode);
 		}
 
 	}
