Interchange: Workaround for HeterogeneousVolumeComponent not refreshing after SVT reimport.

diff --git a/Engine/Plugins/Interchange/Runtime/Source/Import/Private/Volume/InterchangeSparseVolumeTextureFactory.cpp b/Engine/Plugins/Interchange/Runtime/Source/Import/Private/Volume/InterchangeSparseVolumeTextureFactory.cpp
--- a/Engine/Plugins/Interchange/Runtime/Source/Import/Private/Volume/InterchangeSparseVolumeTextureFactory.cpp
+++ b/Engine/Plugins/Interchange/Runtime/Source/Import/Private/Volume/InterchangeSparseVolumeTextureFactory.cpp
@@ -8,6 +8,7 @@
 #include "InterchangeSparseVolumeTextureFactoryNode.h"
 #include "InterchangeTranslatorBase.h"
 #include "InterchangeVolumeNode.h"
+#include "Materials/MaterialInterface.h"
 #include "Nodes/InterchangeBaseNode.h"
 #include "Volume/InterchangeVolumeDefinitions.h"
 #include "Volume/InterchangeVolumePayloadData.h"
@@ -15,6 +16,7 @@
 
 #include "Async/ParallelFor.h"
 #include "Async/TaskGraphInterfaces.h"
+#include "Components/HeterogeneousVolumeComponent.h"
 #include "Misc/CoreStats.h"
 #include "Misc/Paths.h"
 #include "ProfilingDebugging/CpuProfilerTrace.h"
@@ -22,6 +24,7 @@
 #include "SparseVolumeTexture/SparseVolumeTexture.h"
 #include "UDIMUtilities.h"
 #include "UObject/ObjectMacros.h"
+#include "UObject/UObjectIterator.h"
 #if WITH_EDITORONLY_DATA
 #include "EditorFramework/AssetImportData.h"
 #endif	  // WITH_EDITORONLY_DATA
@@ -1041,6 +1044,24 @@ void UInterchangeSparseVolumeTextureFactory::SetupObject_GameThread(const FSetup
 		Streamable->AssetImportData = UE::Interchange::FFactoryCommon::SetImportAssetData(SetImportAssetDataParameters);
 	}
 #endif	  // WITH_EDITOR
+
+	// Temp workaround for UE-255278: It seems that existing HVComponents that are using our SVT will just stop rendering the volume
+	// after we reimport the SVT here, until we trigger a refresh by e.g. reseating the material like we do here, or by opening the material editor.
+	// This seems to be something related to HVComponents and not SVTs directly, as the SVTViewer component doesn't show this issue, but we have 
+	// not yet narrowed down what the actual root cause is
+	if (Arguments.bIsReimport)
+	{
+		for (TObjectIterator<UHeterogeneousVolumeComponent> It; It; ++It)
+		{
+			// Ideally we'd check whether this HVComponent is using this texture that was reimported, but there's no
+			// public interface for that and honestly it may take even longer to check than to just refresh them all anyway
+			if (UMaterialInterface* Material = It->GetMaterial(0))
+			{
+				It->SetMaterial(0, Material);
+				It->MarkRenderStateDirty();
+			}
+		}
+	}
 }
 
 bool UInterchangeSparseVolumeTextureFactory::GetSourceFilenames(const UObject* Object, TArray<FString>& OutSourceFilenames) const
