[Inerchange] Fix missing cache rebuild call during UInterchangeMeshFactoryNode deserialization for the recently added AssemblyPartDependencies attribute map

diff --git a/Engine/Plugins/Interchange/Runtime/Source/FactoryNodes/Public/InterchangeMeshFactoryNode.h b/Engine/Plugins/Interchange/Runtime/Source/FactoryNodes/Public/InterchangeMeshFactoryNode.h
--- a/Engine/Plugins/Interchange/Runtime/Source/FactoryNodes/Public/InterchangeMeshFactoryNode.h
+++ b/Engine/Plugins/Interchange/Runtime/Source/FactoryNodes/Public/InterchangeMeshFactoryNode.h
@@ -39,6 +39,7 @@ public:
 		if (Ar.IsLoading() && bIsInitialized)
 		{
 			SlotMaterialDependencies.RebuildCache();
+			AssemblyPartDependencies.RebuildCache();
 #if WITH_ENGINE
 			// Make sure the class is properly set when we compile with engine, this will set the bIsNodeClassInitialized to true.
 			SetNodeClassFromClassAttribute();
