[Interchange] Updated comment for UInterchangeFactoryBaseNode::Serialize

diff --git a/Engine/Plugins/Interchange/Runtime/Source/FactoryNodes/Public/InterchangeMeshFactoryNode.h b/Engine/Plugins/Interchange/Runtime/Source/FactoryNodes/Public/InterchangeMeshFactoryNode.h
--- a/Engine/Plugins/Interchange/Runtime/Source/FactoryNodes/Public/InterchangeMeshFactoryNode.h
+++ b/Engine/Plugins/Interchange/Runtime/Source/FactoryNodes/Public/InterchangeMeshFactoryNode.h
@@ -30,7 +30,7 @@ public:
 	UE_API UInterchangeMeshFactoryNode();
 
 	/**
-	 * Override Serialize() to restore SlotMaterialDependencies on load.
+	 * Override Serialize() to restore SlotMaterialDependencies and AssemblyPartDependencies on load.
 	 */
 	virtual void Serialize(FArchive& Ar) override
 	{
