Fixed issue with reimport workflow on a level import when assets or actors have been deleted or renamed
- Root cause: 'Force reimport deleted assets'and  'Force reimport deleted actors' options were  not honored anymore
- Solution: Skip factory node if object (asset or actor) to reimport if object not found
Fixed issue with reimport failing when user has chosen to create a sub-folder named after the imported file's name
- Soluion: In this context, cache the name of the imported file for reuse on reimport.
Fixed issue with material instance not being updated on level reimport
- Solution: Do not skip a  material instance factory node on level reimport if it is its actual reimport

#rnx

diff --git a/Engine/Plugins/Interchange/Runtime/Source/Import/Private/Animation/InterchangeAnimSequenceFactory.cpp b/Engine/Plugins/Interchange/Runtime/Source/Import/Private/Animation/InterchangeAnimSequenceFactory.cpp
--- a/Engine/Plugins/Interchange/Runtime/Source/Import/Private/Animation/InterchangeAnimSequenceFactory.cpp
+++ b/Engine/Plugins/Interchange/Runtime/Source/Import/Private/Animation/InterchangeAnimSequenceFactory.cpp
@@ -1167,6 +1167,10 @@ UInterchangeFactoryBase::FImportAssetResult UInterchangeAnimSequenceFactory::Beg
 			Skeleton = Cast<USkeleton>(SkeletonObject);
 
 		}
+		else if (NewAnimSequence)
+		{
+			Skeleton = NewAnimSequence->GetSkeleton();
+		}
 
 		if (!ensure(Skeleton))
 		{
@@ -1440,6 +1444,10 @@ UInterchangeFactoryBase::FImportAssetResult UInterchangeAnimSequenceFactory::End
 			Skeleton = Cast<USkeleton>(SkeletonObject);
 
 		}
+		else if (AnimSequence)
+		{
+			Skeleton = AnimSequence->GetSkeleton();
+		}
 
 		if (!ensure(Skeleton))
 		{
diff --git a/Engine/Plugins/Interchange/Runtime/Source/Pipelines/Private/InterchangeGenericAssetsPipeline.cpp b/Engine/Plugins/Interchange/Runtime/Source/Pipelines/Private/InterchangeGenericAssetsPipeline.cpp
--- a/Engine/Plugins/Interchange/Runtime/Source/Pipelines/Private/InterchangeGenericAssetsPipeline.cpp
+++ b/Engine/Plugins/Interchange/Runtime/Source/Pipelines/Private/InterchangeGenericAssetsPipeline.cpp
@@ -1326,12 +1326,17 @@ void UInterchangeGenericAssetsPipeline::ExecutePipeline(UInterchangeBaseNodeCont
 	}
 	if (bSceneNameSubFolder)
 	{
-		FString SceneNameFolderPrefix = FPaths::GetBaseFilename(InSourceDatas[0]->GetFilename());
+		// Update SceneNameFolderPrefix only if it is not a reimport otherwise reimport strategy will be impacted
+		// when reimporting with new file which usually has a different name
+		if (CacheContextParam.ContextType != EInterchangePipelineContext::AssetReimport && CacheContextParam.ContextType != EInterchangePipelineContext::SceneReimport)
+		{
+			SceneNameFolderPrefix = FPaths::GetBaseFilename(InSourceDatas[0]->GetFilename());
 
-		UInterchangeManager::GetInterchangeManager().SanitizeNameInline(
-			SceneNameFolderPrefix,
-			ESanitizeNameTypeFlags::ObjectName | ESanitizeNameTypeFlags::ObjectPath | ESanitizeNameTypeFlags::LongPackage
-		);
+			UInterchangeManager::GetInterchangeManager().SanitizeNameInline(
+				SceneNameFolderPrefix,
+				ESanitizeNameTypeFlags::ObjectName | ESanitizeNameTypeFlags::ObjectPath | ESanitizeNameTypeFlags::LongPackage
+			);
+		}
 
 		SourceNode->SetCustomSubPathPrefix(SceneNameFolderPrefix);
 	}
diff --git a/Engine/Plugins/Interchange/Runtime/Source/Pipelines/Private/InterchangeGenericMaterialPipeline.cpp b/Engine/Plugins/Interchange/Runtime/Source/Pipelines/Private/InterchangeGenericMaterialPipeline.cpp
--- a/Engine/Plugins/Interchange/Runtime/Source/Pipelines/Private/InterchangeGenericMaterialPipeline.cpp
+++ b/Engine/Plugins/Interchange/Runtime/Source/Pipelines/Private/InterchangeGenericMaterialPipeline.cpp
@@ -4,6 +4,7 @@
 
 #include "AssetRegistry/AssetData.h"
 #include "AssetRegistry/AssetRegistryModule.h"
+#include "InterchangeAssetImportData.h"
 #include "InterchangeGenericTexturePipeline.h"
 #include "InterchangeManager.h"
 #include "InterchangeMaterialDefinitions.h"
@@ -11,6 +12,7 @@
 #include "InterchangeMaterialInstanceNode.h"
 #include "InterchangeMaterialReferenceNode.h"
 #include "InterchangePipelineLog.h"
+#include "InterchangeSceneImportAsset.h"
 #include "InterchangeShaderGraphNode.h"
 #include "InterchangeSourceData.h"
 #include "InterchangeSparseVolumeTexturePipeline.h"
@@ -1109,20 +1111,40 @@ void UInterchangeGenericMaterialPipeline::ExecutePipeline(UInterchangeBaseNodeCo
 			BaseMaterialFactoryNode->AddInt32Attribute(TEXT("SearchLocation"), int32(SearchLocation));
 			BaseMaterialFactoryNode->SetCustomReferenceObject(ExistingMaterial);
 
+			const bool bIsReimport = CacheContextParam.ReimportAsset != nullptr;
+
 			// We're reimporting the existing material we found: Enable the factory node if we can.
 			// TODO: Can this potentially lead to issues, since the existing material to reimport is discovered by name and search location?
-			if (CacheContextParam.ReimportAsset && CacheContextParam.ReimportAsset == ExistingMaterial)
+			if (bIsReimport && CacheContextParam.ReimportAsset == ExistingMaterial)
 			{
 				const bool bIsMaterial = BaseMaterialFactoryNode->IsA<UInterchangeMaterialFactoryNode>()
 											&& ExistingMaterial->IsA<UMaterial>();
 
 				const bool bIsMaterialinstance = BaseMaterialFactoryNode->IsA<UInterchangeMaterialInstanceFactoryNode>()
-													&& ExistingMaterial->IsA<UMaterialInstance>();
+					&& ExistingMaterial->IsA<UMaterialInstance>();
+
 
 				// Reimport can only be done on material instances
 				bEnableFactoryNode = !bIsMaterial && bIsMaterialinstance;
 			}
-			// Found an existing material on the search location during an import --> Just use it as-is
+#if WITH_EDITORONLY_DATA
+			// This is a scene reimport of a material instance
+			// check whether or not it is a reimport of the current material instance factory node 
+			else if (bIsReimport && CacheContextParam.ReimportAsset->IsA<UInterchangeSceneImportAsset>() && ExistingMaterial->IsA<UMaterialInstance>())
+			{
+				// Skip it if it is not a reimport of the current material instance factory node
+				bEnableFactoryNode = false;
+
+				if (const UInterchangeAssetImportData* AssetImportData = Cast<UInterchangeAssetImportData>(ExistingMaterial->AssetImportData))
+				{
+					const UInterchangeSceneImportAsset* SceneImportAsset = Cast<UInterchangeSceneImportAsset>(CacheContextParam.ReimportAsset);
+					if (SceneImportAsset->AssetImportData)
+					{
+						bEnableFactoryNode = AssetImportData->GetFirstFilename().Equals(SceneImportAsset->AssetImportData->GetFirstFilename());
+					}
+				}
+			}
+#endif
 			else
 			{
 				// Disabling the factory node here is important so that we don't overwrite the existing asset.
diff --git a/Engine/Plugins/Interchange/Runtime/Source/Pipelines/Public/InterchangeGenericAssetsPipeline.h b/Engine/Plugins/Interchange/Runtime/Source/Pipelines/Public/InterchangeGenericAssetsPipeline.h
--- a/Engine/Plugins/Interchange/Runtime/Source/Pipelines/Public/InterchangeGenericAssetsPipeline.h
+++ b/Engine/Plugins/Interchange/Runtime/Source/Pipelines/Public/InterchangeGenericAssetsPipeline.h
@@ -108,6 +108,10 @@ public:
 	UPROPERTY(VisibleAnywhere, BlueprintReadWrite, Instanced, Category = "Grooms")
 	TObjectPtr<UInterchangeGenericGroomPipeline> GroomPipeline;
 
+	/** Internal property to cache the name of the sub-folder prefix for use during reimport */
+	UPROPERTY(meta = (PipelineInternalEditionData = "True"))
+	FString SceneNameFolderPrefix;
+
 	UE_API virtual void PreDialogCleanup(const FName PipelineStackName) override;
 
 	UE_API virtual bool IsSettingsAreValid(TOptional<FText>& OutInvalidReason) const override;
diff --git a/Engine/Source/Runtime/Interchange/Engine/Private/Tasks/InterchangeTaskCreateSceneObjects.cpp b/Engine/Source/Runtime/Interchange/Engine/Private/Tasks/InterchangeTaskCreateSceneObjects.cpp
--- a/Engine/Source/Runtime/Interchange/Engine/Private/Tasks/InterchangeTaskCreateSceneObjects.cpp
+++ b/Engine/Source/Runtime/Interchange/Engine/Private/Tasks/InterchangeTaskCreateSceneObjects.cpp
@@ -120,6 +120,13 @@ void UE::Interchange::FTaskCreateSceneObjects_GameThread::Execute()
 		CreateSceneObjectsParams.bIsReimport = ReimportObject != nullptr;
 		CreateSceneObjectsParams.AsyncHelperUid = AsyncHelper->UniqueId;
 
+		// Skip factory node if this is a reimport but no object to reimport is found and reimport is not requested
+		if (!CreateSceneObjectsParams.ReimportObject && ReimportObject && !FactoryNode->ShouldForceNodeReimport())
+		{
+			FactoryNode->SetSkipNodeImport();
+			continue;
+		}
+
 		if (AsyncHelper->BaseNodeContainers.IsValidIndex(SourceIndex))
 		{
 			CreateSceneObjectsParams.NodeContainer = AsyncHelper->BaseNodeContainers[SourceIndex].Get();
diff --git a/Engine/Source/Runtime/Interchange/Engine/Private/Tasks/InterchangeTaskImportObject.cpp b/Engine/Source/Runtime/Interchange/Engine/Private/Tasks/InterchangeTaskImportObject.cpp
--- a/Engine/Source/Runtime/Interchange/Engine/Private/Tasks/InterchangeTaskImportObject.cpp
+++ b/Engine/Source/Runtime/Interchange/Engine/Private/Tasks/InterchangeTaskImportObject.cpp
@@ -303,6 +303,14 @@ void UE::Interchange::FTaskImportObject_GameThread::Execute()
 	FString AssetName;
 	Private::InternalGetPackageName(*AsyncHelper, SourceIndex, PackageBasePath, FactoryNode, PackageName, AssetName);
 	UObject* ObjectToReimport = FFactoryCommon::GetObjectToReimport(Factory, AsyncHelper->TaskData.ReimportObject, *FactoryNode, PackageName, AssetName);
+
+	// Skip factory node if this is a reimport but no object to reimport is found and reimport is not requested
+	if (!ObjectToReimport && AsyncHelper->TaskData.ReimportObject && !FactoryNode->ShouldForceNodeReimport())
+	{
+		FactoryNode->SetSkipNodeImport();
+		return;
+	}
+
 	if (!ensure(!IsGarbageCollecting()))
 	{
 		//Skip this asset
