MaterialX:
- Add precise float value using a divide node when the number is of a magnitude 1-e7 or less (to avoid constant/scalar parameters to round the value to 0)
- Add scale for <heighttonormal> in case of using the Material Function HeightToNormalSmooth
- Set the Compression Setting for EXR texture to half float

diff --git a/Engine/Plugins/Interchange/Runtime/Source/Common/Public/MaterialX/InterchangeMaterialXDefinitions.h b/Engine/Plugins/Interchange/Runtime/Source/Common/Public/MaterialX/InterchangeMaterialXDefinitions.h
--- a/Engine/Plugins/Interchange/Runtime/Source/Common/Public/MaterialX/InterchangeMaterialXDefinitions.h
+++ b/Engine/Plugins/Interchange/Runtime/Source/Common/Public/MaterialX/InterchangeMaterialXDefinitions.h
@@ -255,6 +255,7 @@ namespace UE
 			inline constexpr const TCHAR* Foreground = TEXT("Foreground");
 			inline constexpr const TCHAR* Fraction = TEXT("Fraction");
 			inline constexpr const TCHAR* Gamma = TEXT("Gamma");
+			inline constexpr const TCHAR* Height = TEXT("Height");
 			inline constexpr const TCHAR* in1 = TEXT("in1");
 			inline constexpr const TCHAR* in2 = TEXT("in2");
 			inline constexpr const TCHAR* in3 = TEXT("in3");
diff --git a/Engine/Plugins/Interchange/Runtime/Source/Import/Private/MaterialX/InterchangeMaterialXTranslator.cpp b/Engine/Plugins/Interchange/Runtime/Source/Import/Private/MaterialX/InterchangeMaterialXTranslator.cpp
--- a/Engine/Plugins/Interchange/Runtime/Source/Import/Private/MaterialX/InterchangeMaterialXTranslator.cpp
+++ b/Engine/Plugins/Interchange/Runtime/Source/Import/Private/MaterialX/InterchangeMaterialXTranslator.cpp
@@ -77,6 +77,11 @@ TOptional<UE::Interchange::FImportImage> UInterchangeMaterialXTranslator::GetTex
 	}
 #endif
 
+	if (FPaths::GetExtension(Filename).Equals(TEXT("exr"), ESearchCase::IgnoreCase))
+	{
+		CompressionSettings = TextureCompressionSettings::TC_HalfFloat;
+	}
+
 	UE::Interchange::Private::FScopedTranslator ScopedTranslator(Filename, InResults, InAnalyticsHelper);
 	const IInterchangeTexturePayloadInterface* TextureTranslator = ScopedTranslator.GetPayLoadInterface<IInterchangeTexturePayloadInterface>();
 
@@ -118,6 +123,11 @@ TOptional<UE::Interchange::FImportBlockedImage> UInterchangeMaterialXTranslator:
 	}
 #endif
 
+	if (FPaths::GetExtension(Filename).Equals(TEXT("exr"), ESearchCase::IgnoreCase))
+	{
+		CompressionSettings = TextureCompressionSettings::TC_HalfFloat;
+	}
+
 	FImportBlockedImage BlockedImage;
 	BlockedImage.CompressionSettings = CompressionSettings;
 
diff --git a/Engine/Plugins/Interchange/Runtime/Source/Import/Private/MaterialX/MaterialXUtils/MaterialXSurfaceShaderAbstract.cpp b/Engine/Plugins/Interchange/Runtime/Source/Import/Private/MaterialX/MaterialXUtils/MaterialXSurfaceShaderAbstract.cpp
--- a/Engine/Plugins/Interchange/Runtime/Source/Import/Private/MaterialX/MaterialXUtils/MaterialXSurfaceShaderAbstract.cpp
+++ b/Engine/Plugins/Interchange/Runtime/Source/Import/Private/MaterialX/MaterialXUtils/MaterialXSurfaceShaderAbstract.cpp
@@ -133,10 +133,27 @@ bool FMaterialXSurfaceShaderAbstract::AddFloatAttribute(MaterialX::InputPtr Inpu
 			mx::NodePtr Node = Input->getParent()->asA<mx::Node>();
 			FString NodeName = Node ? Node->getName().c_str() + FString{ TEXT("_") } : FString{};
 			NodeName += Input->getName().c_str();
+			
 
 			UInterchangeShaderNode* ScalarParameterNode = CreateShaderNode(Input, NodeName, ScalarParameter::Name.ToString());
-			ScalarParameterNode->AddFloatAttribute(UInterchangeShaderPortsAPI::MakeInputParameterKey(ScalarParameter::Attributes::DefaultValue.ToString()), Value);
-			return UInterchangeShaderPortsAPI::ConnectDefaultOuputToInput(ShaderNode, InputChannelName, ScalarParameterNode->GetUniqueID());
+			UInterchangeShaderNode* NodeToConnectTo = ScalarParameterNode;
+			
+			const float AbsValue = FMath::Abs(Value);
+			const int32 Exponent = AbsValue > 0.f ? -FMath::FloorToInt(FMath::LogX(10, AbsValue)) : 0;
+			double Power = 1.;
+			// In the Material Editor, Constant/Scalars are rounded to 0 with at least 7 digits after the comma
+			if (Exponent >= 7)
+			{
+				Power = FMath::Pow(10., Exponent);
+
+				UInterchangeShaderNode* DivideNode = CreateShaderNode(Input, NodeName + TEXT("_Divide"), UE::Expressions::Names::Divide);
+				DivideNode->AddFloatAttribute(UInterchangeShaderPortsAPI::MakeInputValueKey(UE::Expressions::Inputs::B), Power);
+				UInterchangeShaderPortsAPI::ConnectDefaultOuputToInput(DivideNode, UE::Expressions::Inputs::A, ScalarParameterNode->GetUniqueID());
+				NodeToConnectTo = DivideNode;
+			}
+
+			ScalarParameterNode->AddFloatAttribute(UInterchangeShaderPortsAPI::MakeInputParameterKey(ScalarParameter::Attributes::DefaultValue.ToString()), Value * Power);
+			return UInterchangeShaderPortsAPI::ConnectDefaultOuputToInput(ShaderNode, InputChannelName, NodeToConnectTo->GetUniqueID());
 		}
 	}
 
@@ -1087,12 +1104,32 @@ void FMaterialXSurfaceShaderAbstract::ConnectHeightToNormalInputToOutput(const F
 		else
 		{
 			using namespace UE::Interchange::Materials::Standard::Nodes;
-			SetAttributeNewName(Connect.UpstreamNode->getInput("in"), "Height");
+
+			// if we have a scale the multiply node input A will be the Parent Shader Node otherwise it's just the MF
+			UInterchangeShaderNode* ScaleNode = nullptr;
+			FString HeightToNormalNodeName = Connect.UpstreamNode->getName().c_str();
+			if (mx::InputPtr InputScale = Connect.UpstreamNode->getInput("scale"))
+			{
+				ScaleNode = CreateShaderNode(Connect.UpstreamNode, HeightToNormalNodeName, Multiply::Name.ToString());
+				HeightToNormalNodeName += TEXT("_Smooth"); // the MF won't be the main node to connect to
+				SetAttributeNewName(Connect.UpstreamNode->getInput("in"), "A");
+				AddAttributeFromValueOrInterface(InputScale, Multiply::Inputs::B.ToString(), ScaleNode);
+			}
+			else
+			{
+				SetAttributeNewName(Connect.UpstreamNode->getInput("in"), "Height");
+			}
+
 			// This MF outputs a normal vector in World Space, we need to transform the output in Tangent Space
 			UInterchangeFunctionCallShaderNode* HeightToNormalSmoothShaderNode =
 				CreateFunctionCallShaderNode(Connect.UpstreamNode,
-											 Connect.UpstreamNode->getName().c_str(),
+											 HeightToNormalNodeName,
 											 UE::MaterialFunctions::Path::HeightToNormalSmooth);
+			
+			if (ScaleNode)
+			{
+				UInterchangeShaderPortsAPI::ConnectDefaultOuputToInput(HeightToNormalSmoothShaderNode, UE::Expressions::Inputs::Height, ScaleNode->GetUniqueID());
+			}
 
 			UInterchangeShaderNode* TransformNode = CreateShaderNode(Connect.UpstreamNode, (Connect.UpstreamNode->getName() + "_TS").c_str(), TransformVector::Name.ToString());
 			TransformNode->AddInt32Attribute(TransformVector::Attributes::TransformSourceType.ToString(), EMaterialVectorCoordTransformSource::TRANSFORMSOURCE_World);
@@ -1105,6 +1142,7 @@ void FMaterialXSurfaceShaderAbstract::ConnectHeightToNormalInputToOutput(const F
 
 			UnitNode->AddFloatAttribute(UInterchangeShaderPortsAPI::MakeInputValueKey(Multiply::Inputs::B.ToString()), 0.01f);
 			UInterchangeShaderPortsAPI::ConnectDefaultOuputToInput(UnitNode, Multiply::Inputs::A.ToString(), PositionNode->GetUniqueID());
+
 			UInterchangeShaderPortsAPI::ConnectDefaultOuputToInput(HeightToNormalSmoothShaderNode, UE::Expressions::Inputs::AbsoluteWorldPosition, UnitNode->GetUniqueID());
 
 			UInterchangeShaderPortsAPI::ConnectDefaultOuputToInput(TransformNode, TransformVector::Inputs::Input.ToString(), HeightToNormalSmoothShaderNode->GetUniqueID());
