Interchange: Prevent a couple warnings when trying to fetch a transform from a node with none authored.

diff --git a/Engine/Plugins/Interchange/Runtime/Source/Nodes/Private/InterchangeSceneNode.cpp b/Engine/Plugins/Interchange/Runtime/Source/Nodes/Private/InterchangeSceneNode.cpp
--- a/Engine/Plugins/Interchange/Runtime/Source/Nodes/Private/InterchangeSceneNode.cpp
+++ b/Engine/Plugins/Interchange/Runtime/Source/Nodes/Private/InterchangeSceneNode.cpp
@@ -395,7 +395,13 @@ bool UInterchangeSceneNode::GetGlobalTransformInternal(const UE::Interchange::FA
 	if (!CacheTransform.IsSet())
 	{
 		FTransform LocalTransform = FTransform::Identity;
-		UE::Interchange::FAttributeStorage::TAttributeHandle<FTransform> AttributeHandle = GetAttributeHandle<FTransform>(TransformKey);
+
+		UE::Interchange::FAttributeStorage::TAttributeHandle<FTransform> AttributeHandle;
+		if (Attributes->ContainAttribute(TransformKey))
+		{
+			AttributeHandle = GetAttributeHandle<FTransform>(TransformKey);
+		}
+
 		if ((AttributeHandle.IsValid() && AttributeHandle.Get(LocalTransform) == UE::Interchange::EAttributeStorageResult::Operation_Success)
 			|| bAllowsNoTransforms)
 		{
