Interchange - Import performance improvement attempt

This is four possible optimizations that came up after looking at why USD legacy was not importing same asset as Interchange and why Interchange seemed to have extra steps, logs and progress bars.

Optimization 1:
- Removing the Sync To Browser at the end of the import that was causing freeze when there is 100+ assets.

Optimization 2:
- Enabling bUseHashAsGuid during CommitMeshDescreption which seems to save an unecessary rebuild of the asset.

Optimization 3:
- Disabling NavigationData on newly created StaticMesh to prevent  allocating time and memory for what might be an unused (or undesired) functionnality of that StaticMesh. (To Be Discussed)

Optimization 4:
- Avoid generating complex mesh collision at all by enabling bNeverNeedsCookedCollisionData on StaticMeshes. (To Be Discussed)

Shelve #2 update:
To enable those optimizations, and be able to test their effects in automation. a CVar has been exposed for each one of them (thanks Jean-Luc for the suggestion)
`Interchange.FeatureFlags.Import.SyncToBrowser 0` // (default is true)
`Interchange.FeatureFlags.Import.UseHashAsGuid 1`  // (default is false)
`Interchange.FeatureFlags.Import.MarkAsNotHavingNavigationData 1` // (default is false)
`Interchange.FeatureFlags.Import.NeverNeedsCookedCollisionData 1` // (default is false)

Some test result:.
- `cube_3000` (3000 SM + 1 Mat) went from ~48 seconds to  ~10 seconds. Legacy is ~15 seconds. (Note: Legacy log showing 3s import but editor was frozen for 15s stopwatched)
- `cube_10000` (10000 SM + 1 Mat) went from ~320 seconds (240 of which are sync to browser) to  ~40 seconds. Legacy is ~110 seconds (13s logged).
- `4004MooreLane_Assembly.usda` went from ~53 seconds to ~25 seconds. Legacy is ~35 seconds (29s logged). (Note: Legacy is importing 882 assets vs 1337 assets for Interchange)

Except optimization #1, this CL targets specifically StaticMeshes.

#rnx

diff --git a/Engine/Plugins/Interchange/Editor/Source/InterchangeEditor/Private/InterchangeEditorModule.cpp b/Engine/Plugins/Interchange/Editor/Source/InterchangeEditor/Private/InterchangeEditorModule.cpp
--- a/Engine/Plugins/Interchange/Editor/Source/InterchangeEditor/Private/InterchangeEditorModule.cpp
+++ b/Engine/Plugins/Interchange/Editor/Source/InterchangeEditor/Private/InterchangeEditorModule.cpp
@@ -22,6 +22,13 @@
 
 DEFINE_LOG_CATEGORY(LogInterchangeEditor);
 
+static bool GInterchangeSyncToBrowser = true;
+static FAutoConsoleVariableRef CCvarInterchangeSyncToBrowser(
+	TEXT("Interchange.FeatureFlags.Import.SyncToBrowser"),
+	GInterchangeSyncToBrowser,
+	TEXT("Whether to select the imported assets inside the content browser once the import is done. Disabling this option can prevent the editor from freezing when importing a large number of assets."),
+	ECVF_Default);
+
 namespace UE::Interchange::InterchangeEditorModule
 {
 	static bool bOldAutoSaveState = false;
diff --git a/Engine/Plugins/Interchange/Runtime/Source/Import/Private/Mesh/InterchangeStaticMeshFactory.cpp b/Engine/Plugins/Interchange/Runtime/Source/Import/Private/Mesh/InterchangeStaticMeshFactory.cpp
--- a/Engine/Plugins/Interchange/Runtime/Source/Import/Private/Mesh/InterchangeStaticMeshFactory.cpp
+++ b/Engine/Plugins/Interchange/Runtime/Source/Import/Private/Mesh/InterchangeStaticMeshFactory.cpp
@@ -44,6 +44,27 @@
 #include "EditorFramework/AssetImportData.h"
 #endif //WITH_EDITORONLY_DATA
 
+static bool GInterchangeUseHashAsGuid = false;
+static FAutoConsoleVariableRef CCvarInterchangeUseHashAsGuid(
+	TEXT("Interchange.FeatureFlags.Import.UseHashAsGuid"),
+	GInterchangeUseHashAsGuid,
+	TEXT("Whether to enable the bUseHashAsGuid option when importing Static Meshes. Useful to prevent recomputing content already in cache and improve import time."),
+	ECVF_Default);
+
+static bool GInterchangeMarkAsNotHavingNavigationData = false;
+static FAutoConsoleVariableRef CCvarInterchangeMarkAsNotHavingNavigationData(
+	TEXT("Interchange.FeatureFlags.Import.MarkAsNotHavingNavigationData"),
+	GInterchangeMarkAsNotHavingNavigationData,
+	TEXT("If true, newly imported Static Meshes will have \"Has Navigation Data\" option disabled and will be ignored during navmesh generation."),
+	ECVF_Default);
+
+static bool GInterchangeNeverNeedsCookedCollisionData = false;
+static FAutoConsoleVariableRef CCvarInterchangeNeverNeedsCookedCollisionData(
+	TEXT("Interchange.FeatureFlags.Import.NeverNeedsCookedCollisionData"),
+	GInterchangeNeverNeedsCookedCollisionData,
+	TEXT("If true, imported Static Meshes will have \"Never Needs Cooked Collision Data\" option enabled and will prevent them from having a Complex Collision. Only used if Import Collisions setting is disabled."),
+	ECVF_Default);
+
 UClass* UInterchangeStaticMeshFactory::GetFactoryClass() const
 {
 	return UStaticMesh::StaticClass();
@@ -290,6 +311,11 @@ UInterchangeFactoryBase::FImportAssetResult UInterchangeStaticMeshFactory::Begin
 	if (!ExistingAsset)
 	{
 		StaticMesh = NewObject<UStaticMesh>(Arguments.Parent, *Arguments.AssetName, RF_Public | RF_Standalone);
+
+		if (GInterchangeMarkAsNotHavingNavigationData)
+		{
+			StaticMesh->MarkAsNotHavingNavigationData();
+		}
 	}
 	else
 	{
@@ -806,9 +832,7 @@ UInterchangeFactoryBase::FImportAssetResult UInterchangeStaticMeshFactory::EndIm
 		else
 		{
 			//If there is some collision, we remove the original imported collision and add any editor generated collision
-#if WITH_EDITOR
 			ImportAssetObjectData.AggregateGeom.EmptyImportedElements();
-#endif
 
 			for (FKBoxElem& BoxElem : ImportAssetObjectData.AggregateGeom.BoxElems)
 			{
@@ -911,15 +935,22 @@ UInterchangeFactoryBase::FImportAssetResult UInterchangeStaticMeshFactory::EndIm
 				break;
 			}
 		}
-#endif // WITH_EDITOR
 #if WITH_EDITORONLY_DATA
 		else
 		{
 			StaticMesh->SetCustomizedCollision(true);
 		}
-	}
 #endif // WITH_EDITORONLY_DATA
-#if WITH_EDITOR
+	}
+	else
+	{
+		if (GInterchangeNeverNeedsCookedCollisionData)
+		{
+			// this will prevent the complex collision mesh from being generated
+			StaticMesh->GetBodySetup()->bNeverNeedsCookedCollisionData = true;
+		}
+	}
+
 	//Lod group need to use the static mesh API and cannot use the apply delegate
 	if (!Arguments.ReimportObject)
 	{
@@ -941,8 +972,8 @@ UInterchangeFactoryBase::FImportAssetResult UInterchangeStaticMeshFactory::EndIm
 #if WITH_EDITORONLY_DATA
 	// Clear the mesh descriptions to free memory, we don't need them anymore since all the meshes have been committed at this point.
 	StaticMesh->ClearMeshDescriptions();
-#endif
-#endif
+#endif // WITH_EDITORONLY_DATA
+#endif // WITH_EDITOR
 
 	ImportAssetResult.ImportedObject = StaticMesh;
 	return ImportAssetResult;
@@ -961,6 +992,11 @@ void UInterchangeStaticMeshFactory::CommitMeshDescriptions(UStaticMesh& StaticMe
 	UStaticMesh::FCommitMeshDescriptionParams CommitMeshDescriptionParams;
 	CommitMeshDescriptionParams.bMarkPackageDirty = false; // Marking packages dirty isn't thread-safe
 
+	if (GInterchangeUseHashAsGuid)
+	{
+		CommitMeshDescriptionParams.bUseHashAsGuid = true;
+	}
+
 	for (int32 LodIndex = 0; LodIndex < LodMeshDescriptions.Num(); ++LodIndex)
 	{
 		FMeshDescription* StaticMeshDescription = StaticMesh.CreateMeshDescription(LodIndex);
diff --git a/Engine/Source/Developer/AssetTools/Private/AssetTools.cpp b/Engine/Source/Developer/AssetTools/Private/AssetTools.cpp
--- a/Engine/Source/Developer/AssetTools/Private/AssetTools.cpp
+++ b/Engine/Source/Developer/AssetTools/Private/AssetTools.cpp
@@ -3869,8 +3869,12 @@ TArray<UObject*> UAssetToolsImpl::ImportAssetsInternal(const TArray<FString>& Fi
 								}
 							}
 
-							//Force browser to sync to the import asset if there is only one asset imported
-							if (bSyncToBrowser)
+							const IConsoleVariable* Var = IConsoleManager::Get().FindConsoleVariable(TEXT("Interchange.FeatureFlags.Import.SyncToBrowser")); // Declared in InterchangeEngineModule.cpp
+
+							// If the CVar is found, it overrides the passed value.
+							const bool bResolvedSyncToBrowser = Var ? Var->GetBool() : bSyncToBrowser;
+							
+							if (bResolvedSyncToBrowser)
 							{
 								UAssetToolsImpl::Get().SyncBrowserToAssets(MainAssets);
 							}
diff --git a/Engine/Source/Runtime/Interchange/Engine/Private/InterchangeManager.cpp b/Engine/Source/Runtime/Interchange/Engine/Private/InterchangeManager.cpp
--- a/Engine/Source/Runtime/Interchange/Engine/Private/InterchangeManager.cpp
+++ b/Engine/Source/Runtime/Interchange/Engine/Private/InterchangeManager.cpp
@@ -815,6 +815,8 @@ void UE::Interchange::FImportResult::SetInProgress()
 
 void UE::Interchange::FImportResult::SetDone()
 {
+	TRACE_CPUPROFILER_EVENT_SCOPE(UE::Interchange::FImportResult::SetDone);
+
 	SetInProgress(); // Make sure we always pass through the InProgress state
 
 	EStatus ExpectedStatus = EStatus::InProgress;
