Interchange - Fixed static mesh actor wrongfully pointing to the same static mesh

This issue is a regression from CL-44242204 (UE-298322)
If the model did not include LOD, but had a mesh associated to its root, then it would try to create a LOD group from one of it's children.

To fix the issue, we need to first check if a scene node is specialized to a LOD Group **before** assigning a value to LodGroupStaticMeshFactoryNode. Which is what the semantic before CL-44242204 was doing.

Shelve 2: Added a test to cover this edge case.

#rnx

diff --git a/Engine/Plugins/Interchange/Runtime/Source/Pipelines/Private/InterchangeGenericScenesPipeline.cpp b/Engine/Plugins/Interchange/Runtime/Source/Pipelines/Private/InterchangeGenericScenesPipeline.cpp
--- a/Engine/Plugins/Interchange/Runtime/Source/Pipelines/Private/InterchangeGenericScenesPipeline.cpp
+++ b/Engine/Plugins/Interchange/Runtime/Source/Pipelines/Private/InterchangeGenericScenesPipeline.cpp
@@ -664,9 +664,10 @@ void UInterchangeGenericLevelPipeline::ExecuteSceneNodePreImport(const FTransfor
 			return SpecializeTypes.Contains(UE::Interchange::FSceneNodeStaticData::GetLodGroupSpecializeTypeString());
 		}();
 
-	const UInterchangeStaticMeshFactoryNode* LodGroupStaticMeshFactoryNode = SceneNodeData.LodGroupFactoryNode;
+	const UInterchangeStaticMeshFactoryNode* LodGroupStaticMeshFactoryNode = nullptr;
 	if (bIsLodGroup)
 	{	
+		LodGroupStaticMeshFactoryNode = SceneNodeData.LodGroupFactoryNode;
 		if (LodGroupStaticMeshFactoryNode)
 		{
 			TArray<FString> LodDataUids;
