[Interchange][FBX] UFBX failing FBX asset import tests (missing features)

- material support about the same as the FBX SDK's - phong, normal/bump,  emissive, opacity/mask. Using materials instancing with same parent materials.

#virtualized

diff --git a/Engine/Plugins/Interchange/Runtime/Source/Parsers/Fbx/Private/Ufbx/UfbxMaterial.cpp b/Engine/Plugins/Interchange/Runtime/Source/Parsers/Fbx/Private/Ufbx/UfbxMaterial.cpp
--- a/Engine/Plugins/Interchange/Runtime/Source/Parsers/Fbx/Private/Ufbx/UfbxMaterial.cpp
+++ b/Engine/Plugins/Interchange/Runtime/Source/Parsers/Fbx/Private/Ufbx/UfbxMaterial.cpp
@@ -6,7 +6,7 @@
 #include "UfbxConvert.h"
 
 #include "InterchangeMaterialDefinitions.h"
-#include "InterchangeShaderGraphNode.h"
+#include "InterchangeMaterialInstanceNode.h"
 #include "InterchangeTexture2DNode.h"
 #include "Nodes/InterchangeBaseNodeContainer.h"
 
@@ -20,196 +20,379 @@
 
 namespace UE::Interchange::Private
 {
+
+	namespace Inputs
+	{
+		///PostFixes
+		namespace PostFix
+		{
+			const FString Color_RGB = TEXT("_RGB");
+			const FString Color_A = TEXT("_A");
+
+			const FString TexCoord = TEXT("_TexCoord");
+
+			const FString OffsetX = TEXT("_Offset_X");
+			const FString OffsetY = TEXT("_Offset_Y");
+			const FString ScaleX = TEXT("_Scale_X");
+			const FString ScaleY = TEXT("_Scale_Y");
+
+			const FString OffsetScale = TEXT("_OffsetScale");
+
+			const FString Rotation = TEXT("_Rotation");
+
+			const FString TilingMethod = TEXT("_TilingMethod");
+		}
+	}
+
 	class  FMaterialMapConverter
 	{
 	public:
-		explicit FMaterialMapConverter(FUfbxParser& InParser, UInterchangeBaseNodeContainer& InNodeContainer, UInterchangeShaderGraphNode& InShaderGraphNode)
+		explicit FMaterialMapConverter(FUfbxParser& InParser, UInterchangeBaseNodeContainer& InNodeContainer, UInterchangeBaseNode& InMaterialInstance)
 			: Parser(InParser)
 			, NodeContainer(InNodeContainer)
-			, ShaderGraphNode(InShaderGraphNode)
+			, MaterialInstance(InMaterialInstance)
 		{
 		}
 
-		UInterchangeShaderNode* CreateTextureSampler(FName InputName, ufbx_texture* Texture)
+		FString FindTextureFile(ufbx_texture* Texture)
 		{
-			using namespace UE::Interchange::Materials;
+			FString TextureFilepath = Convert::ToUnrealString(Texture->absolute_filename);
+			FPaths::NormalizeFilename(TextureFilepath);
+			if(FPaths::FileExists(TextureFilepath))
+			{
+				return TextureFilepath;
+			}
+			FString FileBasePath = FPaths::GetPath(Parser.SourceFilename);
 
-			const FString TextureFilename = [this, Texture]
+			FString RelativeToFBXTextureFilepath = FileBasePath / Convert::ToUnrealString(Texture->relative_filename);
+			FPaths::NormalizeFilename(RelativeToFBXTextureFilepath);
+			if( FPaths::FileExists(RelativeToFBXTextureFilepath) )
 			{
-				FString TextureFilepath = Convert::ToUnrealString(Texture->absolute_filename);
-				FPaths::NormalizeFilename(TextureFilepath);
-				if(FPaths::FileExists(TextureFilepath))
+				return RelativeToFBXTextureFilepath;
+			}
+
+			// Some fbx files do not store the actual absolute filename as absolute and it is actually relative.  Try to get it relative to the FBX file we are importing
+			FString AbsoluteAsRelativeToFBXTextureFilepath = Convert::ToUnrealString(Texture->filename);
+			FPaths::NormalizeFilename(AbsoluteAsRelativeToFBXTextureFilepath );
+			if( FPaths::FileExists(AbsoluteAsRelativeToFBXTextureFilepath ) )
+			{
+				return AbsoluteAsRelativeToFBXTextureFilepath ;
+			}
+
+			if (TextureFilepath.IsEmpty())
+			{
+				// Make a name so it could be used as a key or in a message
+				return FString::Printf(TEXT("Texture_%d"), Texture->element_id);
+			}
+			return TextureFilepath;
+		}
+
+		const UInterchangeTexture2DNode* CreateTexture2DNode(ufbx_texture* Texture)
+		{
+			const FString TextureFilename = FindTextureFile(Texture);
+			if (TextureFilename.IsEmpty() || !FPaths::FileExists(TextureFilename))
+			{
+				if (!GIsAutomationTesting)
 				{
-					return TextureFilepath;
+					UInterchangeResultTextureDisplay_TextureFileDoNotExist* Message = Parser.AddMessage<UInterchangeResultTextureDisplay_TextureFileDoNotExist>();
+					Message->TextureName = TextureFilename;
+					Message->MaterialName = MaterialInstance.GetDisplayLabel();
 				}
-				FString FileBasePath = FPaths::GetPath(Parser.SourceFilename);
 
-				FString RelativeToFBXTextureFilepath = FileBasePath / Convert::ToUnrealString(Texture->relative_filename);
-				FPaths::NormalizeFilename(RelativeToFBXTextureFilepath);
-				if( FPaths::FileExists(RelativeToFBXTextureFilepath) )
+				return nullptr;
+			}
+			return CreateTexture2DNode(TextureFilename);
+		}
+
+		// Expected normalized full path for existing file textures
+		const UInterchangeTexture2DNode* CreateTexture2DNode(const FString& TextureFilePath)
+		{
+			const FString TextureName = FPaths::GetBaseFilename(TextureFilePath);
+			const FString TextureNodeID = UInterchangeTextureNode::MakeNodeUid(TextureName);
+
+			if (const UInterchangeTexture2DNode* TextureNode = Cast<const UInterchangeTexture2DNode>(NodeContainer.GetNode(TextureNodeID)))
+			{
+				return TextureNode;
+			}
+
+			UInterchangeTexture2DNode* NewTextureNode = UInterchangeTexture2DNode::Create(&NodeContainer, TextureNodeID);
+			NewTextureNode->SetDisplayLabel(TextureName);
+
+			//All texture translator expect a file as the payload key
+			NewTextureNode->SetPayLoadKey(TextureFilePath);
+
+			return NewTextureNode;
+		}
+
+		FUfbxParser& Parser;
+		UInterchangeBaseNodeContainer& NodeContainer;
+		UInterchangeBaseNode& MaterialInstance;
+	};
+
+	UInterchangeMaterialInstanceNode* CreateMaterialInstanceNode(UInterchangeBaseNodeContainer& NodeContainer, const FString& NodeUID, const FString& NodeName)
+	{
+		UInterchangeMaterialInstanceNode* MaterialInstanceNode = NewObject<UInterchangeMaterialInstanceNode>(&NodeContainer);
+		NodeContainer.SetupNode(MaterialInstanceNode, NodeUID, NodeName, EInterchangeNodeContainerType::TranslatedAsset);
+
+		return MaterialInstanceNode;
+	}
+
+	void FUfbxMaterial::AddMaterials(UInterchangeBaseNodeContainer& NodeContainer)
+	{
+		for (ufbx_material* Material : Parser.Scene->materials)
+		{
+			// #ufbx_todo: Uid in FFbxMaterial parser is made from sanitized name
+			// but two different materials  might happen to have same name AFTER sanitation.
+			// SO material name shoul dbe made unique for this case. E.g. by adding number after name string.
+			// This means that we need to memoize mapping of material to this generated name to assign proper names to meshes
+			// FString MaterialName = Parser.GetFbxHelper()->GetFbxObjectName(SurfaceMaterial);
+			// FString NodeUid = TEXT("\\Material\\") + MaterialName;
+			FString MaterialName = Parser.GetMaterialName(*Material).ToString();
+			
+			using namespace UE::Interchange::Materials;
+			
+			FString MaterialUid = Parser.GetMaterialUid(MaterialName);
+			UInterchangeMaterialInstanceNode* MaterialInstanceNode = CreateMaterialInstanceNode(NodeContainer, MaterialUid, MaterialName);
+
+			// WithUVTransform
+			bool bHasUVTransform = false;
+			// "/InterchangeAssets/Materials/PhongSurfaceMaterial.PhongSurfaceMaterial
+			FString ParentMaterial{ TEXT("PhongSurfaceMaterial") };
+
+			FMaterialMapConverter MaterialMapConverter(Parser, NodeContainer, *MaterialInstanceNode);
+			
+			using namespace UE::Interchange::Materials;
+
+			auto SetColorFromMap = [&MaterialInstanceNode](FName InputName, const ufbx_material_map& Map, float Factor)
+			{
+				if (Map.value_components >=3)
 				{
-					return RelativeToFBXTextureFilepath;
+					const FString InputAttributeKey = InputName.ToString();
+					FLinearColor PropertyValue(Convert::ConvertVec4(Map.value_vec4) * Factor);
+					MaterialInstanceNode->AddVectorParameterValue(InputAttributeKey, PropertyValue);
 				}
+			};
+
+			auto SetScalarFromMap = [&MaterialInstanceNode](FName InputName, const ufbx_material_map& Map, float Factor)
+			{
+				const FString InputAttributeKey = InputName.ToString();
+				float PropertyValue(Map.value_real);
+				MaterialInstanceNode->AddScalarParameterValue(InputAttributeKey, PropertyValue);
+			};
 
-				// Some fbx files do not store the actual absolute filename as absolute and it is actually relative.  Try to get it relative to the FBX file we are importing
-				FString AbsoluteAsRelativeToFBXTextureFilepath = Convert::ToUnrealString(Texture->filename);
-				FPaths::NormalizeFilename(AbsoluteAsRelativeToFBXTextureFilepath );
-				if( FPaths::FileExists(AbsoluteAsRelativeToFBXTextureFilepath ) )
+			auto SetTextureFromMap = [&MaterialInstanceNode, &MaterialMapConverter, &bHasUVTransform](FName InputName, const ufbx_material_map& Map, float Factor) -> bool
+			{
+				ufbx_texture* Texture = Map.texture_enabled ? Map.texture : nullptr;
+				if (!Texture)
 				{
-					return AbsoluteAsRelativeToFBXTextureFilepath ;
+					return false;	
 				}
 
-				if (TextureFilepath.IsEmpty())
+				const UInterchangeTexture2DNode* TextureNode = MaterialMapConverter.CreateTexture2DNode(Texture);
+
+				if (!TextureNode)
 				{
-					// Make a name so it could be used as a key or in a message
-					return FString::Printf(TEXT("Texture_%d"), Texture->element_id);
+					return false;
 				}
-				return TextureFilepath;
-			}();
 
-			const FString NodeName = TEXT("Sampler_") + InputName.ToString() + TEXT("Map");
+				MaterialInstanceNode->AddTextureParameterValue(InputName.ToString() + TEXT("Map"), TextureNode->GetUniqueID());
 
-			FString ShaderUniqueID = ShaderGraphNode.GetUniqueID();
-			const FString SamplerNodeUid = UInterchangeShaderNode::MakeNodeUid(NodeName, ShaderUniqueID);
+				// #ufbx_todo: FBX parser used MapWeight 1 and multiplied although there was Factor provided. Why?
+				MaterialInstanceNode->AddScalarParameterValue(InputName.ToString() + TEXT("MapWeight"), Factor);
 
-			// #ufbx_todo: ??
-			// if (const UInterchangeShaderNode* SamplerNode = Cast<UInterchangeShaderNode>(NodeContainer.GetNode(SamplerNodeUid)))
-			// {
-			// 	return SamplerNode;
-			// }
+				if (Texture->has_uv_transform)
+				{
+					bHasUVTransform = true;
+					const ufbx_transform& UVTransform = Texture->uv_transform;
+					FVector4f OffsetScale(
+						UVTransform.translation.x, UVTransform.translation.y, 
+						UVTransform.scale.x, UVTransform.scale.y);
+					MaterialInstanceNode->AddVectorParameterValue(InputName.ToString() + TEXT("Map") + Inputs::PostFix::OffsetScale, OffsetScale);
+
+					FString RotationParameterName = InputName.ToString() + TEXT("Map") + Inputs::PostFix::Rotation;
+					Math::TVector<double> AttributeValue = Convert::ConvertQuat(UVTransform.rotation).Euler();
+					MaterialInstanceNode->AddScalarParameterValue(RotationParameterName, FMath::DegreesToRadians(AttributeValue.Z));
+				}
 
-			UInterchangeShaderNode* TextureSampleShader = UInterchangeShaderNode::Create(&NodeContainer, NodeName, ShaderUniqueID);
-			TextureSampleShader->SetDisplayLabel(InputName.ToString());
-			TextureSampleShader->SetCustomShaderType(Standard::Nodes::TextureSample::Name.ToString());
+				return true;
+			};
 
-			// Return incomplete texture sampler if texture file does not exist
-			if (TextureFilename.IsEmpty() || !FPaths::FileExists(TextureFilename))
+			bool bHasInput = false;
+
+			// Diffuse
+			if (Material->fbx.diffuse_color.has_value)
 			{
-				if (Texture->content.data && Texture->content.size > 0)
+				FName InputName = Phong::Parameters::DiffuseColor;
+				const float Factor = Material->fbx.diffuse_factor.value_real;
+				const ufbx_material_map& Map = Material->fbx.diffuse_color;
+
+				if (!SetTextureFromMap(InputName, Map, Factor))
 				{
-#if WITH_ENGINE
-					Parser.PayloadContexts.Add(TextureFilename, FPayloadContext(FPayloadContext::Element, Texture->element_id));
-#endif
+					SetColorFromMap(InputName, Map, Factor);
 				}
-				else
+
+				bHasInput = true;
+			}
+
+			// Ambient
+			if (Material->fbx.ambient_color.has_value)
+			{
+				FName InputName = Phong::Parameters::AmbientColor;
+				const float Factor = Material->fbx.ambient_factor.value_real;
+				const ufbx_material_map& Map = Material->fbx.ambient_color;
+
+				if (!SetTextureFromMap(InputName, Map, Factor))
 				{
-					if(!GIsAutomationTesting)
-					{
-						UInterchangeResultTextureDisplay_TextureFileDoNotExist* Message = Parser.AddMessage<UInterchangeResultTextureDisplay_TextureFileDoNotExist>();
-						Message->TextureName = TextureFilename;
-						Message->MaterialName = ShaderGraphNode.GetDisplayLabel();
-					}
-					return TextureSampleShader;
+					SetColorFromMap(InputName, Map, Factor);
 				}
+
+				bHasInput = true;
 			}
 
-			const UInterchangeTexture2DNode* TextureNode = CreateTexture2DNode(TextureFilename);
+			// Emissive
+			if (Material->fbx.emission_color.has_value)
+			{
+				const FName InputName = Phong::Parameters::EmissiveColor;
+				const float Factor = Material->fbx.emission_factor.value_real;
+				const ufbx_material_map& Map = Material->fbx.emission_color;
+
+				if (!SetTextureFromMap(InputName, Map, Factor))
+				{
+					SetColorFromMap(InputName, Map, Factor);
+				}
 
-			TextureSampleShader->AddStringAttribute(UInterchangeShaderPortsAPI::MakeInputValueKey(Standard::Nodes::TextureSample::Inputs::Texture.ToString()), TextureNode->GetUniqueID());
+				bHasInput = true;
+			}
 
-			if (Texture->has_uv_transform)
+			// if it has no specular no shininess then just take a Lambert material
+			bool bIsPhong = false;
+
+			// Specular
+			if (Material->fbx.specular_color.has_value)
 			{
-				
-				UInterchangeShaderNode* TextureCoordinateShader = UInterchangeShaderNode::Create(&NodeContainer, InputName.ToString() + TEXT("_Coordinate"), ShaderUniqueID);
-				TextureCoordinateShader->SetCustomShaderType(Standard::Nodes::TextureCoordinate::Name.ToString());
+				const FName InputName = Phong::Parameters::SpecularColor;
+				const float Factor = Material->fbx.specular_factor.value_real;
+				const ufbx_material_map& Map = Material->fbx.specular_color;
 
-				TextureCoordinateShader->AddFloatAttribute(UInterchangeShaderPortsAPI::MakeInputValueKey(Standard::Nodes::TextureCoordinate::Inputs::UTiling.ToString()), (float)Texture->uv_transform.scale.x);
-				TextureCoordinateShader->AddFloatAttribute(UInterchangeShaderPortsAPI::MakeInputValueKey(Standard::Nodes::TextureCoordinate::Inputs::VTiling.ToString()), (float)Texture->uv_transform.scale.y);
+				if (!SetTextureFromMap(InputName, Map, Factor))
+				{
+					SetColorFromMap(InputName, Map, Factor);
+				}
 
-				UInterchangeShaderPortsAPI::ConnectDefaultOuputToInput(TextureSampleShader, Standard::Nodes::TextureSample::Inputs::Coordinates.ToString(), TextureCoordinateShader->GetUniqueID());
+				bHasInput = true;
+				bIsPhong = true;
 			}
 
-			return TextureSampleShader;
-		}
+			// Shininess
+			if (Material->fbx.specular_exponent.has_value)
+			{
+				const FName InputName = Phong::Parameters::Shininess;
+				const float Factor = Material->fbx.specular_exponent.value_real;
+				const ufbx_material_map& Map = Material->fbx.specular_exponent;
 
-		// Expected normalized full path for existing file textures
-		const UInterchangeTexture2DNode* CreateTexture2DNode(const FString& TextureFilePath)
-		{
-			const FString TextureName = FPaths::GetBaseFilename(TextureFilePath);
-			const FString TextureNodeID = UInterchangeTextureNode::MakeNodeUid(TextureName);
+				if (!SetTextureFromMap(InputName, Map, Factor))
+				{
+					SetScalarFromMap(InputName, Map, Factor);
+				}
 
-			if (const UInterchangeTexture2DNode* TextureNode = Cast<const UInterchangeTexture2DNode>(NodeContainer.GetNode(TextureNodeID)))
+				bHasInput = true;
+				bIsPhong = true;
+			}
+
+			if (Material->fbx.normal_map.has_value)
 			{
-				return TextureNode;
+				if (SetTextureFromMap(Phong::Parameters::Normal, Material->fbx.normal_map, Material->fbx.normal_map.value_real))
+				{
+					bHasInput = true;
+				}
+				else
+				{
+					if (SetTextureFromMap(Phong::Parameters::Normal, Material->fbx.bump, Material->fbx.bump.value_real))
+					{
+						MaterialInstanceNode->AddStaticSwitchParameterValue(TEXT("bHasBump"), true);
+						bHasInput = true;
+					}
+				}
 			}
 
-			UInterchangeTexture2DNode* NewTextureNode = UInterchangeTexture2DNode::Create(&NodeContainer, TextureNodeID);
-			NewTextureNode->SetDisplayLabel(TextureName);
+			if (!bIsPhong)
+			{
+				ParentMaterial = TEXT("LambertSurfaceMaterial");
+			}
 
-			//All texture translator expect a file as the payload key
-			NewTextureNode->SetPayLoadKey(TextureFilePath);
+			// Opacity
+			// Connect only if transparency is either a texture or different from 0.f
+			if (Material->fbx.transparency_color.has_value || Material->fbx.transparency_factor.has_value)
+			{
+				FName InputName = Phong::Parameters::Opacity;
+				const ufbx_material_map& Map = Material->fbx.transparency_color;
+				const ufbx_material_map& FactorMap = Material->fbx.transparency_factor;
+				const float Factor = FactorMap.value_real;
 
-			return NewTextureNode;
-		}
+				// #ufbx_todo: note - 'transparency' texture in not inverted when setting as opacity, this is same behavior as was done for FBX SDK import
 
+				// The texture is hooked to the OpacityMask when transparency is with a texture
+				bool bIsMask = Map.value_int == 0;
 
-		void ConvertMaterialMap(FName InputName, const ufbx_material_map& MaterialMap)
-		{
-			if (MaterialMap.has_value)
-			{
-				if (!MaterialMap.texture || !MaterialMap.texture_enabled)
+				if (bIsMask)
 				{
-					const FString InputAttributeKey = UInterchangeShaderPortsAPI::MakeInputValueKey(InputName.ToString());
+					InputName = Phong::Parameters::OpacityMask;								
+				}
 
-					if (MaterialMap.value_components == 4)
-					{
-						// #ufbx_todo: converter utils
-						FLinearColor PropertyValue(Convert::ConvertVec4(MaterialMap.value_vec4));
-						ShaderGraphNode.AddLinearColorAttribute(InputAttributeKey, PropertyValue);
-					}
-					else if (MaterialMap.value_components == 3)
+				if (SetTextureFromMap(InputName, Map, Factor))
+				{
+					if (bIsMask)
 					{
-						// #ufbx_todo: converter utils
-						FLinearColor PropertyValue(Convert::ConvertVec3(MaterialMap.value_vec3));
-						ShaderGraphNode.AddLinearColorAttribute(InputAttributeKey, PropertyValue);
+						ParentMaterial = TEXT("PhongSurfaceMaskMaterial");
 					}
-					else if (MaterialMap.value_components == 1)
+					else
 					{
-						ShaderGraphNode.AddDoubleAttribute(InputAttributeKey, MaterialMap.value_real);
+						ParentMaterial = TEXT("PhongSurfaceTranslucentMaterial");
 					}
+
+					bHasInput = true;
 				}
 				else
 				{
-					if (const UInterchangeShaderNode* TextureSampleShader = CreateTextureSampler(InputName, MaterialMap.texture))
+					const float OpacityScalar = 1.f - Factor;
+					if (OpacityScalar < 1.f)
 					{
-						// #ufbx_todo: lerp etc
-						UInterchangeShaderNode* NodeToConnectTo = &ShaderGraphNode;
-						FString InputToConnectTo = InputName.ToString();
-
-						UInterchangeShaderPortsAPI::ConnectDefaultOuputToInput(NodeToConnectTo, InputToConnectTo, TextureSampleShader->GetUniqueID());
+						ParentMaterial = TEXT("PhongSurfaceTranslucentMaterial");
+						MaterialInstanceNode->AddScalarParameterValue(Phong::Parameters::Opacity.ToString(), OpacityScalar);
+						bHasInput = true;
 					}
 				}
 			}
-		};
 
-		FUfbxParser& Parser;
-		UInterchangeBaseNodeContainer& NodeContainer;
-		UInterchangeShaderGraphNode& ShaderGraphNode;
-	};
+			// If no valid property found, create a material anyway
+			TArray<FString> InputNames;
+			if (!bHasInput)
+			{
+				FLinearColor BaseColor;
+				BaseColor.R = 0.7f;
+				BaseColor.G = 0.7f;
+				BaseColor.B = 0.7f;
 
-	void FUfbxMaterial::AddMaterials(UInterchangeBaseNodeContainer& NodeContainer)
-	{
-		for (ufbx_material* Material : Parser.Scene->materials)
-		{
-			// #ufbx_todo: Uid in FFbxMaterial parser is made from sanitized name
-			// but two different materials  might happen to have same name AFTER sanitation.
-			// SO material name shoul dbe made unique for this case. E.g. by adding number after name string.
-			// This means that we need to memoize mapping of material to this generated name to assign proper names to meshes
-			// FString MaterialName = Parser.GetFbxHelper()->GetFbxObjectName(SurfaceMaterial);
-			// FString NodeUid = TEXT("\\Material\\") + MaterialName;
+				const FString InputValueKey = Phong::Parameters::DiffuseColor.ToString();
+				MaterialInstanceNode->AddVectorParameterValue(InputValueKey, BaseColor);
+			}
 
-			FString MaterialName = Parser.GetMaterialName(*Material).ToString();
-			
-			// InterchangeShaderGraphNode if not 'Node of a ShaderGraph' but rather '(Interchange)Node which is a ShaderGraph'
-			UInterchangeShaderGraphNode* ShaderGraphNode = UInterchangeShaderGraphNode::Create(&NodeContainer, MaterialName);
+			// #ufbx_todo: ProcessCustomAttributes(Parser, SurfaceMaterial, MaterialInstanceNode);
 
-			using namespace UE::Interchange::Materials;
-			FMaterialMapConverter MaterialMapConverter(Parser, NodeContainer, *ShaderGraphNode);
+			for (const ufbx_prop& Prop : Material->props.props)
+			{
+				if (Prop.flags & UFBX_PROP_FLAG_USER_DEFINED)
+				{
+					Parser.ConvertProperty(Prop, MaterialInstanceNode);
+				}
+			}
+
+			if (bHasUVTransform)
+			{
+				ParentMaterial += "WithUVTransform";
+			}
 
-			MaterialMapConverter.ConvertMaterialMap(PBRMR::Parameters::BaseColor, Material->pbr.base_color);
-			MaterialMapConverter.ConvertMaterialMap(PBRMR::Parameters::Metallic, Material->pbr.metalness);
-			MaterialMapConverter.ConvertMaterialMap(PBRMR::Parameters::Specular, Material->pbr.specular_factor);
-			MaterialMapConverter.ConvertMaterialMap(PBRMR::Parameters::Roughness, Material->pbr.roughness);
+			MaterialInstanceNode->SetCustomParent(FString(TEXT("/InterchangeAssets/Materials"))/ParentMaterial + TEXT(".") + ParentMaterial);
 		}
 	}
 }
diff --git a/Engine/Plugins/Interchange/Runtime/Source/Parsers/Fbx/Private/Ufbx/UfbxMesh.cpp b/Engine/Plugins/Interchange/Runtime/Source/Parsers/Fbx/Private/Ufbx/UfbxMesh.cpp
--- a/Engine/Plugins/Interchange/Runtime/Source/Parsers/Fbx/Private/Ufbx/UfbxMesh.cpp
+++ b/Engine/Plugins/Interchange/Runtime/Source/Parsers/Fbx/Private/Ufbx/UfbxMesh.cpp
@@ -247,7 +247,7 @@ namespace UE::Interchange::Private
 			for (ufbx_material* Material: Mesh->materials)
 			{
 				const FString MaterialName = Parser.GetMaterialName(*Material).ToString();
-				const FString MaterialUid = UInterchangeShaderGraphNode::MakeNodeUid(*MaterialName);
+				const FString MaterialUid = Parser.GetMaterialUid(MaterialName);
 				MeshNode->SetSlotMaterialDependencyUid(MaterialName, MaterialUid);
 			}
 
