[Interchange] Improves payload removal granularity, so that only Nanite assembly descriptions are retained past ImportAsset_Async

This fix should reinstate the improvements made in UE-314763 for the majority of skeletal mesh imports.

diff --git a/Engine/Plugins/Interchange/Runtime/Source/Import/Private/Mesh/InterchangeMeshHelper.cpp b/Engine/Plugins/Interchange/Runtime/Source/Import/Private/Mesh/InterchangeMeshHelper.cpp
--- a/Engine/Plugins/Interchange/Runtime/Source/Import/Private/Mesh/InterchangeMeshHelper.cpp
+++ b/Engine/Plugins/Interchange/Runtime/Source/Import/Private/Mesh/InterchangeMeshHelper.cpp
@@ -1807,21 +1807,14 @@ namespace UE::Interchange::Private::MeshHelper
 				}
 			}
 
-			// Report invalid instance indices and return false if we failed to add any nodes at all.
+			// Report invalid instance indices
 			if (NumInstancesWithInvalidPartIndex > 0)
 			{
 				UE_LOG(LogInterchangeImport, Warning
 					, TEXT("Nanite assembly description part index data for mesh '%s' contains one or more invalid indices.")
 					, *MeshFactoryNode->GetDisplayLabel()
 				);
-
-				if (NumInstancesWithInvalidPartIndex == NumInstances)
-				{
-					return;
-				}
 			}
-
-			return;
 		}
 
 		/** Enable nanite and finalize/apply the FNaniteAssemblyDataBuilder data to the mesh. If the finalize process 
diff --git a/Engine/Plugins/Interchange/Runtime/Source/Import/Private/Mesh/InterchangeSkeletalMeshFactory.cpp b/Engine/Plugins/Interchange/Runtime/Source/Import/Private/Mesh/InterchangeSkeletalMeshFactory.cpp
--- a/Engine/Plugins/Interchange/Runtime/Source/Import/Private/Mesh/InterchangeSkeletalMeshFactory.cpp
+++ b/Engine/Plugins/Interchange/Runtime/Source/Import/Private/Mesh/InterchangeSkeletalMeshFactory.cpp
@@ -861,25 +861,17 @@ namespace UE
 				return BoneIndexByBoneNameTable;
 			}
 
-			bool RemapNaniteAssemblyBoneIndices(UE::Interchange::FMeshPayloadData& MeshPayloadData
+			bool RemapNaniteAssemblyBoneIndices(
+				UE::Interchange::FNaniteAssemblyDescription& InOutDescription
+				, const TArray<FString>& JointNames
 				, const TMap<FString, int32>& NewBoneIndexByBoneNameMap
 			    , const FString& AssetName)
 			{
-				using FNaniteAssemblyDescription = UE::Interchange::FNaniteAssemblyDescription;
-
-				if (!MeshPayloadData.NaniteAssemblyDescription.IsSet())
-				{
-					return false;
-				}
-
-				FNaniteAssemblyDescription& Description = MeshPayloadData.NaniteAssemblyDescription.GetValue();
-
 				// Remap the bone influence bone indices to our actual skeleton.  Note that Nanite assembly 
 				// bone indices are uint32, while the indices held in the table are int32.
 
-				const TArray<FString>& JointNames = MeshPayloadData.JointNames;
 				int32 bNumRemapped = 0;
-				for (FNaniteAssemblyDescription::FBoneInfluence& BoneInfluence : Description.BoneInfluences)
+				for (FNaniteAssemblyDescription::FBoneInfluence& BoneInfluence : InOutDescription.BoneInfluences)
 				{
 					const uint32 OldBoneIndex = BoneInfluence.BoneIndex;
 					if (IntFitsIn<uint32, int32>(OldBoneIndex) && JointNames.IsValidIndex(static_cast<int32>(OldBoneIndex)))
@@ -893,7 +885,7 @@ namespace UE
 					}
 				}
 
-				if (bNumRemapped != Description.BoneInfluences.Num())
+				if (bNumRemapped != InOutDescription.BoneInfluences.Num())
 				{
 					UE_LOG(LogInterchangeImport, Warning
 						, TEXT("Failed to remap one or more Nanite Assembly Bone Influence's while processing skeletal mesh '%s'")
@@ -1723,11 +1715,21 @@ UInterchangeFactoryBase::FImportAssetResult UInterchangeSkeletalMeshFactory::Imp
 			}
 		}
 
-		// TODO: reinstate this by checking and retaining only the payloads that have Nanite assembly
-		// data. Also, even if we can't remove the payload entirely, we should be ok to call Empty()  
-		// on the mesh description to free up resources most of the resources.
-		// 
-		// // PayloadsPerLodIndex.Remove(LodIndex);
+		// Make sure to retain Nanite assembly descriptions (needed in SetupObject_GameThread).
+		if (FLodPayloads* Payloads = PayloadsPerLodIndex.Find(LodIndex))
+		{
+			for (TTuple<FInterchangeMeshPayLoadKey, TOptional<UE::Interchange::FMeshPayloadData>>& PayloadPair : Payloads->MeshPayloadPerKey)
+			{
+				TOptional<UE::Interchange::FMeshPayloadData>& PayloadData = PayloadPair.Value;
+				if (PayloadData.IsSet() && PayloadData->NaniteAssemblyDescription.IsSet())
+				{
+					UE::Interchange::FNaniteAssemblyDescription& Description = PayloadData->NaniteAssemblyDescription.GetValue();
+					NaniteAssemblyData.Emplace(MoveTemp(Description), MoveTemp(PayloadData->JointNames));
+				}
+			}
+		}
+
+		PayloadsPerLodIndex.Remove(LodIndex);
 
 		CurrentLodIndex++;
 	}
@@ -2244,60 +2246,26 @@ void UInterchangeSkeletalMeshFactory::SetupObject_GameThread(const FSetupObjectP
 			}
 
 #if WITH_EDITOR
-			// Build Nanite Assembly Data
-			if (SkeletalMeshFactoryNode && !ImportAssetObjectData.LodDatas.IsEmpty())
+			// Build Nanite Assembly data, if available
+			if (NaniteAssemblyData.Num() > 0)
 			{
 				if (USkeleton* SkeletonReference = ImportAssetObjectData.SkeletonReference)
 				{
-					auto ForEachLOD0MeshPayload = [this](TFunctionRef<void(UE::Interchange::FMeshPayloadData* MeshPayloadData)> IterationLambda)
-						{
-							constexpr int32 LOD0 = 0;
-							FImportAssetObjectLODData& ImportAssetObjectLODData = ImportAssetObjectData.LodDatas[LOD0];
-							if (UInterchangeSkeletalMeshFactory::FLodPayloads* LodPayloads = PayloadsPerLodIndex.Find(LOD0))
-							{
-								for (const UE::Interchange::FMeshNodeContext& MeshNodeContext : ImportAssetObjectLODData.MeshNodeContexts)
-								{
-									const FInterchangeMeshPayLoadKey Key = MeshNodeContext.GetTranslatorAndTransformPayloadKey();
-									if (TOptional<UE::Interchange::FMeshPayloadData>* MeshPayload = LodPayloads->MeshPayloadPerKey.Find(Key))
-									{
-										if (UE::Interchange::FMeshPayloadData* MeshPayloadData = MeshPayload->GetPtrOrNull())
-										{
-											IterationLambda(MeshPayloadData);
-										}
-									}
-								}	
-							}
-						};
-
-					// Two passes here to avoid Nanite assembly builder construction for the common case of zero nanite assemblies.
+					using namespace UE::Interchange::Private::MeshHelper;
 
-					int32 NumActiveNaniteAssemblyPayloads = 0;
-					ForEachLOD0MeshPayload([&NumActiveNaniteAssemblyPayloads](UE::Interchange::FMeshPayloadData* MeshPayloadData)
-						{
-							if (MeshPayloadData->NaniteAssemblyDescription.IsSet())
-							{
-								NumActiveNaniteAssemblyPayloads++;
-							}
-						});
-
-					if (NumActiveNaniteAssemblyPayloads > 0)
+					const TMap<FString, int32> NewBoneIndexByBoneNameTable = UE::Interchange::Private::GetSkeletonBoneIndexByBoneNameTable(*SkeletonReference);
+					FInterchangeNaniteAssemblyBuilder Builder = FInterchangeNaniteAssemblyBuilder::Create(SkeletalMesh, SkeletalMeshFactoryNode, Arguments.NodeContainer);
+			
+					for (FNaniteAssemblyData& Data : NaniteAssemblyData)
 					{
-						using namespace UE::Interchange::Private::MeshHelper;
-
-						const TMap<FString, int32> NewBoneIndexByBoneNameTable = UE::Interchange::Private::GetSkeletonBoneIndexByBoneNameTable(*SkeletonReference);
-						FInterchangeNaniteAssemblyBuilder Builder = FInterchangeNaniteAssemblyBuilder::Create(SkeletalMesh, SkeletalMeshFactoryNode, Arguments.NodeContainer);
-
-						ForEachLOD0MeshPayload([&Builder, &NewBoneIndexByBoneNameTable, &Arguments](UE::Interchange::FMeshPayloadData* MeshPayloadData)
-							{
-								// The bone influence bone indices data in the payload likely needs remapping since translation may have added extra bones and/or 
-								// changed the indices. Note the remapping here may as well be done in-place since the current (invalid) indices are of little use.
-								if (UE::Interchange::Private::RemapNaniteAssemblyBoneIndices(*MeshPayloadData, NewBoneIndexByBoneNameTable, Arguments.NodeUniqueID))
-								{
-									Builder.AddDescription(MeshPayloadData->NaniteAssemblyDescription);
-								}
-							});
+						if (UE::Interchange::Private::RemapNaniteAssemblyBoneIndices(Data.Description, Data.JointNames, NewBoneIndexByBoneNameTable, Arguments.NodeUniqueID))
+						{
+							Builder.AddDescription(Data.Description);
+						}
 					}
 				}
+
+				NaniteAssemblyData.Empty();
 			}
 
 			//Re-apply the alternate skinning data
@@ -2464,9 +2432,6 @@ void UInterchangeSkeletalMeshFactory::SetupObject_GameThread(const FSetupObjectP
 			}
 		}
 	}
-
-	// Payloads have been processed, so remove
-	PayloadsPerLodIndex.Empty();
 }
 
 void UInterchangeSkeletalMeshFactory::BuildObject_GameThread(const FSetupObjectParams& Arguments, bool& OutPostEditchangeCalled)
diff --git a/Engine/Plugins/Interchange/Runtime/Source/Import/Public/Mesh/InterchangeSkeletalMeshFactory.h b/Engine/Plugins/Interchange/Runtime/Source/Import/Public/Mesh/InterchangeSkeletalMeshFactory.h
--- a/Engine/Plugins/Interchange/Runtime/Source/Import/Public/Mesh/InterchangeSkeletalMeshFactory.h
+++ b/Engine/Plugins/Interchange/Runtime/Source/Import/Public/Mesh/InterchangeSkeletalMeshFactory.h
@@ -109,12 +109,20 @@ public:
 		TMap<FInterchangeMeshPayLoadKey, TOptional<UE::Interchange::FMeshPayloadData>> MorphPayloadPerKey;
 	};
 
+	struct FNaniteAssemblyData
+	{
+		UE::Interchange::FNaniteAssemblyDescription Description;
+		TArray<FString> JointNames;
+	};
+
 private:
 	TUniquePtr<FSkinnedMeshComponentRecreateRenderStateContext> ScopedReimportUtility;
 
 	TMap<int32, FLodPayloads> PayloadsPerLodIndex;
 
 	FImportAssetObjectData ImportAssetObjectData;
+
+	TArray<FNaniteAssemblyData> NaniteAssemblyData;
 };
 
 
