Interchange FBX:
- Add parent material to reproduce the same rendering as Legacy FBX using material instances
- Add the support for setting the BlendMode on Material Instances during import
- Fixed the opacity in the shader graph

#virtualized

diff --git a/Engine/Plugins/Interchange/Runtime/Source/FactoryNodes/Private/InterchangeMaterialFactoryNode.cpp b/Engine/Plugins/Interchange/Runtime/Source/FactoryNodes/Private/InterchangeMaterialFactoryNode.cpp
--- a/Engine/Plugins/Interchange/Runtime/Source/FactoryNodes/Private/InterchangeMaterialFactoryNode.cpp
+++ b/Engine/Plugins/Interchange/Runtime/Source/FactoryNodes/Private/InterchangeMaterialFactoryNode.cpp
@@ -486,6 +486,16 @@ bool UInterchangeMaterialInstanceFactoryNode::SetCustomParent(const FString& Att
 	IMPLEMENT_NODE_ATTRIBUTE_SETTER_NODELEGATE(Parent, FString);
 }
 
+bool UInterchangeMaterialInstanceFactoryNode::GetCustomBlendMode(TEnumAsByte<EBlendMode>& AttributeValue) const
+{
+	IMPLEMENT_NODE_ATTRIBUTE_GETTER(BlendMode, TEnumAsByte<EBlendMode>);
+}
+
+bool UInterchangeMaterialInstanceFactoryNode::SetCustomBlendMode(TEnumAsByte<EBlendMode> AttributeValue, bool bAddApplyDelegate)
+{
+	IMPLEMENT_NODE_ATTRIBUTE_SETTER_NODELEGATE(BlendMode, TEnumAsByte<EBlendMode>);
+}
+
 FString UInterchangeMaterialReferenceFactoryNode::GetTypeName() const
 {
 	const FString TypeName = TEXT("MaterialReferenceFactoryNode");
diff --git a/Engine/Plugins/Interchange/Runtime/Source/FactoryNodes/Public/InterchangeMaterialFactoryNode.h b/Engine/Plugins/Interchange/Runtime/Source/FactoryNodes/Public/InterchangeMaterialFactoryNode.h
--- a/Engine/Plugins/Interchange/Runtime/Source/FactoryNodes/Public/InterchangeMaterialFactoryNode.h
+++ b/Engine/Plugins/Interchange/Runtime/Source/FactoryNodes/Public/InterchangeMaterialFactoryNode.h
@@ -328,9 +328,16 @@ public:
 	UFUNCTION(BlueprintCallable, Category = "Interchange | Node | MaterialInstanceFactory")
 	INTERCHANGEFACTORYNODES_API bool SetCustomParent(const FString& AttributeValue);
 
+	UFUNCTION(BlueprintCallable, Category = "Interchange | Node | MaterialInstanceFactory")
+	INTERCHANGEFACTORYNODES_API bool GetCustomBlendMode(TEnumAsByte<EBlendMode>& AttributeValue) const;
+
+	UFUNCTION(BlueprintCallable, Category = "Interchange | Node | MaterialInstanceFactory")
+	INTERCHANGEFACTORYNODES_API bool SetCustomBlendMode(TEnumAsByte<EBlendMode> AttributeValue, bool bAddApplyDelegate = true);
+
 private:
-	const UE::Interchange::FAttributeKey Macro_CustomInstanceClassNameKey = UE::Interchange::FAttributeKey(TEXT("InstanceClassName"));
-	const UE::Interchange::FAttributeKey Macro_CustomParentKey = UE::Interchange::FAttributeKey(TEXT("Parent"));
+	IMPLEMENT_NODE_ATTRIBUTE_KEY(InstanceClassName)
+	IMPLEMENT_NODE_ATTRIBUTE_KEY(Parent)
+	IMPLEMENT_NODE_ATTRIBUTE_KEY(BlendMode)
 
 };
 
diff --git a/Engine/Plugins/Interchange/Runtime/Source/Import/Private/Material/InterchangeMaterialFactory.cpp b/Engine/Plugins/Interchange/Runtime/Source/Import/Private/Material/InterchangeMaterialFactory.cpp
--- a/Engine/Plugins/Interchange/Runtime/Source/Import/Private/Material/InterchangeMaterialFactory.cpp
+++ b/Engine/Plugins/Interchange/Runtime/Source/Import/Private/Material/InterchangeMaterialFactory.cpp
@@ -798,6 +798,16 @@ namespace UE::Interchange::MaterialFactory::Internal
 			}
 		}
 	}
+
+	void OverrideBlendMode(UMaterialInstance& MaterialInstance, const UInterchangeMaterialInstanceFactoryNode& FactoryNode)
+	{
+		if (TEnumAsByte<EBlendMode> BlendMode; FactoryNode.GetCustomBlendMode(BlendMode))
+		{
+			MaterialInstance.BasePropertyOverrides.bOverride_BlendMode = true;
+			MaterialInstance.BasePropertyOverrides.BlendMode = BlendMode;
+			MaterialInstance.BlendMode = BlendMode;
+		}
+	}
 }
 
 namespace UE::Interchange::Materials::HashUtils
@@ -1994,6 +2004,8 @@ void UInterchangeMaterialFactory::SetupMaterialInstance(UMaterialInstance& Mater
 			}
 		}
 	}
+
+	OverrideBlendMode(MaterialInstance, FactoryNode);
 }
 
 void UInterchangeMaterialFactory::SetupReimportedMaterialInstance(UMaterialInstance& MaterialInstance, const UInterchangeBaseNodeContainer& NodeContainer, const UInterchangeMaterialInstanceFactoryNode& FactoryNode, const UInterchangeMaterialInstanceFactoryNode& PreviousFactoryNode)
@@ -2097,6 +2109,8 @@ void UInterchangeMaterialFactory::SetupReimportedMaterialInstance(UMaterialInsta
 			break;
 		}
 	}
+
+	OverrideBlendMode(MaterialInstance, FactoryNode);
 }
 
 UClass* UInterchangeMaterialFunctionFactory::GetFactoryClass() const
diff --git a/Engine/Plugins/Interchange/Runtime/Source/Nodes/Private/InterchangeMaterialInstanceNode.cpp b/Engine/Plugins/Interchange/Runtime/Source/Nodes/Private/InterchangeMaterialInstanceNode.cpp
--- a/Engine/Plugins/Interchange/Runtime/Source/Nodes/Private/InterchangeMaterialInstanceNode.cpp
+++ b/Engine/Plugins/Interchange/Runtime/Source/Nodes/Private/InterchangeMaterialInstanceNode.cpp
@@ -88,3 +88,13 @@ bool UInterchangeMaterialInstanceNode::GetStaticSwitchParameterValue(const FStri
 {
 	return GetBooleanAttribute(UInterchangeShaderPortsAPI::MakeInputValueKey(ParameterName), AttributeValue);
 }
+
+bool UInterchangeMaterialInstanceNode::GetCustomBlendMode(int& AttributeValue) const
+{
+	IMPLEMENT_NODE_ATTRIBUTE_GETTER(BlendMode, int);
+}
+
+bool UInterchangeMaterialInstanceNode::SetCustomBlendMode(int AttributeValue)
+{
+	IMPLEMENT_NODE_ATTRIBUTE_SETTER_NODELEGATE(BlendMode, int);
+}
diff --git a/Engine/Plugins/Interchange/Runtime/Source/Nodes/Public/InterchangeMaterialInstanceNode.h b/Engine/Plugins/Interchange/Runtime/Source/Nodes/Public/InterchangeMaterialInstanceNode.h
--- a/Engine/Plugins/Interchange/Runtime/Source/Nodes/Public/InterchangeMaterialInstanceNode.h
+++ b/Engine/Plugins/Interchange/Runtime/Source/Nodes/Public/InterchangeMaterialInstanceNode.h
@@ -64,9 +64,16 @@ public:
 	UFUNCTION(BlueprintCallable, Category = "Interchange | Node | Material Instance")
 	UE_API bool GetStaticSwitchParameterValue(const FString& ParameterName, bool& AttributeValue) const;
 
+	UFUNCTION(BlueprintCallable, Category = "Interchange | Node | Material")
+	UE_API bool GetCustomBlendMode(int& AttributeValue) const;
+
+	UFUNCTION(BlueprintCallable, Category = "Interchange | Node | Material")
+	UE_API bool SetCustomBlendMode(int AttributeValue);
+
 private:
 
 	IMPLEMENT_NODE_ATTRIBUTE_KEY(Parent);
+	IMPLEMENT_NODE_ATTRIBUTE_KEY(BlendMode);
 
 };
 
diff --git a/Engine/Plugins/Interchange/Runtime/Source/Parsers/Fbx/Private/FbxMaterial.cpp b/Engine/Plugins/Interchange/Runtime/Source/Parsers/Fbx/Private/FbxMaterial.cpp
--- a/Engine/Plugins/Interchange/Runtime/Source/Parsers/Fbx/Private/FbxMaterial.cpp
+++ b/Engine/Plugins/Interchange/Runtime/Source/Parsers/Fbx/Private/FbxMaterial.cpp
@@ -14,6 +14,7 @@
 #include "InterchangeTextureNode.h"
 #include "Misc/Paths.h"
 #include "Nodes/InterchangeBaseNodeContainer.h"
+#include "Engine/EngineTypes.h"
 
 #define LOCTEXT_NAMESPACE "InterchangeFbxMaterial"
 
@@ -26,6 +27,7 @@ namespace UE
 			UInterchangeMaterialInstanceNode* FFbxMaterial::CreateMaterialInstanceNode(UInterchangeBaseNodeContainer& NodeContainer, const FString& NodeUID, const FString& NodeName)
 			{
 				UInterchangeMaterialInstanceNode* MaterialInstanceNode = NewObject<UInterchangeMaterialInstanceNode>(&NodeContainer);
+				MaterialInstanceNode->SetCustomParent(TEXT("/InterchangeAssets/Materials/FBXLegacyPhongSurfaceMaterial.FBXLegacyPhongSurfaceMaterial"));
 				NodeContainer.SetupNode(MaterialInstanceNode, NodeUID, NodeName, EInterchangeNodeContainerType::TranslatedAsset);
 
 				return MaterialInstanceNode;
@@ -129,7 +131,7 @@ namespace UE
 					if (DataType == eFbxDouble || DataType == eFbxFloat || DataType == eFbxInt)
 					{
 						const float PropertyValue = Property.Get<float>() * Factor;
-						MaterialInstanceNode->AddFloatAttribute(InputAttributeKey, bInverse ? 1.f - PropertyValue : PropertyValue);
+						MaterialInstanceNode->AddScalarParameterValue(InputAttributeKey, bInverse ? 1.f - PropertyValue : PropertyValue);
 					}
 					else if (DataType == eFbxDouble3 || DataType == eFbxDouble4)
 					{
@@ -239,9 +241,6 @@ namespace UE
 
 				bool bHasInput = false;
 
-				// By default we create a Phong material, depending later of the different input hooked we'll change it
-				FString ParentMaterial{ TEXT("/InterchangeAssets/Materials/PhongSurfaceMaterial.PhongSurfaceMaterial") };
-
 				// Diffuse
 				{
 					const float Factor = GetFactor(FbxSurfaceMaterial::sDiffuseFactor);
@@ -296,16 +295,25 @@ namespace UE
 
 						if (MaterialProperty.GetSrcObjectCount<FbxFileTexture>() > 0)
 						{
-							ParentMaterial = TEXT("/InterchangeAssets/Materials/PhongSurfaceTranslucentMaterial.PhongSurfaceTranslucentMaterial");
 							TVariant<FLinearColor, float> DefaultValue;
 							DefaultValue.Set<float>(0.f); // Opaque
 							FName InputName = Phong::Parameters::Opacity;
+							FString HasOpacity{ TEXT("bHasOpacity") };
+							EBlendMode BlendMode = EBlendMode::BLEND_Translucent;
+
 							// The texture is hooked to the OpacityMask when transparency is with a texture
 							if (MaterialProperty.Get<FbxDouble>() == 0.)
 							{
-								InputName = Phong::Parameters::OpacityMask;								
-								ParentMaterial = TEXT("/InterchangeAssets/Materials/PhongSurfaceMaskMaterial.PhongSurfaceMaskMaterial");
+								InputName = Phong::Parameters::OpacityMask;
+								BlendMode = EBlendMode::BLEND_Masked;
+								MaterialInstanceNode->AddStaticSwitchParameterValue(TEXT("bHasOpacityMaskMap"), true);
+								HasOpacity += TEXT("Mask");
 							}
+
+							HasOpacity += TEXT("Map");
+							MaterialInstanceNode->AddStaticSwitchParameterValue(HasOpacity, true);
+							MaterialInstanceNode->SetCustomBlendMode(BlendMode);
+							
 							bHasInput |= ConvertPropertyToMaterialInstanceNode(NodeContainer, MaterialInstanceNode, MaterialProperty, Factor, InputName, DefaultValue, true);
 						}
 						else
@@ -313,7 +321,7 @@ namespace UE
 							const float OpacityScalar = 1.f - Factor;
 							if (OpacityScalar < 1.f)
 							{
-								ParentMaterial = TEXT("/InterchangeAssets/Materials/PhongSurfaceTranslucentMaterial.PhongSurfaceTranslucentMaterial");
+								MaterialInstanceNode->SetCustomBlendMode(EBlendMode::BLEND_Translucent);
 								MaterialInstanceNode->AddScalarParameterValue(Phong::Parameters::Opacity.ToString(), OpacityScalar);
 							}
 						}
@@ -334,19 +342,22 @@ namespace UE
 
 				//Shininess
 				{
-					const float Factor = GetFactor(FbxSurfaceMaterial::sShininess);
-					TVariant<FLinearColor, float> DefaultValue;
-					DefaultValue.Set<float>(20.f);
-					bIsPhong |= ConnectInput(Phong::Parameters::Shininess, FbxSurfaceMaterial::sShininess, Factor, DefaultValue, false);
-					bHasInput |= bIsPhong;
-				}
+					FBXSDK_NAMESPACE::FbxProperty ShininessProperty = SurfaceMaterial->FindProperty(FbxSurfaceMaterial::sShininess);
+					if (ShininessProperty.IsValid())
+					{
+						TVariant<FLinearColor, float> DefaultValue;
+						DefaultValue.Set<float>(20.f);
+						if (ShininessProperty.GetSrcObjectCount<FBXSDK_NAMESPACE::FbxTexture>() > 0)
+						{
+							MaterialInstanceNode->AddStaticSwitchParameterValue(TEXT("bHasShininessMap"), true);
+						}
 
-				if (!bIsPhong)
-				{
-					ParentMaterial = TEXT("/InterchangeAssets/Materials/LambertSurfaceMaterial.LambertSurfaceMaterial");
+						bIsPhong |= ConnectInput(Phong::Parameters::Shininess, FbxSurfaceMaterial::sShininess, 1.f, DefaultValue, false);
+						bHasInput |= bIsPhong;
+					}
 				}
 
-				MaterialInstanceNode->SetCustomParent(ParentMaterial);
+				MaterialInstanceNode->AddStaticSwitchParameterValue(TEXT("bIsPhong"), bIsPhong);
 
 				// If no valid property found, create a material anyway
 				TArray<FString> InputNames;
diff --git a/Engine/Plugins/Interchange/Runtime/Source/Pipelines/Private/InterchangeGenericMaterialPipeline.cpp b/Engine/Plugins/Interchange/Runtime/Source/Pipelines/Private/InterchangeGenericMaterialPipeline.cpp
--- a/Engine/Plugins/Interchange/Runtime/Source/Pipelines/Private/InterchangeGenericMaterialPipeline.cpp
+++ b/Engine/Plugins/Interchange/Runtime/Source/Pipelines/Private/InterchangeGenericMaterialPipeline.cpp
@@ -1012,6 +1012,11 @@ void UInterchangeGenericMaterialPipeline::ExecutePipeline(UInterchangeBaseNodeCo
 		const UClass* MaterialClass = (EditorUtilities && EditorUtilities->IsRuntimeOrPIE()) ? UMaterialInstanceDynamic::StaticClass() : UMaterialInstanceConstant::StaticClass();
 		MaterialFactoryNode->SetCustomInstanceClassName(MaterialClass->GetPathName());
 
+		if (EBlendMode BlendMode; MaterialNode->GetCustomBlendMode(reinterpret_cast<int&>(BlendMode)))
+		{
+			MaterialFactoryNode->SetCustomBlendMode(BlendMode);
+		}
+
 		UE::Interchange::InterchangeGenericMaterialPipeline::Private::AddSpecularProfileToFactoryNode(MaterialNode, MaterialFactoryNode, BaseNodeContainer);
 
 		TArray<FString> Inputs;
