Fixed issue with object's properties being overwritten on reimport
- Root cause: factory nodes with custom attributes with custom delegates were duplicated without duplicating the custom delegates.
There two reasons for this:
1 - Some factory nodes did not provide a way to duplicate their custom delegates. DecalActor and StaticMesh factory nodes were the faulty ones.
Solution: As already done by the other factory nodes, override their CopyWithObject method to perform the copy of the custom delegates
2 - Code copy/paste from factory to factory was creating a factory node meant to represent the object to reimport and using it to apply the reimport strategy. However, the custom delegates were not part of the duplication. Hence the failure to honor the reimport strategy to preserve properties modified after import.
Solution: Replace the faulty code with a call to the UInterchangeFactoryBaseNode::DuplicateWithObject method which call the CopyWithObject method on the target factory node. The CopyWithObject method was made virtual on purpose. To allow factory nodes to complete the duplication if need be.

While at it...
Added a missing call to SetNodeClassFromClassAttribute in UInterchangeMeshFactoryNode::CopyWithObject. Otherwise the call to GetObjectClass was asserting since the target class was not set.
Moved some misplaced 'return true' in InterchangeFactoryBaseNode.h because of which the method would return true while the operation was not really successful..

#rnx

diff --git a/Engine/Plugins/Interchange/Runtime/Source/FactoryNodes/Private/InterchangeDecalActorFactoryNode.cpp b/Engine/Plugins/Interchange/Runtime/Source/FactoryNodes/Private/InterchangeDecalActorFactoryNode.cpp
--- a/Engine/Plugins/Interchange/Runtime/Source/FactoryNodes/Private/InterchangeDecalActorFactoryNode.cpp
+++ b/Engine/Plugins/Interchange/Runtime/Source/FactoryNodes/Private/InterchangeDecalActorFactoryNode.cpp
@@ -18,6 +18,18 @@ UClass* UInterchangeDecalActorFactoryNode::GetObjectClass() const
 #endif
 }
 
+void UInterchangeDecalActorFactoryNode::CopyWithObject(const UInterchangeFactoryBaseNode* SourceNode, UObject* Object)
+{
+	Super::CopyWithObject(SourceNode, Object);
+
+	if (const UInterchangeDecalActorFactoryNode* SourceFactoryNode = Cast<UInterchangeDecalActorFactoryNode>(SourceNode))
+	{
+		UClass* Class = UDecalComponent::StaticClass();
+		COPY_NODE_DELEGATES_WITH_CUSTOM_DELEGATE(SourceFactoryNode, UInterchangeDecalActorFactoryNode, SortOrder, int32, Class)
+		COPY_NODE_DELEGATES_WITH_CUSTOM_DELEGATE(SourceFactoryNode, UInterchangeDecalActorFactoryNode, DecalSize, FVector, Class)
+	}
+}
+
 bool UInterchangeDecalActorFactoryNode::GetCustomSortOrder(int32& AttributeValue) const
 {
 	IMPLEMENT_NODE_ATTRIBUTE_GETTER(SortOrder, int32);
diff --git a/Engine/Plugins/Interchange/Runtime/Source/FactoryNodes/Private/InterchangeMeshFactoryNode.cpp b/Engine/Plugins/Interchange/Runtime/Source/FactoryNodes/Private/InterchangeMeshFactoryNode.cpp
--- a/Engine/Plugins/Interchange/Runtime/Source/FactoryNodes/Private/InterchangeMeshFactoryNode.cpp
+++ b/Engine/Plugins/Interchange/Runtime/Source/FactoryNodes/Private/InterchangeMeshFactoryNode.cpp
@@ -536,6 +536,11 @@ void UInterchangeMeshFactoryNode::CopyWithObject(const UInterchangeFactoryBaseNo
 
 	if (const UInterchangeMeshFactoryNode* MeshFactoryNode = Cast<UInterchangeMeshFactoryNode>(SourceNode))
 	{
+#if WITH_ENGINE
+		// Make sure the class is properly set when we compile with engine, this will set the bIsNodeClassInitialized to true.
+		SetNodeClassFromClassAttribute();
+#endif
+
 		UClass* Class = GetObjectClass();
 		COPY_NODE_DELEGATES_WITH_CUSTOM_DELEGATE(MeshFactoryNode, UInterchangeMeshFactoryNode, RecomputeNormals, bool, Class)
 		COPY_NODE_DELEGATES_WITH_CUSTOM_DELEGATE(MeshFactoryNode, UInterchangeMeshFactoryNode, RecomputeTangents, bool, Class)
diff --git a/Engine/Plugins/Interchange/Runtime/Source/FactoryNodes/Private/InterchangeStaticMeshFactoryNode.cpp b/Engine/Plugins/Interchange/Runtime/Source/FactoryNodes/Private/InterchangeStaticMeshFactoryNode.cpp
--- a/Engine/Plugins/Interchange/Runtime/Source/FactoryNodes/Private/InterchangeStaticMeshFactoryNode.cpp
+++ b/Engine/Plugins/Interchange/Runtime/Source/FactoryNodes/Private/InterchangeStaticMeshFactoryNode.cpp
@@ -104,6 +104,30 @@ UClass* UInterchangeStaticMeshFactoryNode::GetObjectClass() const
 #endif
 }
 
+void UInterchangeStaticMeshFactoryNode::CopyWithObject(const UInterchangeFactoryBaseNode* SourceNode, UObject* Object)
+{
+	Super::CopyWithObject(SourceNode, Object);
+
+	if (const UInterchangeStaticMeshFactoryNode* StaticMeshFactoryNode = Cast<UInterchangeStaticMeshFactoryNode>(SourceNode))
+	{
+		UClass* Class = GetObjectClass();
+#if WITH_EDITORONLY_DATA
+		COPY_NODE_DELEGATES_WITH_CUSTOM_DELEGATE(StaticMeshFactoryNode, UInterchangeStaticMeshFactoryNode, BuildNanite, bool, Class)
+#endif
+		COPY_NODE_DELEGATES_WITH_CUSTOM_DELEGATE(StaticMeshFactoryNode, UInterchangeStaticMeshFactoryNode, BuildReversedIndexBuffer, bool, Class)
+		COPY_NODE_DELEGATES_WITH_CUSTOM_DELEGATE(StaticMeshFactoryNode, UInterchangeStaticMeshFactoryNode, GenerateLightmapUVs, bool, Class)
+		COPY_NODE_DELEGATES_WITH_CUSTOM_DELEGATE(StaticMeshFactoryNode, UInterchangeStaticMeshFactoryNode, GenerateDistanceFieldAsIfTwoSided, bool, Class)
+		COPY_NODE_DELEGATES_WITH_CUSTOM_DELEGATE(StaticMeshFactoryNode, UInterchangeStaticMeshFactoryNode, SupportFaceRemap, bool, Class)
+		COPY_NODE_DELEGATES_WITH_CUSTOM_DELEGATE(StaticMeshFactoryNode, UInterchangeStaticMeshFactoryNode, MinLightmapResolution, int32, Class)
+		COPY_NODE_DELEGATES_WITH_CUSTOM_DELEGATE(StaticMeshFactoryNode, UInterchangeStaticMeshFactoryNode, SrcLightmapIndex, int32, Class)
+		COPY_NODE_DELEGATES_WITH_CUSTOM_DELEGATE(StaticMeshFactoryNode, UInterchangeStaticMeshFactoryNode, DstLightmapIndex, int32, Class)
+		COPY_NODE_DELEGATES_WITH_CUSTOM_DELEGATE(StaticMeshFactoryNode, UInterchangeStaticMeshFactoryNode, BuildScale3D, FVector, Class)
+		COPY_NODE_DELEGATES_WITH_CUSTOM_DELEGATE(StaticMeshFactoryNode, UInterchangeStaticMeshFactoryNode, DistanceFieldResolutionScale, float, Class)
+		COPY_NODE_DELEGATES_WITH_CUSTOM_DELEGATE(StaticMeshFactoryNode, UInterchangeStaticMeshFactoryNode, DistanceFieldReplacementMesh, FSoftObjectPath, Class)
+		COPY_NODE_DELEGATES_WITH_CUSTOM_DELEGATE(StaticMeshFactoryNode, UInterchangeStaticMeshFactoryNode, MaxLumenMeshCards, int32, Class)
+	}
+}
+
 #if WITH_EDITOR
 
 FString UInterchangeStaticMeshFactoryNode::GetKeyDisplayName(const UE::Interchange::FAttributeKey& NodeAttributeKey) const
diff --git a/Engine/Plugins/Interchange/Runtime/Source/FactoryNodes/Public/InterchangeDecalActorFactoryNode.h b/Engine/Plugins/Interchange/Runtime/Source/FactoryNodes/Public/InterchangeDecalActorFactoryNode.h
--- a/Engine/Plugins/Interchange/Runtime/Source/FactoryNodes/Public/InterchangeDecalActorFactoryNode.h
+++ b/Engine/Plugins/Interchange/Runtime/Source/FactoryNodes/Public/InterchangeDecalActorFactoryNode.h
@@ -16,6 +16,8 @@ class UInterchangeDecalActorFactoryNode : public UInterchangeActorFactoryNode
 public:
 	UE_API UClass* GetObjectClass()const override;
 
+	UE_API virtual void CopyWithObject(const UInterchangeFactoryBaseNode* SourceNode, UObject* Object) override;
+
 	UFUNCTION(BlueprintCallable, Category = "Interchange | Node | Decal")
 	UE_API bool GetCustomSortOrder(int32& AttributeValue) const;
 	
diff --git a/Engine/Plugins/Interchange/Runtime/Source/FactoryNodes/Public/InterchangeStaticMeshFactoryNode.h b/Engine/Plugins/Interchange/Runtime/Source/FactoryNodes/Public/InterchangeStaticMeshFactoryNode.h
--- a/Engine/Plugins/Interchange/Runtime/Source/FactoryNodes/Public/InterchangeStaticMeshFactoryNode.h
+++ b/Engine/Plugins/Interchange/Runtime/Source/FactoryNodes/Public/InterchangeStaticMeshFactoryNode.h
@@ -12,6 +12,7 @@
 
 #include "InterchangeStaticMeshFactoryNode.generated.h"
 
+#define UE_API INTERCHANGEFACTORYNODES_API
 
 namespace UE
 {
@@ -19,8 +20,8 @@ namespace UE
 	{
 		struct FStaticMeshNodeStaticData : public FBaseNodeStaticData
 		{
-			static INTERCHANGEFACTORYNODES_API const FAttributeKey& GetLODScreenSizeBaseKey();
-			static INTERCHANGEFACTORYNODES_API const FAttributeKey& GetSocketUidsBaseKey();
+			static UE_API const FAttributeKey& GetLODScreenSizeBaseKey();
+			static UE_API const FAttributeKey& GetSocketUidsBaseKey();
 		};
 	} // namespace Interchange
 } // namespace UE
@@ -32,7 +33,7 @@ class UInterchangeStaticMeshFactoryNode : public UInterchangeMeshFactoryNode
 	GENERATED_BODY()
 
 public:
-	INTERCHANGEFACTORYNODES_API UInterchangeStaticMeshFactoryNode();
+	UE_API UInterchangeStaticMeshFactoryNode();
 
 	/**
 	 * Initialize node data. Also adds it to NodeContainer.
@@ -41,159 +42,161 @@ public:
 	 * @param InAssetClass - The class the StaticMesh factory will create for this node.
 	 */
 	UFUNCTION(BlueprintCallable, Category = "Interchange | Node | StaticMesh")
-	INTERCHANGEFACTORYNODES_API void InitializeStaticMeshNode(const FString& UniqueID, const FString& DisplayLabel, const FString& InAssetClass, UInterchangeBaseNodeContainer* NodeContainer);
+	UE_API void InitializeStaticMeshNode(const FString& UniqueID, const FString& DisplayLabel, const FString& InAssetClass, UInterchangeBaseNodeContainer* NodeContainer);
 
 	/**
 	 * Return the node type name of the class. This is used when reporting errors.
 	 */
-	INTERCHANGEFACTORYNODES_API virtual FString GetTypeName() const override;
+	UE_API virtual FString GetTypeName() const override;
 
 	/** Get the class this node creates. */
-	INTERCHANGEFACTORYNODES_API virtual class UClass* GetObjectClass() const override;
+	UE_API virtual class UClass* GetObjectClass() const override;
+
+	UE_API virtual void CopyWithObject(const UInterchangeFactoryBaseNode* SourceNode, UObject* Object) override;
 
 #if WITH_EDITOR
 
-	INTERCHANGEFACTORYNODES_API virtual FString GetKeyDisplayName(const UE::Interchange::FAttributeKey& NodeAttributeKey) const override;
-	INTERCHANGEFACTORYNODES_API virtual FString GetAttributeCategory(const UE::Interchange::FAttributeKey& NodeAttributeKey) const override;
+	UE_API virtual FString GetKeyDisplayName(const UE::Interchange::FAttributeKey& NodeAttributeKey) const override;
+	UE_API virtual FString GetAttributeCategory(const UE::Interchange::FAttributeKey& NodeAttributeKey) const override;
 
 #endif //WITH_EDITOR
 
 public:
 	/** Get whether the static mesh factory should auto compute LOD Screen Sizes. Return false if the attribute was not set. */
 	UFUNCTION(BlueprintCallable, Category = "Interchange | Node | StaticMesh")
-	INTERCHANGEFACTORYNODES_API bool GetCustomAutoComputeLODScreenSizes(bool& AttributeValue) const;
+	UE_API bool GetCustomAutoComputeLODScreenSizes(bool& AttributeValue) const;
 
 	/** Set whether the static mesh factory should auto compute LOD Screen Sizes. Return false if the attribute was not set. */
 	UFUNCTION(BlueprintCallable, Category = "Interchange | Node | StaticMesh")
-	INTERCHANGEFACTORYNODES_API bool SetCustomAutoComputeLODScreenSizes(const bool& AttributeValue);
+	UE_API bool SetCustomAutoComputeLODScreenSizes(const bool& AttributeValue);
 
 	/** Returns the number of LOD Screen Sizes the static mesh has.*/
 	UFUNCTION(BlueprintCallable, Category = "Interchange | Node | StaticMesh")
-	INTERCHANGEFACTORYNODES_API int32 GetLODScreenSizeCount() const;
+	UE_API int32 GetLODScreenSizeCount() const;
 
 	/** Returns All the LOD Screen Sizes set for the static mesh.*/
 	UFUNCTION(BlueprintCallable, Category = "Interchange | Node | StaticMesh")
-	INTERCHANGEFACTORYNODES_API void GetLODScreenSizes(TArray<float>& OutLODScreenSizes) const;
+	UE_API void GetLODScreenSizes(TArray<float>& OutLODScreenSizes) const;
 
 	/** Sets the LOD Screen Sizes for the static mesh.*/
 	UFUNCTION(BlueprintCallable, Category = "Interchange | Node | StaticMesh")
-	INTERCHANGEFACTORYNODES_API bool SetLODScreenSizes(const TArray<float>& InLODScreenSizes);
+	UE_API bool SetLODScreenSizes(const TArray<float>& InLODScreenSizes);
 
 	/** Get whether the static mesh factory should set the Nanite build setting. Return false if the attribute was not set. */
 	UFUNCTION(BlueprintCallable, Category = "Interchange | Node | StaticMesh")
-	INTERCHANGEFACTORYNODES_API bool GetCustomBuildNanite(bool& AttributeValue) const;
+	UE_API bool GetCustomBuildNanite(bool& AttributeValue) const;
 
 	/** Set whether the static mesh factory should set the Nanite build setting. Return false if the attribute was not set. */
 	UFUNCTION(BlueprintCallable, Category = "Interchange | Node | StaticMesh")
-	INTERCHANGEFACTORYNODES_API bool SetCustomBuildNanite(const bool& AttributeValue, bool bAddApplyDelegate = true);
+	UE_API bool SetCustomBuildNanite(const bool& AttributeValue, bool bAddApplyDelegate = true);
 
 	/** Return the number of socket UIDs this static mesh has. */
 	UFUNCTION(BlueprintCallable, Category = "Interchange | Node | StaticMesh")
-	INTERCHANGEFACTORYNODES_API int32 GetSocketUidCount() const;
+	UE_API int32 GetSocketUidCount() const;
 
 	UFUNCTION(BlueprintCallable, Category = "Interchange | Node | StaticMesh")
-	INTERCHANGEFACTORYNODES_API void GetSocketUids(TArray<FString>& OutSocketUids) const;
+	UE_API void GetSocketUids(TArray<FString>& OutSocketUids) const;
 
 	UFUNCTION(BlueprintCallable, Category = "Interchange | Node | StaticMesh")
-	INTERCHANGEFACTORYNODES_API bool AddSocketUid(const FString& SocketUid);
+	UE_API bool AddSocketUid(const FString& SocketUid);
 
 	UFUNCTION(BlueprintCallable, Category = "Interchange | Node | StaticMesh")
-	INTERCHANGEFACTORYNODES_API bool AddSocketUids(const TArray<FString>& InSocketUids);
+	UE_API bool AddSocketUids(const TArray<FString>& InSocketUids);
 
 	UFUNCTION(BlueprintCallable, Category = "Interchange | Node | StaticMesh")
-	INTERCHANGEFACTORYNODES_API bool RemoveSocketUd(const FString& SocketUid);
+	UE_API bool RemoveSocketUd(const FString& SocketUid);
 
 	/** Get whether the static mesh should build a reversed index buffer. */
 	UFUNCTION(BlueprintCallable, Category = "Interchange | Node | StaticMesh")
-	INTERCHANGEFACTORYNODES_API bool GetCustomBuildReversedIndexBuffer(bool& AttributeValue) const;
+	UE_API bool GetCustomBuildReversedIndexBuffer(bool& AttributeValue) const;
 
 	/** Set whether the static mesh should build a reversed index buffer. */
 	UFUNCTION(BlueprintCallable, Category = "Interchange | Node | StaticMesh")
-	INTERCHANGEFACTORYNODES_API bool SetCustomBuildReversedIndexBuffer(const bool& AttributeValue, bool bAddApplyDelegate = true);
+	UE_API bool SetCustomBuildReversedIndexBuffer(const bool& AttributeValue, bool bAddApplyDelegate = true);
 
 	/** Get whether the static mesh should generate lightmap UVs. */
 	UFUNCTION(BlueprintCallable, Category = "Interchange | Node | StaticMesh")
-	INTERCHANGEFACTORYNODES_API bool GetCustomGenerateLightmapUVs(bool& AttributeValue) const;
+	UE_API bool GetCustomGenerateLightmapUVs(bool& AttributeValue) const;
 
 	/** Set whether the static mesh should generate lightmap UVs. */
 	UFUNCTION(BlueprintCallable, Category = "Interchange | Node | StaticMesh")
-	INTERCHANGEFACTORYNODES_API bool SetCustomGenerateLightmapUVs(const bool& AttributeValue, bool bAddApplyDelegate = true);
+	UE_API bool SetCustomGenerateLightmapUVs(const bool& AttributeValue, bool bAddApplyDelegate = true);
 
 	/**
 	 * Get whether to generate the distance field by treating every triangle hit as a front face.  
 	 * This prevents the distance field from being discarded due to the mesh being open, but also lowers distance field ambient occlusion quality.
 	 */
 	UFUNCTION(BlueprintCallable, Category = "Interchange | Node | StaticMesh")
-	INTERCHANGEFACTORYNODES_API bool GetCustomGenerateDistanceFieldAsIfTwoSided(bool& AttributeValue) const;
+	UE_API bool GetCustomGenerateDistanceFieldAsIfTwoSided(bool& AttributeValue) const;
 
 	/**
 	 * Set whether to generate the distance field by treating every triangle hit as a front face.
 	 * This prevents the distance field from being discarded due to the mesh being open, but also lowers distance field ambient occlusion quality.
 	 */
 	UFUNCTION(BlueprintCallable, Category = "Interchange | Node | StaticMesh")
-	INTERCHANGEFACTORYNODES_API bool SetCustomGenerateDistanceFieldAsIfTwoSided(const bool& AttributeValue, bool bAddApplyDelegate = true);
+	UE_API bool SetCustomGenerateDistanceFieldAsIfTwoSided(const bool& AttributeValue, bool bAddApplyDelegate = true);
 
 	/** Get whether the static mesh is set up for use with physical material masks. */
 	UFUNCTION(BlueprintCallable, Category = "Interchange | Node | StaticMesh")
-	INTERCHANGEFACTORYNODES_API bool GetCustomSupportFaceRemap(bool& AttributeValue) const;
+	UE_API bool GetCustomSupportFaceRemap(bool& AttributeValue) const;
 
 	/** Set whether the static mesh is set up for use with physical material masks. */
 	UFUNCTION(BlueprintCallable, Category = "Interchange | Node | StaticMesh")
-	INTERCHANGEFACTORYNODES_API bool SetCustomSupportFaceRemap(const bool& AttributeValue, bool bAddApplyDelegate = true);
+	UE_API bool SetCustomSupportFaceRemap(const bool& AttributeValue, bool bAddApplyDelegate = true);
 
 	/** Get the amount of padding used to pack UVs for the static mesh. */
 	UFUNCTION(BlueprintCallable, Category = "Interchange | Node | StaticMesh")
-	INTERCHANGEFACTORYNODES_API bool GetCustomMinLightmapResolution(int32& AttributeValue) const;
+	UE_API bool GetCustomMinLightmapResolution(int32& AttributeValue) const;
 
 	/** Set the amount of padding used to pack UVs for the static mesh. */
 	UFUNCTION(BlueprintCallable, Category = "Interchange | Node | StaticMesh")
-	INTERCHANGEFACTORYNODES_API bool SetCustomMinLightmapResolution(const int32& AttributeValue, bool bAddApplyDelegate = true);
+	UE_API bool SetCustomMinLightmapResolution(const int32& AttributeValue, bool bAddApplyDelegate = true);
 
 	/** Get the index of the UV that is used as the source for generating lightmaps for the static mesh.  */
 	UFUNCTION(BlueprintCallable, Category = "Interchange | Node | StaticMesh")
-	INTERCHANGEFACTORYNODES_API bool GetCustomSrcLightmapIndex(int32& AttributeValue) const;
+	UE_API bool GetCustomSrcLightmapIndex(int32& AttributeValue) const;
 
 	/** Set the index of the UV that is used as the source for generating lightmaps for the static mesh. */
 	UFUNCTION(BlueprintCallable, Category = "Interchange | Node | StaticMesh")
-	INTERCHANGEFACTORYNODES_API bool SetCustomSrcLightmapIndex(const int32& AttributeValue, bool bAddApplyDelegate = true);
+	UE_API bool SetCustomSrcLightmapIndex(const int32& AttributeValue, bool bAddApplyDelegate = true);
 
 	/** Get the index of the UV that is used to store generated lightmaps for the static mesh. */
 	UFUNCTION(BlueprintCallable, Category = "Interchange | Node | StaticMesh")
-	INTERCHANGEFACTORYNODES_API bool GetCustomDstLightmapIndex(int32& AttributeValue) const;
+	UE_API bool GetCustomDstLightmapIndex(int32& AttributeValue) const;
 
 	/** Set the index of the UV that is used to store generated lightmaps for the static mesh. */
 	UFUNCTION(BlueprintCallable, Category = "Interchange | Node | StaticMesh")
-	INTERCHANGEFACTORYNODES_API bool SetCustomDstLightmapIndex(const int32& AttributeValue, bool bAddApplyDelegate = true);
+	UE_API bool SetCustomDstLightmapIndex(const int32& AttributeValue, bool bAddApplyDelegate = true);
 
 	/** Get the local scale that is applied when building the static mesh. */
 	UFUNCTION(BlueprintCallable, Category = "Interchange | Node | StaticMesh")
-	INTERCHANGEFACTORYNODES_API bool GetCustomBuildScale3D(FVector& AttributeValue) const;
+	UE_API bool GetCustomBuildScale3D(FVector& AttributeValue) const;
 
 	/** Set the local scale that is applied when building the static mesh. */
 	UFUNCTION(BlueprintCallable, Category = "Interchange | Node | StaticMesh")
-	INTERCHANGEFACTORYNODES_API bool SetCustomBuildScale3D(const FVector& AttributeValue, bool bAddApplyDelegate = true);
+	UE_API bool SetCustomBuildScale3D(const FVector& AttributeValue, bool bAddApplyDelegate = true);
 
 	/**
 	 * Get the scale to apply to the mesh when allocating the distance field volume texture.
 	 * The default scale is 1, which assumes that the mesh will be placed unscaled in the world.
 	 */
 	UFUNCTION(BlueprintCallable, Category = "Interchange | Node | StaticMesh")
-	INTERCHANGEFACTORYNODES_API bool GetCustomDistanceFieldResolutionScale(float& AttributeValue) const;
+	UE_API bool GetCustomDistanceFieldResolutionScale(float& AttributeValue) const;
 
 	/**
 	 * Set the scale to apply to the mesh when allocating the distance field volume texture.
 	 * The default scale is 1, which assumes that the mesh will be placed unscaled in the world.
 	 */
 	UFUNCTION(BlueprintCallable, Category = "Interchange | Node | StaticMesh")
-	INTERCHANGEFACTORYNODES_API bool SetCustomDistanceFieldResolutionScale(const float& AttributeValue, bool bAddApplyDelegate = true);
+	UE_API bool SetCustomDistanceFieldResolutionScale(const float& AttributeValue, bool bAddApplyDelegate = true);
 
 	/** Get the static mesh asset whose distance field will be used as the distance field for the imported mesh. */
 	UFUNCTION(BlueprintCallable, Category = "Interchange | Node | StaticMesh")
-	INTERCHANGEFACTORYNODES_API bool GetCustomDistanceFieldReplacementMesh(FSoftObjectPath& AttributeValue) const;
+	UE_API bool GetCustomDistanceFieldReplacementMesh(FSoftObjectPath& AttributeValue) const;
 
 	/** Set the static mesh asset whose distance field will be used as the distance field for the imported mesh. */
 	UFUNCTION(BlueprintCallable, Category = "Interchange | Node | StaticMesh")
-	INTERCHANGEFACTORYNODES_API bool SetCustomDistanceFieldReplacementMesh(const FSoftObjectPath& AttributeValue, bool bAddApplyDelegate = true);
+	UE_API bool SetCustomDistanceFieldReplacementMesh(const FSoftObjectPath& AttributeValue, bool bAddApplyDelegate = true);
 
 	/**
 	 * Get the maximum number of Lumen mesh cards to generate for this mesh.
@@ -202,7 +205,7 @@ public:
 	 * The default is 12.
 	 */
 	UFUNCTION(BlueprintCallable, Category = "Interchange | Node | StaticMesh")
-	INTERCHANGEFACTORYNODES_API bool GetCustomMaxLumenMeshCards(int32& AttributeValue) const;
+	UE_API bool GetCustomMaxLumenMeshCards(int32& AttributeValue) const;
 
 	/**
 	 * Set the maximum number of Lumen mesh cards to generate for this mesh.
@@ -211,11 +214,11 @@ public:
 	 * The default is 12.
 	 */
 	UFUNCTION(BlueprintCallable, Category = "Interchange | Node | StaticMesh")
-	INTERCHANGEFACTORYNODES_API bool SetCustomMaxLumenMeshCards(const int32& AttributeValue, bool bAddApplyDelegate = true);
+	UE_API bool SetCustomMaxLumenMeshCards(const int32& AttributeValue, bool bAddApplyDelegate = true);
 
 private:
-	INTERCHANGEFACTORYNODES_API virtual void FillAssetClassFromAttribute() override;
-	INTERCHANGEFACTORYNODES_API virtual bool SetNodeClassFromClassAttribute() override;
+	UE_API virtual void FillAssetClassFromAttribute() override;
+	UE_API virtual bool SetNodeClassFromClassAttribute() override;
 
 	const UE::Interchange::FAttributeKey Macro_CustomBuildReversedIndexBufferKey = UE::Interchange::FAttributeKey(TEXT("BuildReversedIndexBuffer"));
 	const UE::Interchange::FAttributeKey Macro_CustomGenerateLightmapUVsKey = UE::Interchange::FAttributeKey(TEXT("GenerateLightmapUVs"));
@@ -241,30 +244,32 @@ protected:
 	IMPLEMENT_NODE_ATTRIBUTE_DELEGATE_BY_PROPERTYNAME(BuildNanite, bool, UStaticMesh, TEXT("NaniteSettings.bEnabled"));
 #endif
 
-	INTERCHANGEFACTORYNODES_API bool ApplyCustomBuildReversedIndexBufferToAsset(UObject* Asset) const;
-	INTERCHANGEFACTORYNODES_API bool FillCustomBuildReversedIndexBufferFromAsset(UObject* Asset);
-	INTERCHANGEFACTORYNODES_API bool ApplyCustomGenerateLightmapUVsToAsset(UObject* Asset) const;
-	INTERCHANGEFACTORYNODES_API bool FillCustomGenerateLightmapUVsFromAsset(UObject* Asset);
-	INTERCHANGEFACTORYNODES_API bool ApplyCustomGenerateDistanceFieldAsIfTwoSidedToAsset(UObject* Asset) const;
-	INTERCHANGEFACTORYNODES_API bool FillCustomGenerateDistanceFieldAsIfTwoSidedFromAsset(UObject* Asset);
-	INTERCHANGEFACTORYNODES_API bool ApplyCustomSupportFaceRemapToAsset(UObject* Asset) const;
-	INTERCHANGEFACTORYNODES_API bool FillCustomSupportFaceRemapFromAsset(UObject* Asset);
-	INTERCHANGEFACTORYNODES_API bool ApplyCustomMinLightmapResolutionToAsset(UObject* Asset) const;
-	INTERCHANGEFACTORYNODES_API bool FillCustomMinLightmapResolutionFromAsset(UObject* Asset);
-	INTERCHANGEFACTORYNODES_API bool ApplyCustomSrcLightmapIndexToAsset(UObject* Asset) const;
-	INTERCHANGEFACTORYNODES_API bool FillCustomSrcLightmapIndexFromAsset(UObject* Asset);
-	INTERCHANGEFACTORYNODES_API bool ApplyCustomDstLightmapIndexToAsset(UObject* Asset) const;
-	INTERCHANGEFACTORYNODES_API bool FillCustomDstLightmapIndexFromAsset(UObject* Asset);
-	INTERCHANGEFACTORYNODES_API bool ApplyCustomBuildScale3DToAsset(UObject* Asset) const;
-	INTERCHANGEFACTORYNODES_API bool FillCustomBuildScale3DFromAsset(UObject* Asset);
-	INTERCHANGEFACTORYNODES_API bool ApplyCustomDistanceFieldResolutionScaleToAsset(UObject* Asset) const;
-	INTERCHANGEFACTORYNODES_API bool FillCustomDistanceFieldResolutionScaleFromAsset(UObject* Asset);
-	INTERCHANGEFACTORYNODES_API bool ApplyCustomDistanceFieldReplacementMeshToAsset(UObject* Asset) const;
-	INTERCHANGEFACTORYNODES_API bool FillCustomDistanceFieldReplacementMeshFromAsset(UObject* Asset);
-	INTERCHANGEFACTORYNODES_API bool ApplyCustomMaxLumenMeshCardsToAsset(UObject* Asset) const;
-	INTERCHANGEFACTORYNODES_API bool FillCustomMaxLumenMeshCardsFromAsset(UObject* Asset);
+	UE_API bool ApplyCustomBuildReversedIndexBufferToAsset(UObject* Asset) const;
+	UE_API bool FillCustomBuildReversedIndexBufferFromAsset(UObject* Asset);
+	UE_API bool ApplyCustomGenerateLightmapUVsToAsset(UObject* Asset) const;
+	UE_API bool FillCustomGenerateLightmapUVsFromAsset(UObject* Asset);
+	UE_API bool ApplyCustomGenerateDistanceFieldAsIfTwoSidedToAsset(UObject* Asset) const;
+	UE_API bool FillCustomGenerateDistanceFieldAsIfTwoSidedFromAsset(UObject* Asset);
+	UE_API bool ApplyCustomSupportFaceRemapToAsset(UObject* Asset) const;
+	UE_API bool FillCustomSupportFaceRemapFromAsset(UObject* Asset);
+	UE_API bool ApplyCustomMinLightmapResolutionToAsset(UObject* Asset) const;
+	UE_API bool FillCustomMinLightmapResolutionFromAsset(UObject* Asset);
+	UE_API bool ApplyCustomSrcLightmapIndexToAsset(UObject* Asset) const;
+	UE_API bool FillCustomSrcLightmapIndexFromAsset(UObject* Asset);
+	UE_API bool ApplyCustomDstLightmapIndexToAsset(UObject* Asset) const;
+	UE_API bool FillCustomDstLightmapIndexFromAsset(UObject* Asset);
+	UE_API bool ApplyCustomBuildScale3DToAsset(UObject* Asset) const;
+	UE_API bool FillCustomBuildScale3DFromAsset(UObject* Asset);
+	UE_API bool ApplyCustomDistanceFieldResolutionScaleToAsset(UObject* Asset) const;
+	UE_API bool FillCustomDistanceFieldResolutionScaleFromAsset(UObject* Asset);
+	UE_API bool ApplyCustomDistanceFieldReplacementMeshToAsset(UObject* Asset) const;
+	UE_API bool FillCustomDistanceFieldReplacementMeshFromAsset(UObject* Asset);
+	UE_API bool ApplyCustomMaxLumenMeshCardsToAsset(UObject* Asset) const;
+	UE_API bool FillCustomMaxLumenMeshCardsFromAsset(UObject* Asset);
 
 #if WITH_ENGINE
 	TSubclassOf<UStaticMesh> AssetClass;
 #endif
 };
+
+#undef UE_API
\ No newline at end of file
diff --git a/Engine/Plugins/Interchange/Runtime/Source/Import/Private/Animation/InterchangeAnimSequenceFactory.cpp b/Engine/Plugins/Interchange/Runtime/Source/Import/Private/Animation/InterchangeAnimSequenceFactory.cpp
--- a/Engine/Plugins/Interchange/Runtime/Source/Import/Private/Animation/InterchangeAnimSequenceFactory.cpp
+++ b/Engine/Plugins/Interchange/Runtime/Source/Import/Private/Animation/InterchangeAnimSequenceFactory.cpp
@@ -1527,9 +1527,7 @@ UInterchangeFactoryBase::FImportAssetResult UInterchangeAnimSequenceFactory::End
 		{
 			PreviousNode = InterchangeAssetImportData->GetStoredFactoryNode(InterchangeAssetImportData->NodeUniqueID);
 		}
-		UInterchangeFactoryBaseNode* CurrentNode = NewObject<UInterchangeAnimSequenceFactoryNode>(GetTransientPackage());
-		UInterchangeBaseNode::CopyStorage(AnimSequenceFactoryNode, CurrentNode);
-		CurrentNode->FillAllCustomAttributeFromObject(AnimSequence);
+		UInterchangeFactoryBaseNode* CurrentNode = UInterchangeFactoryBaseNode::DuplicateWithObject(AnimSequenceFactoryNode, AnimSequence);
 		UE::Interchange::FFactoryCommon::ApplyReimportStrategyToAsset(AnimSequence, PreviousNode, CurrentNode, AnimSequenceFactoryNode);
 	}
 #endif
diff --git a/Engine/Plugins/Interchange/Runtime/Source/Import/Private/Audio/InterchangeAudioSoundWaveFactory.cpp b/Engine/Plugins/Interchange/Runtime/Source/Import/Private/Audio/InterchangeAudioSoundWaveFactory.cpp
--- a/Engine/Plugins/Interchange/Runtime/Source/Import/Private/Audio/InterchangeAudioSoundWaveFactory.cpp
+++ b/Engine/Plugins/Interchange/Runtime/Source/Import/Private/Audio/InterchangeAudioSoundWaveFactory.cpp
@@ -251,11 +251,7 @@ void UInterchangeAudioSoundWaveFactory::SetupObject_GameThread(const FSetupObjec
 			PreviousNode = InterchangeAssetImportData->GetStoredFactoryNode(InterchangeAssetImportData->NodeUniqueID);
 		}
 
-		UInterchangeFactoryBaseNode* CurrentNode = NewObject<UInterchangeFactoryBaseNode>(GetTransientPackage(), UInterchangeAudioSoundWaveFactoryNode::StaticClass());
-		ensure(CurrentNode);
-		UInterchangeBaseNode::CopyStorage(SoundWaveFactoryNode, CurrentNode);
-		CurrentNode->FillAllCustomAttributeFromObject(SoundWaveAsset);
-
+		UInterchangeFactoryBaseNode* CurrentNode = UInterchangeFactoryBaseNode::DuplicateWithObject(SoundWaveFactoryNode, SoundWaveAsset);
 		UE::Interchange::FFactoryCommon::ApplyReimportStrategyToAsset(SoundWaveAsset, PreviousNode, CurrentNode, SoundWaveFactoryNode);
 	}
 	else
diff --git a/Engine/Plugins/Interchange/Runtime/Source/Import/Private/Mesh/InterchangeSkeletalMeshFactory.cpp b/Engine/Plugins/Interchange/Runtime/Source/Import/Private/Mesh/InterchangeSkeletalMeshFactory.cpp
--- a/Engine/Plugins/Interchange/Runtime/Source/Import/Private/Mesh/InterchangeSkeletalMeshFactory.cpp
+++ b/Engine/Plugins/Interchange/Runtime/Source/Import/Private/Mesh/InterchangeSkeletalMeshFactory.cpp
@@ -2079,9 +2079,7 @@ UInterchangeFactoryBase::FImportAssetResult UInterchangeSkeletalMeshFactory::End
 		{
 			PreviousNode = InterchangeAssetImportData->GetStoredFactoryNode(InterchangeAssetImportData->NodeUniqueID);
 		}
-		UInterchangeFactoryBaseNode* CurrentNode = NewObject<UInterchangeSkeletalMeshFactoryNode>(GetTransientPackage());
-		UInterchangeBaseNode::CopyStorage(SkeletalMeshFactoryNode, CurrentNode);
-		CurrentNode->FillAllCustomAttributeFromObject(SkeletalMesh);
+		UInterchangeFactoryBaseNode* CurrentNode = UInterchangeFactoryBaseNode::DuplicateWithObject(SkeletalMeshFactoryNode, SkeletalMesh);
 		UE::Interchange::FFactoryCommon::ApplyReimportStrategyToAsset(SkeletalMesh, PreviousNode, CurrentNode, SkeletalMeshFactoryNode);
 	}
 
diff --git a/Engine/Plugins/Interchange/Runtime/Source/Import/Private/Mesh/InterchangeStaticMeshFactory.cpp b/Engine/Plugins/Interchange/Runtime/Source/Import/Private/Mesh/InterchangeStaticMeshFactory.cpp
--- a/Engine/Plugins/Interchange/Runtime/Source/Import/Private/Mesh/InterchangeStaticMeshFactory.cpp
+++ b/Engine/Plugins/Interchange/Runtime/Source/Import/Private/Mesh/InterchangeStaticMeshFactory.cpp
@@ -796,9 +796,8 @@ UInterchangeFactoryBase::FImportAssetResult UInterchangeStaticMeshFactory::EndIm
 		{
 			PreviousNode = InterchangeAssetImportData->GetStoredFactoryNode(InterchangeAssetImportData->NodeUniqueID);
 		}
-		UInterchangeFactoryBaseNode* CurrentNode = NewObject<UInterchangeStaticMeshFactoryNode>(GetTransientPackage());
-		UInterchangeBaseNode::CopyStorage(StaticMeshFactoryNode, CurrentNode);
-		CurrentNode->FillAllCustomAttributeFromObject(StaticMesh);
+
+		UInterchangeFactoryBaseNode* CurrentNode = UInterchangeFactoryBaseNode::DuplicateWithObject(StaticMeshFactoryNode, StaticMesh);
 		UE::Interchange::FFactoryCommon::ApplyReimportStrategyToAsset(StaticMesh, PreviousNode, CurrentNode, StaticMeshFactoryNode);
 
 		//Reorder the Hires mesh description in the same order has the lod 0 mesh description
diff --git a/Engine/Plugins/Interchange/Runtime/Source/Import/Private/Texture/InterchangeTextureFactory.cpp b/Engine/Plugins/Interchange/Runtime/Source/Import/Private/Texture/InterchangeTextureFactory.cpp
--- a/Engine/Plugins/Interchange/Runtime/Source/Import/Private/Texture/InterchangeTextureFactory.cpp
+++ b/Engine/Plugins/Interchange/Runtime/Source/Import/Private/Texture/InterchangeTextureFactory.cpp
@@ -2260,9 +2260,7 @@ void UInterchangeTextureFactory::SetupObject_GameThread(const FSetupObjectParams
 				PreviousNode = InterchangeAssetImportData->GetStoredFactoryNode(InterchangeAssetImportData->NodeUniqueID);
 			}
 
-			UInterchangeFactoryBaseNode* CurrentNode = NewObject<UInterchangeFactoryBaseNode>(GetTransientPackage(), GetSupportedFactoryNodeClass(TextureFactoryNode));
-			UInterchangeBaseNode::CopyStorage(TextureFactoryNode, CurrentNode);
-			CurrentNode->FillAllCustomAttributeFromObject(Texture);
+			UInterchangeFactoryBaseNode* CurrentNode = UInterchangeFactoryBaseNode::DuplicateWithObject(TextureFactoryNode, Texture);
 			//Apply reimport strategy
 			UE::Interchange::FFactoryCommon::ApplyReimportStrategyToAsset(Texture, PreviousNode, CurrentNode, TextureFactoryNode);
 			
diff --git a/Engine/Plugins/Interchange/Runtime/Source/Import/Private/Volume/InterchangeSparseVolumeTextureFactory.cpp b/Engine/Plugins/Interchange/Runtime/Source/Import/Private/Volume/InterchangeSparseVolumeTextureFactory.cpp
--- a/Engine/Plugins/Interchange/Runtime/Source/Import/Private/Volume/InterchangeSparseVolumeTextureFactory.cpp
+++ b/Engine/Plugins/Interchange/Runtime/Source/Import/Private/Volume/InterchangeSparseVolumeTextureFactory.cpp
@@ -990,9 +990,7 @@ UInterchangeFactoryBase::FImportAssetResult UInterchangeSparseVolumeTextureFacto
 				UInterchangeFactoryBaseNode* PreviousNode = AssetImportData->GetStoredFactoryNode(AssetImportData->NodeUniqueID);
 
 				// Create a factory node filled with the property values we can extract from our new imported asset
-				UInterchangeFactoryBaseNode* CurrentNode = NewObject<UInterchangeSparseVolumeTextureFactoryNode>(GetTransientPackage());
-				UInterchangeBaseNode::CopyStorage(FactoryNode, CurrentNode);
-				CurrentNode->FillAllCustomAttributeFromObject(ImportedSparseVolumeTexture);
+				UInterchangeFactoryBaseNode* CurrentNode = UInterchangeFactoryBaseNode::DuplicateWithObject(FactoryNode, ImportedSparseVolumeTexture);
 
 				UE::Interchange::FFactoryCommon::ApplyReimportStrategyToAsset(ImportedSparseVolumeTexture, PreviousNode, CurrentNode, FactoryNode);
 			}
