[Interchange-FBX] Not allowing SceneRoots in Skeletons (for Force SkeletalMesh imports) to match Legacy FBX behavior.

#rnx
#virtualized

diff --git a/Engine/Plugins/Interchange/Runtime/Source/Nodes/Private/InterchangeSceneNode.cpp b/Engine/Plugins/Interchange/Runtime/Source/Nodes/Private/InterchangeSceneNode.cpp
--- a/Engine/Plugins/Interchange/Runtime/Source/Nodes/Private/InterchangeSceneNode.cpp
+++ b/Engine/Plugins/Interchange/Runtime/Source/Nodes/Private/InterchangeSceneNode.cpp
@@ -566,3 +566,13 @@ bool UInterchangeSceneNode::GetCustomGlobalMatrixForT0Rebinding(FMatrix& Attribu
 {
 	IMPLEMENT_NODE_ATTRIBUTE_GETTER(GlobalMatrixForT0Rebinding, FMatrix);
 }
+
+bool UInterchangeSceneNode::SetCustomIsSceneRoot(const bool& AttributeValue)
+{
+	IMPLEMENT_NODE_ATTRIBUTE_SETTER_NODELEGATE(IsSceneRoot, bool);
+}
+
+bool UInterchangeSceneNode::GetCustomIsSceneRoot(bool& AttributeValue) const
+{
+	IMPLEMENT_NODE_ATTRIBUTE_GETTER(IsSceneRoot, bool);
+}
\ No newline at end of file
diff --git a/Engine/Plugins/Interchange/Runtime/Source/Nodes/Public/InterchangeSceneNode.h b/Engine/Plugins/Interchange/Runtime/Source/Nodes/Public/InterchangeSceneNode.h
--- a/Engine/Plugins/Interchange/Runtime/Source/Nodes/Public/InterchangeSceneNode.h
+++ b/Engine/Plugins/Interchange/Runtime/Source/Nodes/Public/InterchangeSceneNode.h
@@ -304,6 +304,15 @@ public:
 	UFUNCTION(BlueprintCallable, Category = "Interchange | Scene | Component")
 	UE_API void GetComponentUids(TArray<FString>& OutComponentUids) const;
 
+
+	/** Sets IsSceneRootNode, which indicates if the SceneNode is a Scene Root Node. */
+	UFUNCTION(BlueprintCallable, Category = "Interchange | Scene | Component")
+	UE_API bool SetCustomIsSceneRoot(const bool& IsSceneRoot);
+
+	/** Gets IsSceneRoot, which indicates if the SceneNode is a Scene Root Node. */
+	UFUNCTION(BlueprintCallable, Category = "Interchange | Scene | Component")
+	UE_API bool GetCustomIsSceneRoot(bool& IsSceneRoot) const;
+
 private:
 
 	UE_API bool GetGlobalTransformInternal(const UE::Interchange::FAttributeKey LocalTransformKey
@@ -368,6 +377,9 @@ private:
 	//Global Transform Matrix specifically used for T0 re-binding.
 	const UE::Interchange::FAttributeKey Macro_CustomGlobalMatrixForT0RebindingKey = UE::Interchange::FAttributeKey(TEXT("GlobalMatrixForT0Rebinding"));
 
+	//For storing flag in case SceneNode is a Scene Root Node.
+	const UE::Interchange::FAttributeKey Macro_CustomIsSceneRootKey = UE::Interchange::FAttributeKey(TEXT("IsSceneRoot"));
+
 	//mutable caches for global transforms
 	mutable TOptional<FTransform> CacheGlobalTransform;
 	mutable TOptional<FTransform> CacheBindPoseGlobalTransform;
diff --git a/Engine/Plugins/Interchange/Runtime/Source/Parsers/Fbx/Private/FbxAPI.cpp b/Engine/Plugins/Interchange/Runtime/Source/Parsers/Fbx/Private/FbxAPI.cpp
--- a/Engine/Plugins/Interchange/Runtime/Source/Parsers/Fbx/Private/FbxAPI.cpp
+++ b/Engine/Plugins/Interchange/Runtime/Source/Parsers/Fbx/Private/FbxAPI.cpp
@@ -162,6 +162,9 @@ namespace UE
 				// This fix the issue with blender armature bone skip
 				SourceNode->SetCustomUseLegacySkeletalMeshBakeTransform(true);
 
+				//Fbx legacy does not allow Scene Root Nodes to be part of the skeletons (to be joints).
+				SourceNode->SetCustomAllowSceneRootAsJoint(false);
+
 				// Get the version number of the FBX file format.
 				int32 FileMajor, FileMinor, FileRevision;
 				SDKImporter->GetFileVersion(FileMajor, FileMinor, FileRevision);
diff --git a/Engine/Plugins/Interchange/Runtime/Source/Parsers/Fbx/Private/FbxScene.cpp b/Engine/Plugins/Interchange/Runtime/Source/Parsers/Fbx/Private/FbxScene.cpp
--- a/Engine/Plugins/Interchange/Runtime/Source/Parsers/Fbx/Private/FbxScene.cpp
+++ b/Engine/Plugins/Interchange/Runtime/Source/Parsers/Fbx/Private/FbxScene.cpp
@@ -149,6 +149,12 @@ namespace UE
 				FString NodeUniqueID = Parser.GetFbxHelper()->GetFbxNodeHierarchyName(Node);
 				const bool bIsRootNode = Node == SDKScene->GetRootNode();
 				UInterchangeSceneNode* UnrealNode = CreateTransformNode(NodeContainer, NodeName, NodeUniqueID, UnrealParentNode ? UnrealParentNode->GetUniqueID() : TEXT(""));
+
+				if (bIsRootNode)
+				{
+					UnrealNode->SetCustomIsSceneRoot(true);
+				}
+
 				if (!ensure(UnrealNode))
 				{
 					return;
diff --git a/Engine/Plugins/Interchange/Runtime/Source/Pipelines/Private/InterchangePipelineMeshesUtilities.cpp b/Engine/Plugins/Interchange/Runtime/Source/Pipelines/Private/InterchangePipelineMeshesUtilities.cpp
--- a/Engine/Plugins/Interchange/Runtime/Source/Pipelines/Private/InterchangePipelineMeshesUtilities.cpp
+++ b/Engine/Plugins/Interchange/Runtime/Source/Pipelines/Private/InterchangePipelineMeshesUtilities.cpp
@@ -4,6 +4,7 @@
 
 #include "CoreMinimal.h"
 #include "Nodes/InterchangeBaseNode.h"
+#include "Nodes/InterchangeSourceNode.h"
 #include "InterchangeMeshFactoryNode.h"
 #include "InterchangeMeshLODContainerNode.h"
 #include "InterchangeMeshNode.h"
@@ -916,7 +917,17 @@ void UInterchangePipelineMeshesUtilities::IterateAllMeshInstanceUsingMeshGeometr
 void UInterchangePipelineMeshesUtilities::GetCombinedSkinnedMeshInstances(TMap<FString, TArray<FString>>& OutMeshInstanceUidsPerSkeletonRootUid,
 	bool bUseSingleBoneForConvertedMeshes /*= false*/) const
 {
-	check(BaseNodeContainer);
+	if (!ensure(BaseNodeContainer))
+	{
+		return;
+	}
+
+	bool bAllowSceneRootAsJoint = true;
+	if (const UInterchangeSourceNode* SourceNode = UInterchangeSourceNode::GetUniqueInstance(BaseNodeContainer))
+	{
+		SourceNode->GetCustomAllowSceneRootAsJoint(bAllowSceneRootAsJoint);
+	}
+
 
 	for (const TPair<FString, FInterchangeMeshInstance>& MeshInstanceUidAndMeshInstance : MeshInstancesPerMeshInstanceUid)
 	{
@@ -958,18 +969,28 @@ void UInterchangePipelineMeshesUtilities::GetCombinedSkinnedMeshInstances(TMap<F
 					{
 						FString ParentUid = SceneNode->GetParentUid();
 						FString LastSceneNodeUid = SkeletonRootUid;
-						while (!ParentUid.Equals(UInterchangeBaseNode::InvalidNodeUid()))
+
+						const UInterchangeSceneNode* ParentNode = Cast<UInterchangeSceneNode>(BaseNodeContainer->GetNode(ParentUid));
+						while (ParentNode)
 						{
-							if (const UInterchangeSceneNode* ParentNode = Cast<UInterchangeSceneNode>(BaseNodeContainer->GetNode(ParentUid)))
+							bool bIsSceneRoot = false;
+							ParentNode->GetCustomIsSceneRoot(bIsSceneRoot);
+							if (!bAllowSceneRootAsJoint && bIsSceneRoot)
 							{
-								if (ParentNode->IsSpecializedTypeContains(UE::Interchange::FSceneNodeStaticData::GetJointSpecializeTypeString()))
-								{
-									SkeletonRootUid = ParentUid;
-								}
-								LastSceneNodeUid = ParentUid;
-								ParentUid = ParentNode->GetParentUid();
+								//Do not include scene root node in Skeleton.
+								break;
+							}
+
+							if (ParentNode->IsSpecializedTypeContains(UE::Interchange::FSceneNodeStaticData::GetJointSpecializeTypeString()))
+							{
+								SkeletonRootUid = ParentUid;
 							}
+							LastSceneNodeUid = ParentUid;
+							ParentUid = ParentNode->GetParentUid();
+
+							ParentNode = Cast<UInterchangeSceneNode>(BaseNodeContainer->GetNode(ParentUid));
 						}
+
 						//If we did not find any joint because we have non nested mesh, get the deepest scene node
 						if (!bIsNestedIntoSkeleton && SkeletonRootUid.Equals(MeshInstance.MeshInstanceUid))
 						{
diff --git a/Engine/Source/Runtime/Interchange/Core/Private/Nodes/InterchangeSourceNode.cpp b/Engine/Source/Runtime/Interchange/Core/Private/Nodes/InterchangeSourceNode.cpp
--- a/Engine/Source/Runtime/Interchange/Core/Private/Nodes/InterchangeSourceNode.cpp
+++ b/Engine/Source/Runtime/Interchange/Core/Private/Nodes/InterchangeSourceNode.cpp
@@ -234,3 +234,14 @@ bool UInterchangeSourceNode::SetCustomNaniteTriangleThreshold(const int64& Attri
 {
 	IMPLEMENT_NODE_ATTRIBUTE_SETTER_NODELEGATE(NaniteTriangleThreshold, int64)
 }
+
+bool UInterchangeSourceNode::GetCustomAllowSceneRootAsJoint(bool& AttributeValue) const
+{
+	IMPLEMENT_NODE_ATTRIBUTE_GETTER(AllowSceneRootAsJoint, bool);
+}
+
+bool UInterchangeSourceNode::SetCustomAllowSceneRootAsJoint(const bool& AttributeValue)
+{
+	IMPLEMENT_NODE_ATTRIBUTE_SETTER_NODELEGATE(AllowSceneRootAsJoint, bool)
+}
+
diff --git a/Engine/Source/Runtime/Interchange/Core/Public/Nodes/InterchangeSourceNode.h b/Engine/Source/Runtime/Interchange/Core/Public/Nodes/InterchangeSourceNode.h
--- a/Engine/Source/Runtime/Interchange/Core/Public/Nodes/InterchangeSourceNode.h
+++ b/Engine/Source/Runtime/Interchange/Core/Public/Nodes/InterchangeSourceNode.h
@@ -207,6 +207,14 @@ public:
 	UFUNCTION(BlueprintCallable, Category = "Interchange | Node | Source")
 	INTERCHANGECORE_API bool SetCustomNaniteTriangleThreshold(const int64& MinNumTriangles);
 
+	/** Gets, if the skeleton processing should accept a SceneRoot as a joint (root node), for when import is forced as Skeletal. */
+	UFUNCTION(BlueprintCallable, Category = "Interchange | Node | Source")
+	INTERCHANGECORE_API bool GetCustomAllowSceneRootAsJoint(bool& bAllowSceneRootAsJoint) const;
+
+	/** Sets, if the skeleton processing should accept a SceneRoot as a joint (root node), for when import is forced as Skeletal. */
+	UFUNCTION(BlueprintCallable, Category = "Interchange | Node | Source")
+	INTERCHANGECORE_API bool SetCustomAllowSceneRootAsJoint(const bool& bAllowSceneRootAsJoint);
+
 private:
 	IMPLEMENT_NODE_ATTRIBUTE_KEY(SourceFrameRateNumerator);
 	IMPLEMENT_NODE_ATTRIBUTE_KEY(SourceFrameRateDenominator);
@@ -222,6 +230,7 @@ private:
 	IMPLEMENT_NODE_ATTRIBUTE_KEY(ReimportStrategyFlags);
 	IMPLEMENT_NODE_ATTRIBUTE_KEY(SkeletalMeshFrontAxis);
 	IMPLEMENT_NODE_ATTRIBUTE_KEY(NaniteTriangleThreshold);
+	IMPLEMENT_NODE_ATTRIBUTE_KEY(AllowSceneRootAsJoint);
 
 	// Extra InformationTo show in the Config Panel.
 	UE::Interchange::TMapAttributeHelper<FString, FString> ExtraInformation;
