[Interchange/Thumbnail Render] Crash fix when re-importing SkeletalMesh with fewer Bones than original.

#rnx

diff --git a/Engine/Plugins/Interchange/Runtime/Source/Import/Private/Mesh/InterchangeSkeletalMeshFactory.cpp b/Engine/Plugins/Interchange/Runtime/Source/Import/Private/Mesh/InterchangeSkeletalMeshFactory.cpp
--- a/Engine/Plugins/Interchange/Runtime/Source/Import/Private/Mesh/InterchangeSkeletalMeshFactory.cpp
+++ b/Engine/Plugins/Interchange/Runtime/Source/Import/Private/Mesh/InterchangeSkeletalMeshFactory.cpp
@@ -65,6 +65,32 @@
 
 #endif //WITH_EDITORONLY_DATA
 
+FScopedSkeletalMeshReimportUtility::FScopedSkeletalMeshReimportUtility(USkeletalMesh* InSkeletalMesh)
+	: FSkinnedMeshComponentRecreateRenderStateContext(InSkeletalMesh)
+{
+	InSkeletalMesh->ReleaseResources();
+
+#if WITH_EDITOR
+	// Scope PostEditChange before locking the skeletal mesh's properties
+	ScopedPostEditChange = MakeUnique<FScopedSkeletalMeshPostEditChange>(InSkeletalMesh);
+	// Lock the skeletal mesh's properties to prevent any async access
+	LockPropertiesEvent = InSkeletalMesh->LockPropertiesUntil();
+#endif
+}
+FScopedSkeletalMeshReimportUtility::~FScopedSkeletalMeshReimportUtility()
+{
+#if WITH_EDITOR
+	// Unlock the skeletal mesh's properties first
+	if (ensure(LockPropertiesEvent))
+	{
+		LockPropertiesEvent->Trigger();
+	}
+	LockPropertiesEvent = nullptr;
+	// Now apply PostEditChange to update skeletal mesh's resources
+	ScopedPostEditChange.Reset();
+#endif
+}
+
 #if WITH_EDITOR
 namespace UE
 {
@@ -1149,7 +1175,7 @@ UInterchangeFactoryBase::FImportAssetResult UInterchangeSkeletalMeshFactory::Beg
 	//Lock the skeletalmesh properties if the skeletal mesh already exist (re-import)
 	if (ExistingAsset)
 	{
-		ScopedReimportUtility = MakeUnique<FSkinnedMeshComponentRecreateRenderStateContext>(SkeletalMesh);
+		ScopedReimportUtility = MakeUnique<FScopedSkeletalMeshReimportUtility>(SkeletalMesh);
 	}
 
 	ImportAssetResult.ImportedObject = SkeletalMesh;
@@ -2089,11 +2115,12 @@ UInterchangeFactoryBase::FImportAssetResult UInterchangeSkeletalMeshFactory::End
 
 	ImportAssetResult.ImportedObject = SkeletalMesh;
 
+#endif
+
 	if (ScopedReimportUtility.IsValid())
 	{
 		ScopedReimportUtility.Reset();
 	}
-#endif
 
 	return ImportAssetResult;
 }
@@ -2351,6 +2378,9 @@ void UInterchangeSkeletalMeshFactory::SetupObject_GameThread(const FSetupObjectP
 								AttributesDest.UnregisterSkinWeightAttribute(ProfileName);
 							}
 
+							int32 DestinationNumOfBones = AttributesDest.GetNumBones();
+							bool bSkippedBoneWeight = false;
+
 							FSkeletalMeshConstAttributes AttributesExistingSrc(ExistingMeshDescriptionSrc);
 							if (AttributesExistingSrc.GetSkinWeightProfileNames().Contains(ProfileName))
 							{
@@ -2366,12 +2396,28 @@ void UInterchangeSkeletalMeshFactory::SetupObject_GameThread(const FSetupObjectP
 									TArray<UE::AnimationCore::FBoneWeight> BoneWeights;
 									for (const UE::AnimationCore::FBoneWeight BoneWeight : VertexBoneWeightsExistingSrc)
 									{
-										BoneWeights.Add(BoneWeight);
+										if (BoneWeight.GetBoneIndex() < DestinationNumOfBones)
+										{
+											BoneWeights.Add(BoneWeight);
+										}
+										else
+										{
+											bSkippedBoneWeight = true;
+										}
 									}
 									SkinWeightDest.Set(VertexID, BoneWeights);
 								}
 							}
 
+							if (bSkippedBoneWeight)
+							{
+								UInterchangeResultWarning_Generic* Message = AddMessage<UInterchangeResultWarning_Generic>();
+								Message->Text = FText::Format(NSLOCTEXT("InterchangeSkeletalMeshFactory", "SetupObject_GameThread_AlternateProfileReinstateIssue", "At least 1 Bone Weight was skipped in [{0}] Alternate Skin Weight Profile on LOD [{1}] upon re-instate due to Bone Number missmatch.")
+									, FText::FromName(ProfileName)
+									, FText::AsNumber(LodIndex));
+								Message->AssetFriendlyName = SkeletalMesh->GetName();
+							}
+
 							SaveLodIndexes[LodIndex] = true;
 						}
 						else if (bIsGeneratedByEngine)
diff --git a/Engine/Plugins/Interchange/Runtime/Source/Import/Public/Mesh/InterchangeSkeletalMeshFactory.h b/Engine/Plugins/Interchange/Runtime/Source/Import/Public/Mesh/InterchangeSkeletalMeshFactory.h
--- a/Engine/Plugins/Interchange/Runtime/Source/Import/Public/Mesh/InterchangeSkeletalMeshFactory.h
+++ b/Engine/Plugins/Interchange/Runtime/Source/Import/Public/Mesh/InterchangeSkeletalMeshFactory.h
@@ -23,6 +23,29 @@ class UInterchangeSkeletalMeshFactoryNode;
 class USkeletalMesh;
 class USkeleton;
 
+// Utility class used on reimport only to: 
+// - Call PostEditchange at the end of the reimport which will reallocate the render resource
+// - Lock the skeletal mesh's properties while it is updated
+class FScopedSkeletalMeshReimportUtility : public FSkinnedMeshComponentRecreateRenderStateContext
+{
+public:
+	/*
+	 * This constructor increment the skeletal mesh PostEditChangeStackCounter. If the stack counter is zero before the increment
+	 * the skeletal mesh component will be unregister from the world. The component will also release there rendering resources.
+	 * Parameters:
+	 * @param InbCallPostEditChange - if we are the first scope PostEditChange will be called.
+	 * @param InbReregisterComponents - if we are the first scope we will re register component from world and also component render data.
+	 */
+	FScopedSkeletalMeshReimportUtility(USkeletalMesh* InSkeletalMesh);
+	~FScopedSkeletalMeshReimportUtility();
+
+private:
+#if WITH_EDITOR
+	TUniquePtr<FScopedSkeletalMeshPostEditChange> ScopedPostEditChange;
+#endif
+	FEvent* LockPropertiesEvent = nullptr;
+};
+
 namespace UE::Interchange
 {
 	//Get the mesh node context for each MeshUids
@@ -116,7 +139,7 @@ public:
 	};
 
 private:
-	TUniquePtr<FSkinnedMeshComponentRecreateRenderStateContext> ScopedReimportUtility;
+	TUniquePtr<FScopedSkeletalMeshReimportUtility> ScopedReimportUtility;
 
 	TMap<int32, FLodPayloads> PayloadsPerLodIndex;
 
diff --git a/Engine/Source/Editor/DetailCustomizations/Private/ComponentMaterialCategory.cpp b/Engine/Source/Editor/DetailCustomizations/Private/ComponentMaterialCategory.cpp
--- a/Engine/Source/Editor/DetailCustomizations/Private/ComponentMaterialCategory.cpp
+++ b/Engine/Source/Editor/DetailCustomizations/Private/ComponentMaterialCategory.cpp
@@ -86,6 +86,12 @@ public:
 
 				if( PrimitiveComp )
 				{
+					if (PrimitiveComp->IsCompiling())
+					{
+						++CurComponentIndex;
+						continue;
+					}
+
 					NumMaterials = PrimitiveComp->GetNumMaterials();
 					MaterialSlotNames = PrimitiveComp->GetMaterialSlotNames();
 				}
