[Interchange] Documentation fix for FScopedSkeletalMeshReimportUtility in InterchangeSkeletalMeshFactory file.

#rnx

diff --git a/Engine/Plugins/Interchange/Runtime/Source/Import/Public/Mesh/InterchangeSkeletalMeshFactory.h b/Engine/Plugins/Interchange/Runtime/Source/Import/Public/Mesh/InterchangeSkeletalMeshFactory.h
--- a/Engine/Plugins/Interchange/Runtime/Source/Import/Public/Mesh/InterchangeSkeletalMeshFactory.h
+++ b/Engine/Plugins/Interchange/Runtime/Source/Import/Public/Mesh/InterchangeSkeletalMeshFactory.h
@@ -25,17 +25,11 @@ class USkeleton;
 
 // Utility class used on reimport only to: 
 // - Call PostEditchange at the end of the reimport which will reallocate the render resource
+// - Used to recreate all skinned mesh components for a given skinned asset via FSkinnedMeshComponentRecreateRenderStateContext.
 // - Lock the skeletal mesh's properties while it is updated
 class FScopedSkeletalMeshReimportUtility : public FSkinnedMeshComponentRecreateRenderStateContext
 {
 public:
-	/*
-	 * This constructor increment the skeletal mesh PostEditChangeStackCounter. If the stack counter is zero before the increment
-	 * the skeletal mesh component will be unregister from the world. The component will also release there rendering resources.
-	 * Parameters:
-	 * @param InbCallPostEditChange - if we are the first scope PostEditChange will be called.
-	 * @param InbReregisterComponents - if we are the first scope we will re register component from world and also component render data.
-	 */
 	FScopedSkeletalMeshReimportUtility(USkeletalMesh* InSkeletalMesh);
 	~FScopedSkeletalMeshReimportUtility();
 
