Interchange USD: Fixed support for animated transforms when importing geometry caches.
This will also bake static meshes with animated transforms into the geometry cache.

#rnx

diff --git a/Engine/Plugins/Interchange/Extensions/OpenUSD/Source/Import/Private/InterchangeUsdTranslator.cpp b/Engine/Plugins/Interchange/Extensions/OpenUSD/Source/Import/Private/InterchangeUsdTranslator.cpp
--- a/Engine/Plugins/Interchange/Extensions/OpenUSD/Source/Import/Private/InterchangeUsdTranslator.cpp
+++ b/Engine/Plugins/Interchange/Extensions/OpenUSD/Source/Import/Private/InterchangeUsdTranslator.cpp
@@ -966,6 +966,31 @@ namespace UE::InterchangeUsdTranslator::Private
 		return true;
 	}
 
+	bool ReadRawTransforms(
+		const UE::FUsdStage& UsdStage,
+		const TArray<double>& UsdTimeSamples,
+		const TFunction<FTransform(double)>& ReaderFunc,
+		Interchange::FAnimationPayloadData& OutPayloadData
+	)
+	{
+		OutPayloadData.Transforms.Reserve(UsdTimeSamples.Num());
+
+		double LastTimeSample = TNumericLimits<double>::Lowest();
+		for (const double UsdTimeSample : UsdTimeSamples)
+		{
+			// We never want to evaluate the same time twice
+			if (FMath::IsNearlyEqual(UsdTimeSample, LastTimeSample))
+			{
+				continue;
+			}
+			LastTimeSample = UsdTimeSample;
+
+			OutPayloadData.Transforms.Emplace(ReaderFunc(UsdTimeSample));
+		}
+
+		return true;
+	}
+
 	void AddTextureNode(
 		const UE::FUsdPrim& Prim,
 		const FString& NodeUid,
@@ -6013,6 +6038,51 @@ namespace UE::InterchangeUsdTranslator::Private
 		return false;
 	}
 
+	bool GetTransformAnimationPayloadData(
+		const UInterchangeUSDTranslatorImpl& Impl,
+		const FString& PayloadKey,
+		Interchange::FAnimationPayloadData& OutPayloadData
+	)
+	{
+		TRACE_CPUPROFILER_EVENT_SCOPE(GetTransformAnimationPayloadData)
+
+		FString PrimPath;
+		FString UEPropertyNameStr;
+		bool bSplit = PayloadKey.Split(TEXT("\\"), &PrimPath, &UEPropertyNameStr, ESearchCase::CaseSensitive, ESearchDir::FromEnd);
+		if (!bSplit)
+		{
+			return false;
+		}
+
+		UE::FUsdPrim Prim = Impl.UsdStage.GetPrimAtPath(UE::FSdfPath{*PrimPath});
+		FName UEPropertyName = *UEPropertyNameStr;
+		if (!Prim || UEPropertyName == NAME_None)
+		{
+			return false;
+		}
+
+		TArray<double> TimeSampleUnion;
+
+		TArray<UE::FUsdAttribute> Attrs = UsdUtils::GetAttributesForProperty(Prim, UEPropertyName);
+		bool bSuccess = UE::FUsdAttribute::GetUnionedTimeSamples(Attrs, TimeSampleUnion);
+		if (!bSuccess)
+		{
+			return false;
+		}
+
+		// Use the first timecode as an offset when indexing
+		OutPayloadData.RangeStartTime = TimeSampleUnion[0];
+
+		const bool bIgnorePrimLocalTransform = false;
+		UsdToUnreal::FPropertyTrackReader Reader = UsdToUnreal::CreatePropertyTrackReader(Prim, UEPropertyName, bIgnorePrimLocalTransform);
+		if (Reader.TransformReader)
+		{
+			return ReadRawTransforms(Impl.UsdStage, TimeSampleUnion, Reader.TransformReader, OutPayloadData);
+		}
+
+		return false;
+	}
+
 	bool GetJointAnimationCurvePayloadData(
 		const UInterchangeUSDTranslatorImpl& Impl,
 		const TArray<const UE::Interchange::FAnimationPayloadQuery*>& Queries,
@@ -6948,6 +7018,8 @@ TArray<UE::Interchange::FAnimationPayloadData> UInterchangeUSDTranslator::GetAni
 	TArray<TArray<UE::Interchange::FAnimationPayloadData>> BakeAnimationPayloads;
 	TArray<int32> CurveQueryIndexes;
 	TArray<TArray<UE::Interchange::FAnimationPayloadData>> CurveAnimationPayloads;
+	TArray<int32> TransformQueryIndexes;
+	TArray<TArray<UE::Interchange::FAnimationPayloadData>> TransformAnimationPayloads;
 
 	// Get all curves with a parallel for
 	int32 PayloadCount = PayloadQueries.Num();
@@ -6959,6 +7031,10 @@ TArray<UE::Interchange::FAnimationPayloadData> UInterchangeUSDTranslator::GetAni
 		{
 			BakeQueryIndexes.Add(PayloadIndex);
 		}
+		else if (QueryType == EInterchangeAnimationPayLoadType::GEOMETRY_CACHE_TRANSFORM)
+		{
+			TransformQueryIndexes.Add(PayloadIndex);
+		}
 		else
 		{
 			CurveQueryIndexes.Add(PayloadIndex);
@@ -7089,6 +7165,27 @@ TArray<UE::Interchange::FAnimationPayloadData> UInterchangeUSDTranslator::GetAni
 			AnimationPayloads.Append(AnimationPayload);
 		}
 	}
+
+	// Import raw transforms payloads
+	if (int32 TransformPayloadCount = TransformQueryIndexes.Num(); TransformPayloadCount > 0)
+	{
+		TransformAnimationPayloads.Reserve(TransformPayloadCount);
+		for (int32 QueryIndex : TransformQueryIndexes)
+		{
+			const UE::Interchange::FAnimationPayloadQuery& PayloadQuery = PayloadQueries[QueryIndex];
+			FAnimationPayloadData AnimationPayLoadData{PayloadQuery.SceneNodeUniqueID, PayloadQuery.PayloadKey};
+			if (GetTransformAnimationPayloadData(*Impl, PayloadQuery.PayloadKey.UniqueId, AnimationPayLoadData))
+			{
+				TransformAnimationPayloads.Emplace(TArray<UE::Interchange::FAnimationPayloadData>{AnimationPayLoadData});
+			}
+		}
+
+		// Append the transforms results
+		for (TArray<UE::Interchange::FAnimationPayloadData>& AnimationPayload : TransformAnimationPayloads)
+		{
+			AnimationPayloads.Append(AnimationPayload);
+		}
+	}
 #endif	  // USE_USD_SDK
 
 	return AnimationPayloads;
diff --git a/Engine/Plugins/Interchange/Runtime/Source/FactoryNodes/Private/InterchangeGeometryCacheFactoryNode.cpp b/Engine/Plugins/Interchange/Runtime/Source/FactoryNodes/Private/InterchangeGeometryCacheFactoryNode.cpp
--- a/Engine/Plugins/Interchange/Runtime/Source/FactoryNodes/Private/InterchangeGeometryCacheFactoryNode.cpp
+++ b/Engine/Plugins/Interchange/Runtime/Source/FactoryNodes/Private/InterchangeGeometryCacheFactoryNode.cpp
@@ -7,8 +7,18 @@
 
 #include UE_INLINE_GENERATED_CPP_BY_NAME(InterchangeGeometryCacheFactoryNode)
 
+namespace UE::Interchange
+{
+	const FAttributeKey& FGeometryCacheNodeStaticData::GetSceneNodeAnimationPayloadKeyUidMapKey()
+	{
+		static FAttributeKey AttributeKey(TEXT("__SceneNodeAnimationPayloadKeyUidMap__"));
+		return AttributeKey;
+	}
+}
+
 UInterchangeGeometryCacheFactoryNode::UInterchangeGeometryCacheFactoryNode()
 {
+	SceneNodeAnimationPayloadKeyUidMap.Initialize(Attributes.ToSharedRef(), UE::Interchange::FGeometryCacheNodeStaticData::GetSceneNodeAnimationPayloadKeyUidMapKey().ToString());
 	AssetClass = nullptr;
 }
 
@@ -146,3 +156,13 @@ bool UInterchangeGeometryCacheFactoryNode::SetCustomOptimizeIndexBuffers(const b
 {
 	IMPLEMENT_NODE_ATTRIBUTE_SETTER_NODELEGATE(OptimizeIndexBuffers, bool);
 }
+
+void UInterchangeGeometryCacheFactoryNode::GetSceneNodeAnimationPayloadKeys(TMap<FString, FString>& OutSceneNodeAnimationPayloadKeyUids) const
+{
+	OutSceneNodeAnimationPayloadKeyUids = SceneNodeAnimationPayloadKeyUidMap.ToMap();
+}
+
+bool UInterchangeGeometryCacheFactoryNode::SetAnimationPayloadKeyForSceneNodeUid(const FString& SceneNodeUid, const FString& InUniqueId)
+{
+	return SceneNodeAnimationPayloadKeyUidMap.SetKeyValue(SceneNodeUid, InUniqueId);
+}
diff --git a/Engine/Plugins/Interchange/Runtime/Source/FactoryNodes/Public/InterchangeGeometryCacheFactoryNode.h b/Engine/Plugins/Interchange/Runtime/Source/FactoryNodes/Public/InterchangeGeometryCacheFactoryNode.h
--- a/Engine/Plugins/Interchange/Runtime/Source/FactoryNodes/Public/InterchangeGeometryCacheFactoryNode.h
+++ b/Engine/Plugins/Interchange/Runtime/Source/FactoryNodes/Public/InterchangeGeometryCacheFactoryNode.h
@@ -12,6 +12,14 @@
 
 #define UE_API INTERCHANGEFACTORYNODES_API
 
+namespace UE::Interchange
+{
+	struct FGeometryCacheNodeStaticData : public FBaseNodeStaticData
+	{
+		static UE_API const FAttributeKey& GetSceneNodeAnimationPayloadKeyUidMapKey();
+	};
+}
+
 UCLASS(MinimalAPI, BlueprintType)
 class UInterchangeGeometryCacheFactoryNode : public UInterchangeMeshFactoryNode
 {
@@ -103,6 +111,14 @@ public:
 	UFUNCTION(BlueprintCallable, Category = "Interchange | Node | GeometryCache")
 	UE_API bool SetCustomOptimizeIndexBuffers(const bool& AttributeValue);
 
+	/** Get the animation payload keys that are relevant to building the geometry cache */
+	UFUNCTION(BlueprintCallable, Category = "Interchange | Node | GeometryCache")
+	UE_API void GetSceneNodeAnimationPayloadKeys(TMap<FString, FString>& OutSceneNodeAnimationPayloadKeyUids) const;
+
+	/** Set the animation payload keys that are relevant to building the geometry cache */
+	UFUNCTION(BlueprintCallable, Category = "Interchange | Node | GeometryCache")
+	UE_API bool SetAnimationPayloadKeyForSceneNodeUid(const FString& SceneNodeUid, const FString& InUniqueId);
+
 	/**
 	 * Return the node type name of the class. This is used when reporting errors.
 	 */
@@ -129,6 +145,8 @@ private:
 	IMPLEMENT_NODE_ATTRIBUTE_KEY(ApplyConstantTopologyOptimization);
 	IMPLEMENT_NODE_ATTRIBUTE_KEY(StoreImportedVertexNumbers);
 	IMPLEMENT_NODE_ATTRIBUTE_KEY(OptimizeIndexBuffers);
+
+	UE::Interchange::TMapAttributeHelper<FString, FString> SceneNodeAnimationPayloadKeyUidMap;
 };
 
 #undef UE_API
diff --git a/Engine/Plugins/Interchange/Runtime/Source/Import/Private/Mesh/InterchangeGeometryCacheFactory.cpp b/Engine/Plugins/Interchange/Runtime/Source/Import/Private/Mesh/InterchangeGeometryCacheFactory.cpp
--- a/Engine/Plugins/Interchange/Runtime/Source/Import/Private/Mesh/InterchangeGeometryCacheFactory.cpp
+++ b/Engine/Plugins/Interchange/Runtime/Source/Import/Private/Mesh/InterchangeGeometryCacheFactory.cpp
@@ -2,6 +2,7 @@
 
 #include "Mesh/InterchangeGeometryCacheFactory.h"
 
+#include "Animation/InterchangeAnimationPayloadInterface.h"
 #include "Async/ParallelFor.h"
 #include "GeometryCache.h"
 #include "GeometryCacheCodecV1.h"
@@ -81,6 +82,55 @@ UClass* UInterchangeGeometryCacheFactory::GetFactoryClass() const
 	return UGeometryCache::StaticClass();
 }
 
+void UInterchangeGeometryCacheFactory::CreatePayloadTasks(const FImportAssetObjectParams& Arguments, bool bAsync, TArray<TSharedPtr<UE::Interchange::FInterchangeTaskBase>>& PayloadTasks)
+{
+	TRACE_CPUPROFILER_EVENT_SCOPE(UInterchangeGeometryCacheFactory::CreatePayloadTasks);
+
+	const IInterchangeAnimationPayloadInterface* AnimSequenceTranslatorPayloadInterface = Cast<IInterchangeAnimationPayloadInterface>(Arguments.Translator);
+	if (!AnimSequenceTranslatorPayloadInterface)
+	{
+		return;
+	}
+
+	const UInterchangeGeometryCacheFactoryNode* GeometryCacheFactoryNode = Cast<UInterchangeGeometryCacheFactoryNode>(Arguments.AssetNode);
+	if (!GeometryCacheFactoryNode)
+	{
+		return;
+	}
+
+	TMap<FString, FString> PayloadKeys;
+	GeometryCacheFactoryNode->GetSceneNodeAnimationPayloadKeys(PayloadKeys);
+	TransformAnimationPayloadResults.Reserve(PayloadKeys.Num());
+
+	// Create one payload task per query
+	for (const TPair<FString, FString>& SceneNodeUidAndPayloadKey : PayloadKeys)
+	{
+		const FString& SceneNodeUid = SceneNodeUidAndPayloadKey.Key;
+
+		// Force the payload key type to TRANSFORM since we need transforms, not curves
+		FInterchangeAnimationPayLoadKey PayloadKey(SceneNodeUidAndPayloadKey.Value, EInterchangeAnimationPayLoadType::GEOMETRY_CACHE_TRANSFORM);
+		UE::Interchange::FAnimationPayloadQuery AnimationPayloadQuery(SceneNodeUid, PayloadKey);
+
+		// Allocate the results so we do not need any mutex in the callback
+		TransformAnimationPayloadResults.Add(SceneNodeUid, UE::Interchange::FAnimationPayloadData(SceneNodeUid, PayloadKey));
+
+		TArray<UE::Interchange::FAnimationPayloadQuery> SingleQuery{AnimationPayloadQuery};
+		TSharedPtr<UE::Interchange::FInterchangeTaskLambda, ESPMode::ThreadSafe> GetAnimationPayloadTask = MakeShared<UE::Interchange::FInterchangeTaskLambda, ESPMode::ThreadSafe>(
+			bAsync ? UE::Interchange::EInterchangeTaskThread::AsyncThread : UE::Interchange::EInterchangeTaskThread::GameThread,
+			[SingleQuery, &TransformAnimationPayloadResults = TransformAnimationPayloadResults, AnimSequenceTranslatorPayloadInterface]()
+			{
+				TRACE_CPUPROFILER_EVENT_SCOPE(UInterchangeGeometryCacheFactory::GetAnimationPayloadSingleQuery);
+
+				TArray<UE::Interchange::FAnimationPayloadData> AnimationPayloads = AnimSequenceTranslatorPayloadInterface->GetAnimationPayloadData(SingleQuery);
+				for (const UE::Interchange::FAnimationPayloadData& PayloadData : AnimationPayloads)
+				{
+					TransformAnimationPayloadResults[PayloadData.SceneNodeUniqueID] = PayloadData;
+				}
+			});
+		PayloadTasks.Add(GetAnimationPayloadTask);
+	}
+}
+
 UInterchangeFactoryBase::FImportAssetResult UInterchangeGeometryCacheFactory::BeginImportAsset_GameThread(const FImportAssetObjectParams& Arguments)
 {
 	TRACE_CPUPROFILER_EVENT_SCOPE(UInterchangeGeometryCacheFactory::BeginImportAsset_GameThread);
@@ -155,11 +205,74 @@ UInterchangeFactoryBase::FImportAssetResult UInterchangeGeometryCacheFactory::Be
 		CommonPipelineDataFactoryNode->GetCustomGlobalOffsetTransform(GlobalOffsetTransform);
 	}
 
+	// Helper class to get the transform for a mesh at a given frame number
+	// The animated transforms from the payloads are local so they must be
+	// composed together from leaf to root to get the world transform
+	class FAnimatedTransformStack
+	{
+	public:
+		FAnimatedTransformStack() = default;
+
+		FAnimatedTransformStack(const FTransform& InDefaultTransform)
+		: DefaultTransform(InDefaultTransform)
+		{
+		}
+
+		FTransform GetTransform(int32 FrameNumber) const
+		{
+			if (TransformStack.Num() == 0)
+			{
+				return DefaultTransform;
+			}
+			else
+			{
+				FTransform ComposedTransform;
+				for (int32 Index = 0; Index < TransformStack.Num(); ++Index)
+				{
+					int32 NormalizedFrameIndex = FrameNumber - FrameIndexOffsets[Index];
+					const TArray<FTransform>& CurrentTransforms = *TransformStack[Index];
+					if (CurrentTransforms.IsValidIndex(NormalizedFrameIndex))
+					{
+						ComposedTransform *= CurrentTransforms[NormalizedFrameIndex];
+					}
+					else
+					{
+						ComposedTransform *= NormalizedFrameIndex >= CurrentTransforms.Num() ? CurrentTransforms.Last() : CurrentTransforms[0];
+					}
+				}
+				return ComposedTransform;
+			}
+		}
+
+		// Add static transform
+		void AddTransform(const FTransform& InTransform)
+		{
+			// Create a transform array with a single transform that it owns to free on destruction
+			TUniquePtr<TArray<FTransform>>& Transform = OwnedTransforms.Emplace_GetRef(MakeUnique<TArray<FTransform>>(TArray<FTransform>({InTransform})));
+			TransformStack.Add(Transform.Get());
+			FrameIndexOffsets.Add(0);
+		}
+
+		// Add animated transform
+		void AddTransform(const TArray<FTransform>& InTransforms, int32 IndexOffset)
+		{
+			// Keep a reference to an array of transforms
+			TransformStack.Add(&InTransforms);
+			FrameIndexOffsets.Add(IndexOffset);
+		}
+
+	private:
+		TArray<const TArray<FTransform>*> TransformStack;
+		TArray<int32> FrameIndexOffsets;
+		TArray<TUniquePtr<TArray<FTransform>>> OwnedTransforms;
+		FTransform DefaultTransform;
+	};
+
 	struct FNodeInfo
 	{
 		FString Uid;
-		const UInterchangeGeometryCacheNode* Node = nullptr;
-		FTransform GlobalTransform;
+		const UInterchangeMeshNode* Node = nullptr;
+		FAnimatedTransformStack TransformStack;
 		double FrameRate = 24.0;
 		bool bConstantTopology = true;
 	};
@@ -179,16 +292,25 @@ UInterchangeFactoryBase::FImportAssetResult UInterchangeGeometryCacheFactory::Be
 		FTransform GlobalMeshTransform;
 		const UInterchangeBaseNode* Node = Arguments.NodeContainer->GetNode(MeshUid);
 		const UInterchangeGeometryCacheNode* GeometryCacheNode = Cast<const UInterchangeGeometryCacheNode>(Node);
+		const UInterchangeMeshNode* MeshNode = nullptr;
+		const UInterchangeSceneNode* SceneNode = nullptr;
 		if (GeometryCacheNode == nullptr)
 		{
 			// MeshUid must refer to a scene node
-			const UInterchangeSceneNode* SceneNode = Cast<const UInterchangeSceneNode>(Node);
+			SceneNode = Cast<const UInterchangeSceneNode>(Node);
 			if (!ensure(SceneNode))
 			{
 				UE_LOG(LogInterchangeImport, Warning, TEXT("Invalid mesh reference when importing GeometryCache asset %s."), *AssetName);
 				continue;
 			}
 
+			// Ignore invisible meshes
+			bool bVisible = true;
+			if (SceneNode->GetCustomActorVisibility(bVisible) && !bVisible)
+			{
+				continue;
+			}
+
 			// Get the transform from the scene node
 			FTransform SceneNodeGlobalTransform;
 			if (SceneNode->GetCustomGlobalTransform(Arguments.NodeContainer, GlobalOffsetTransform, SceneNodeGlobalTransform))
@@ -204,6 +326,7 @@ UInterchangeFactoryBase::FImportAssetResult UInterchangeGeometryCacheFactory::Be
 			FString MeshDependencyUid;
 			SceneNode->GetCustomAssetInstanceUid(MeshDependencyUid);
 			GeometryCacheNode = Cast<UInterchangeGeometryCacheNode>(Arguments.NodeContainer->GetNode(MeshDependencyUid));
+			MeshNode = Cast<UInterchangeMeshNode>(Arguments.NodeContainer->GetNode(MeshDependencyUid));
 		}
 		else
 		{
@@ -211,30 +334,56 @@ UInterchangeFactoryBase::FImportAssetResult UInterchangeGeometryCacheFactory::Be
 			GlobalMeshTransform = GlobalOffsetTransform;
 		}
 
-		if (!ensure(GeometryCacheNode))
+		if (!ensure(GeometryCacheNode || MeshNode))
 		{
 			UE_LOG(LogInterchangeImport, Warning, TEXT("Invalid mesh reference when importing GeometryCache asset %s."), *AssetName);
 			continue;
 		}
 
+		// Build the transform stack for the mesh
+		// The global mesh transform will be used as the default if there's no animated transforms
+		FAnimatedTransformStack TransformStack(GlobalMeshTransform);
+		if (TransformAnimationPayloadResults.Num() > 0)
+		{
+			// Add the local transform to the stack from leaf to root
+			const UInterchangeSceneNode* LoopNode = SceneNode;
+			while (LoopNode)
+			{
+				// Use the animated transforms associated with the scene node, otherwise use its local transform
+				UE::Interchange::FAnimationPayloadData* Payload = TransformAnimationPayloadResults.Find(LoopNode->GetUniqueID());
+				if (Payload)
+				{
+					TransformStack.AddTransform(Payload->Transforms, Payload->RangeStartTime);
+				}
+				else if (FTransform LocalTransform; LoopNode->GetCustomLocalTransform(LocalTransform))
+				{
+					TransformStack.AddTransform(LocalTransform);
+				}
+				LoopNode = Cast<UInterchangeSceneNode>(Arguments.NodeContainer->GetNode(LoopNode->GetParentUid()));
+			}
+		}
+
 		// Cache the node info for the validated node
 		FNodeInfo& NodeInfo = ValidatedNodes.Emplace_GetRef();
 		NodeInfo.Uid = MeshUid;
-		NodeInfo.Node = GeometryCacheNode;
-		NodeInfo.GlobalTransform = GlobalMeshTransform;
+		NodeInfo.Node = MeshNode;
+		NodeInfo.TransformStack = MoveTemp(TransformStack);
 
-		GeometryCacheNode->GetCustomHasConstantTopology(NodeInfo.bConstantTopology);
-		bGloballyConstantTopology &= NodeInfo.bConstantTopology;
+		if (GeometryCacheNode)
+		{
+			GeometryCacheNode->GetCustomHasConstantTopology(NodeInfo.bConstantTopology);
+			bGloballyConstantTopology &= NodeInfo.bConstantTopology;
 
-		int32 MeshFrameStart = 0;
-		GeometryCacheNode->GetCustomStartFrame(MeshFrameStart);
-		FrameStart = FMath::Max(FrameStart, MeshFrameStart);
+			int32 MeshFrameStart = 0;
+			GeometryCacheNode->GetCustomStartFrame(MeshFrameStart);
+			FrameStart = FMath::Max(FrameStart, MeshFrameStart);
 
-		int32 MeshFrameEnd = 0;
-		GeometryCacheNode->GetCustomEndFrame(MeshFrameEnd);
-		FrameEnd = FMath::Min(FrameEnd, MeshFrameEnd);
+			int32 MeshFrameEnd = 0;
+			GeometryCacheNode->GetCustomEndFrame(MeshFrameEnd);
+			FrameEnd = FMath::Min(FrameEnd, MeshFrameEnd);
 
-		GeometryCacheNode->GetCustomFrameRate(NodeInfo.FrameRate);
+			GeometryCacheNode->GetCustomFrameRate(NodeInfo.FrameRate);
+		}
 	}
 
 	// Apply the time range overrides if enabled
@@ -358,7 +507,7 @@ UInterchangeFactoryBase::FImportAssetResult UInterchangeGeometryCacheFactory::Be
 	FScopedSlowTask SlowTask(static_cast<float>(FrameEnd - FrameStart), FText::FromString(Title));
 	SlowTask.MakeDialog(true);
 
-	const int32 MaxWorkerThreads = FTaskGraphInterface::Get().GetNumWorkerThreads();
+	const int32 MaxWorkerThreads = FApp::ShouldUseThreadingForPerformance() ? FTaskGraphInterface::Get().GetNumWorkerThreads() : 1;
 	const int32 NumFrameThreads = FMath::Clamp(MaxWorkerThreads, 1, GGeometryCacheParallelFrameReads);
 
 	// The GeometryCache data is processed in the game thread since the payloads are queried in
@@ -377,9 +526,9 @@ UInterchangeFactoryBase::FImportAssetResult UInterchangeGeometryCacheFactory::Be
 			TArray<FMeshDescription> MeshDescriptions;
 			for (const FNodeInfo& NodeInfo : ValidatedNodes)
 			{
-				const UInterchangeGeometryCacheNode* GeometryCacheNode = NodeInfo.Node;
+				const UInterchangeMeshNode* MeshNode = NodeInfo.Node;
 
-				TOptional<FInterchangeMeshPayLoadKey> OptionalPayLoadKey = GeometryCacheNode->GetPayLoadKey();
+				TOptional<FInterchangeMeshPayLoadKey> OptionalPayLoadKey = MeshNode->GetPayLoadKey();
 				if (!ensure(OptionalPayLoadKey.IsSet()))
 				{
 					MeshDescriptions.Add({});
@@ -388,8 +537,11 @@ UInterchangeFactoryBase::FImportAssetResult UInterchangeGeometryCacheFactory::Be
 				}
 
 				FInterchangeMeshPayLoadKey PayLoadKey(OptionalPayLoadKey->UniqueId, FrameIndex);
-				PayloadAttributes.RegisterAttribute(UE::Interchange::FAttributeKey{ UE::Interchange::MeshPayload::Attributes::MeshGlobalTransform }, NodeInfo.GlobalTransform);
-				TOptional<UE::Interchange::FMeshPayloadData> PayloadData = MeshTranslatorPayloadInterface->GetMeshPayloadData(PayLoadKey, PayloadAttributes);
+
+				// Use a copy of the payload attributes since the transform will vary for each mesh and frame
+				UE::Interchange::FAttributeStorage PayloadAttributesCopy(PayloadAttributes);
+				PayloadAttributesCopy.RegisterAttribute(UE::Interchange::FAttributeKey{ UE::Interchange::MeshPayload::Attributes::MeshGlobalTransform }, NodeInfo.TransformStack.GetTransform(FrameIndex));
+				TOptional<UE::Interchange::FMeshPayloadData> PayloadData = MeshTranslatorPayloadInterface->GetMeshPayloadData(PayLoadKey, PayloadAttributesCopy);
 
 				if (PayloadData.IsSet() && !PayloadData->MeshDescription.IsEmpty())
 				{
diff --git a/Engine/Plugins/Interchange/Runtime/Source/Import/Public/Mesh/InterchangeGeometryCacheFactory.h b/Engine/Plugins/Interchange/Runtime/Source/Import/Public/Mesh/InterchangeGeometryCacheFactory.h
--- a/Engine/Plugins/Interchange/Runtime/Source/Import/Public/Mesh/InterchangeGeometryCacheFactory.h
+++ b/Engine/Plugins/Interchange/Runtime/Source/Import/Public/Mesh/InterchangeGeometryCacheFactory.h
@@ -2,6 +2,7 @@
 
 #pragma once
 
+#include "InterchangeCommonAnimationPayload.h"
 #include "InterchangeFactoryBase.h"
 #include "Templates/PimplPtr.h"
 
@@ -22,6 +23,7 @@ public:
 
 	UE_API virtual UClass* GetFactoryClass() const override;
 	virtual EInterchangeFactoryAssetType GetFactoryAssetType() override { return EInterchangeFactoryAssetType::Meshes; }
+	UE_API virtual void CreatePayloadTasks(const FImportAssetObjectParams& Arguments, bool bAsync, TArray<TSharedPtr<UE::Interchange::FInterchangeTaskBase>>& PayloadTasks) override;
 	UE_API virtual FImportAssetResult BeginImportAsset_GameThread(const FImportAssetObjectParams& Arguments) override;
 	UE_API virtual FImportAssetResult ImportAsset_Async(const FImportAssetObjectParams& Arguments) override;
 	UE_API virtual FImportAssetResult EndImportAsset_GameThread(const FImportAssetObjectParams& Arguments) override;
@@ -38,6 +40,9 @@ public:
 private:
 	TArray<TObjectPtr<UGeometryCacheTrack>> Tracks;
 	TPimplPtr<class FGeometryCacheComponentResetAsset> ResetAssetOnReimport;
+
+	// Animation queries results
+	TMap<FString, UE::Interchange::FAnimationPayloadData> TransformAnimationPayloadResults;
 };
 
 #undef UE_API
diff --git a/Engine/Plugins/Interchange/Runtime/Source/Nodes/Public/InterchangeAnimationTrackSetNode.h b/Engine/Plugins/Interchange/Runtime/Source/Nodes/Public/InterchangeAnimationTrackSetNode.h
--- a/Engine/Plugins/Interchange/Runtime/Source/Nodes/Public/InterchangeAnimationTrackSetNode.h
+++ b/Engine/Plugins/Interchange/Runtime/Source/Nodes/Public/InterchangeAnimationTrackSetNode.h
@@ -18,13 +18,14 @@ enum class EInterchangeAnimationPayLoadType : uint8
 	MORPHTARGETCURVE, // Handles/generates the same properties as the CURVE variation, but the way it is acquired might be different (depending on the Translator).
 	STEPCURVE,
 	BAKED,
-	MORPHTARGETCURVEWEIGHTINSTANCE	//Handled within UInterchangeAnimSequenceFactory, contrary to the others which are handled in the Translators
+	MORPHTARGETCURVEWEIGHTINSTANCE,	//Handled within UInterchangeAnimSequenceFactory, contrary to the others which are handled in the Translators
 									// Purpose of this type is to generate a 1 frame long Animation with the Instantiated Morph Target Curve Weights.
 									//		This is needed for the special use cases where the imported 3d file format has a concept of MorphTargetWeight usages on StaticMeshes.
 									//		UE5 does not have support for this particular concept, and the workaround is to create a 1 frame long animation with the desired MorphTargetWeight settings.
 									//		This is also needed for the use cases where the import is a Level import and the 3d file format is instantiating a StaticMesh with a particular MorphTargetWeight settings.
 									//			This is done through the AnimSequnce being used on the SkeletalMeshActor's instance.
 									//		Related Ticket: UE-186102
+	GEOMETRY_CACHE_TRANSFORM // Array of time-sampled local transforms used to build a geometry cache
 };
 
 
diff --git a/Engine/Plugins/Interchange/Runtime/Source/Pipelines/Private/InterchangeGenericGeometryCachePipeline.cpp b/Engine/Plugins/Interchange/Runtime/Source/Pipelines/Private/InterchangeGenericGeometryCachePipeline.cpp
--- a/Engine/Plugins/Interchange/Runtime/Source/Pipelines/Private/InterchangeGenericGeometryCachePipeline.cpp
+++ b/Engine/Plugins/Interchange/Runtime/Source/Pipelines/Private/InterchangeGenericGeometryCachePipeline.cpp
@@ -2,6 +2,7 @@
 
 #include "InterchangeGenericMeshPipeline.h"
 
+#include "InterchangeAnimationTrackSetNode.h"
 #include "InterchangeGeometryCacheFactoryNode.h"
 #include "InterchangeMeshNode.h"
 #include "InterchangePipelineMeshesUtilities.h"
@@ -20,10 +21,91 @@ void UInterchangeGenericMeshPipeline::ExecutePreImportPipelineGeometryCache()
 		return;
 	}
 
+	struct FPayloadInfo
+	{
+		FString SceneNodeUid;
+		FString PayloadKeyId;
+
+		FPayloadInfo(const FString& InSceneNodeUid, const FString& InPayloadKeyId)
+		: SceneNodeUid(InSceneNodeUid)
+		, PayloadKeyId(InPayloadKeyId)
+		{ }
+
+	};
+
+	TArray<FPayloadInfo> AnimPayloads;
+
+	// Retrieve all raw transform animations
+	TArray<UInterchangeAnimationTrackSetNode*> TrackSetNodes;
+	BaseNodeContainer->IterateNodesOfType<UInterchangeAnimationTrackSetNode>([&](const FString& NodeUid, UInterchangeAnimationTrackSetNode* Node)
+	{
+		TrackSetNodes.Add(Node);
+	});
+
+	for (UInterchangeAnimationTrackSetNode* TrackSetNode : TrackSetNodes)
+	{
+		if (!TrackSetNode)
+		{
+			continue;
+		}
+
+		TArray<FString> AnimationTrackUids;
+		TrackSetNode->GetCustomAnimationTrackUids(AnimationTrackUids);
+
+		for (const FString& AnimationTrackUid : AnimationTrackUids)
+		{
+			const UInterchangeTransformAnimationTrackNode* TransformTrackNode = Cast<UInterchangeTransformAnimationTrackNode>(BaseNodeContainer->GetNode(AnimationTrackUid));
+			if (!TransformTrackNode)
+			{
+				continue;
+			}
+
+			FString ActorNodeUid;
+			if (!TransformTrackNode->GetCustomActorDependencyUid(ActorNodeUid))
+			{
+				continue;
+			}
+
+			const UInterchangeSceneNode* SceneNode = Cast<UInterchangeSceneNode>(BaseNodeContainer->GetNode(ActorNodeUid));
+			if (!SceneNode)
+			{
+				continue;
+			}
+
+			FInterchangeAnimationPayLoadKey AnimationPayloadKey;
+			if (TransformTrackNode->GetCustomAnimationPayloadKey(AnimationPayloadKey))
+			{
+				AnimPayloads.Emplace(FPayloadInfo(SceneNode->GetUniqueID(), AnimationPayloadKey.UniqueId));
+			}
+		}
+	}
+
 	// If there's an animated mesh, combine all meshes into a single Geometry Cache
 	TArray<FString> MeshInstanceUids;
 	PipelineMeshesUtilities->GetAllGeometryCacheInstance(MeshInstanceUids);
 
+	if (MeshInstanceUids.Num() == 0)
+	{
+		return;
+	}
+
+	// With animated transforms, static/skeletal meshes can get baked into a geometry cache
+	if (AnimPayloads.Num() > 0)
+	{
+		// The bAutoDetectMeshType setting converts static meshes with animated transforms into skeletal meshes with AnimSequence
+		// For the purpose of geometry cache, both are considered as just meshes
+		TArray<FString> AnimMeshInstanceUids;
+		if (CommonMeshesProperties->bAutoDetectMeshType && CommonMeshesProperties->ForceAllMeshAsType == EInterchangeForceMeshType::IFMT_None)
+		{
+			PipelineMeshesUtilities->GetAllSkinnedMeshInstance(AnimMeshInstanceUids);
+		}
+		else
+		{
+			PipelineMeshesUtilities->GetAllStaticMeshInstance(AnimMeshInstanceUids);
+		}
+		MeshInstanceUids.Append(AnimMeshInstanceUids);
+	}
+
 	TSet<FString> VisitedMeshAssetUids;
 
 	TArray<FString> MeshUids;
@@ -72,6 +154,11 @@ void UInterchangeGenericMeshPipeline::ExecutePreImportPipelineGeometryCache()
 
 	UInterchangeGeometryCacheFactoryNode* GeometryCacheFactoryNode = CreateGeometryCacheFactoryNode(MeshUids);
 	GeometryCacheFactoryNodes.Add(GeometryCacheFactoryNode);
+
+	for (const FPayloadInfo& PayloadInfo : AnimPayloads)
+	{
+		GeometryCacheFactoryNode->SetAnimationPayloadKeyForSceneNodeUid(PayloadInfo.SceneNodeUid, PayloadInfo.PayloadKeyId);
+	}
 }
 
 UInterchangeGeometryCacheFactoryNode* UInterchangeGenericMeshPipeline::CreateGeometryCacheFactoryNode(const TArray<FString>& MeshUids)
