[Interchange][FBX] UFBX failing FBX asset import tests

- fixed edge smoothing

- fixed uid uniqueness for nodes, meshes, materials, bones

- fixed Opacity handling

- fixed - not creating materials slots for empty material parts UE-336774

Tests fixed:
Editor.Interchange.Import.Asset.FBX.SM_MaterialNameClash
Editor.Interchange.Import.Asset.FBX.SKM_Cube_Blue
Editor.Interchange.Import.Asset.FBX.SM_MaterialOrder
Editor.Interchange.Import.Level.FBX.TP_Scene_Basic

diff --git a/Engine/Plugins/Interchange/Runtime/Source/Parsers/Fbx/Private/Ufbx/UfbxConvert.cpp b/Engine/Plugins/Interchange/Runtime/Source/Parsers/Fbx/Private/Ufbx/UfbxConvert.cpp
--- a/Engine/Plugins/Interchange/Runtime/Source/Parsers/Fbx/Private/Ufbx/UfbxConvert.cpp
+++ b/Engine/Plugins/Interchange/Runtime/Source/Parsers/Fbx/Private/Ufbx/UfbxConvert.cpp
@@ -9,13 +9,14 @@
 namespace UE::Interchange::Private
 {
 
-Convert::FMeshNameAndUid::FMeshNameAndUid(FUfbxParser& Parser, const ufbx_element& Mesh): Name(Parser.GetMeshName(Mesh))
-                                                                                             , UniqueID(Parser.GetMeshUid(Mesh))
+Convert::FMeshNameAndUid::FMeshNameAndUid(FUfbxParser& Parser, const ufbx_element& Mesh)
+	: Label(Parser.GetMeshLabel(Mesh))
+	, UniqueID(Parser.GetMeshUid(Mesh))
 {
 }
 
 Convert::FNodeNameAndUid::FNodeNameAndUid(FUfbxParser& Parser, UInterchangeSceneNode* ParentSceneNode, const ufbx_node& Node)
-	: Name(ParentSceneNode ?  Parser.GetElementName(Node.element) : TEXT("RootNode"))
+	: Label(ParentSceneNode ?  Parser.GetElementName(Node.element) : TEXT("RootNode"))
 	, UniqueID(Parser.GetNodeUid(Node))
 {
 }
diff --git a/Engine/Plugins/Interchange/Runtime/Source/Parsers/Fbx/Private/Ufbx/UfbxConvert.h b/Engine/Plugins/Interchange/Runtime/Source/Parsers/Fbx/Private/Ufbx/UfbxConvert.h
--- a/Engine/Plugins/Interchange/Runtime/Source/Parsers/Fbx/Private/Ufbx/UfbxConvert.h
+++ b/Engine/Plugins/Interchange/Runtime/Source/Parsers/Fbx/Private/Ufbx/UfbxConvert.h
@@ -26,7 +26,7 @@ namespace Convert
 		FMeshNameAndUid(FUfbxParser& Parser, const ufbx_mesh& Mesh) : FMeshNameAndUid(Parser, Mesh.element) {}
 		FMeshNameAndUid(FUfbxParser& Parser, const ufbx_blend_shape& Shape) : FMeshNameAndUid(Parser, Shape.element) {}
 
-		FString Name;
+		FString Label;
 		FString UniqueID;
 	};
 
@@ -36,7 +36,7 @@ namespace Convert
 		//		But the root node doesn't have actual name, in fact. Also it might be better to name this node after the file name?
 		FNodeNameAndUid(FUfbxParser& Parser, UInterchangeSceneNode* ParentSceneNode, const ufbx_node& Node);
 
-		FString Name;
+		FString Label;
 		FString UniqueID;
 	};
 
diff --git a/Engine/Plugins/Interchange/Runtime/Source/Parsers/Fbx/Private/Ufbx/UfbxMaterial.cpp b/Engine/Plugins/Interchange/Runtime/Source/Parsers/Fbx/Private/Ufbx/UfbxMaterial.cpp
--- a/Engine/Plugins/Interchange/Runtime/Source/Parsers/Fbx/Private/Ufbx/UfbxMaterial.cpp
+++ b/Engine/Plugins/Interchange/Runtime/Source/Parsers/Fbx/Private/Ufbx/UfbxMaterial.cpp
@@ -129,30 +129,24 @@ namespace UE::Interchange::Private
 		UInterchangeBaseNode& MaterialInstance;
 	};
 
-	UInterchangeMaterialInstanceNode* CreateMaterialInstanceNode(UInterchangeBaseNodeContainer& NodeContainer, const FString& NodeUID, const FString& NodeName)
+	UInterchangeMaterialInstanceNode* CreateMaterialInstanceNode(UInterchangeBaseNodeContainer& NodeContainer, const FString& NodeUID, const FString& NodeLabel)
 	{
 		UInterchangeMaterialInstanceNode* MaterialInstanceNode = NewObject<UInterchangeMaterialInstanceNode>(&NodeContainer);
-		NodeContainer.SetupNode(MaterialInstanceNode, NodeUID, NodeName, EInterchangeNodeContainerType::TranslatedAsset);
+		NodeContainer.SetupNode(MaterialInstanceNode, NodeUID, NodeLabel, EInterchangeNodeContainerType::TranslatedAsset);
 
 		return MaterialInstanceNode;
 	}
 
 	void FUfbxMaterial::AddMaterials(UInterchangeBaseNodeContainer& NodeContainer)
 	{
+		using namespace UE::Interchange::Materials;
+
 		for (ufbx_material* Material : Parser.Scene->materials)
 		{
-			// #ufbx_todo: Uid in FFbxMaterial parser is made from sanitized name
-			// but two different materials  might happen to have same name AFTER sanitation.
-			// SO material name shoul dbe made unique for this case. E.g. by adding number after name string.
-			// This means that we need to memoize mapping of material to this generated name to assign proper names to meshes
-			// FString MaterialName = Parser.GetFbxHelper()->GetFbxObjectName(SurfaceMaterial);
-			// FString NodeUid = TEXT("\\Material\\") + MaterialName;
-			FString MaterialName = Parser.GetMaterialName(*Material).ToString();
-			
-			using namespace UE::Interchange::Materials;
+			FString MaterialLabel = Parser.GetMaterialLabel(*Material).ToString();
 			
-			FString MaterialUid = Parser.GetMaterialUid(MaterialName);
-			UInterchangeMaterialInstanceNode* MaterialInstanceNode = CreateMaterialInstanceNode(NodeContainer, MaterialUid, MaterialName);
+			FString MaterialUid = Parser.GetMaterialUid(*Material);
+			UInterchangeMaterialInstanceNode* MaterialInstanceNode = CreateMaterialInstanceNode(NodeContainer, MaterialUid, MaterialLabel);
 
 			// WithUVTransform
 			bool bHasUVTransform = false;
@@ -161,8 +155,6 @@ namespace UE::Interchange::Private
 
 			FMaterialMapConverter MaterialMapConverter(Parser, NodeContainer, *MaterialInstanceNode);
 			
-			using namespace UE::Interchange::Materials;
-
 			auto SetColorFromMap = [&MaterialInstanceNode](FName InputName, const ufbx_material_map& Map, float Factor)
 			{
 				if (Map.value_components >=3)
@@ -322,6 +314,7 @@ namespace UE::Interchange::Private
 
 			// Opacity
 			// Connect only if transparency is either a texture or different from 0.f
+			// Check fbx.transparency first...
 			if (Material->fbx.transparency_color.has_value || Material->fbx.transparency_factor.has_value)
 			{
 				FName InputName = Phong::Parameters::Opacity;
@@ -354,8 +347,8 @@ namespace UE::Interchange::Private
 				}
 				else
 				{
-					const float OpacityScalar = 1.f - Factor;
-					if (OpacityScalar < 1.f)
+					const float OpacityScalar = 1.f - Factor*Material->fbx.transparency_color.value_real;
+					if (OpacityScalar+UE_KINDA_SMALL_NUMBER < 1.f)
 					{
 						ParentMaterial = TEXT("PhongSurfaceTranslucentMaterial");
 						MaterialInstanceNode->AddScalarParameterValue(Phong::Parameters::Opacity.ToString(), OpacityScalar);
diff --git a/Engine/Plugins/Interchange/Runtime/Source/Parsers/Fbx/Private/Ufbx/UfbxMesh.cpp b/Engine/Plugins/Interchange/Runtime/Source/Parsers/Fbx/Private/Ufbx/UfbxMesh.cpp
--- a/Engine/Plugins/Interchange/Runtime/Source/Parsers/Fbx/Private/Ufbx/UfbxMesh.cpp
+++ b/Engine/Plugins/Interchange/Runtime/Source/Parsers/Fbx/Private/Ufbx/UfbxMesh.cpp
@@ -190,7 +190,7 @@ namespace UE::Interchange::Private
 				continue;
 			}
 
-			UInterchangeMeshNode* MeshNode = CreateMeshNode(NodeContainer, MeshNameAndUid.Name, MeshNameAndUid.UniqueID);
+			UInterchangeMeshNode* MeshNode = CreateMeshNode(NodeContainer, MeshNameAndUid.Label, MeshNameAndUid.UniqueID);
 
 			if (Mesh->skin_deformers.count > 0)
 			{
@@ -244,11 +244,15 @@ namespace UE::Interchange::Private
 			MeshNode->SetCustomUVCount(MeshUVCount);
 
 			// As long as there are materials then ufbx_mesh::material_parts indexing corresponds to ufbx_mesh::materials
-			for (ufbx_material* Material: Mesh->materials)
+			for (int32 MaterialIndex = 0; MaterialIndex < Mesh->materials.count; ++MaterialIndex)
 			{
-				const FString MaterialName = Parser.GetMaterialName(*Material).ToString();
-				const FString MaterialUid = Parser.GetMaterialUid(MaterialName);
-				MeshNode->SetSlotMaterialDependencyUid(MaterialName, MaterialUid);
+				const ufbx_mesh_part& MeshPart = Mesh->material_parts[MaterialIndex];
+				// Skipping creating slot for empty material_parts as well as skipping creation of a polygon group below
+				if (MeshPart.num_faces > 0)
+				{
+					ufbx_material* Material =  Mesh->materials[MaterialIndex];
+					MeshNode->SetSlotMaterialDependencyUid(Parser.GetMaterialSlotName(Material).ToString(), Parser.GetMaterialUid(*Material));
+				}
 			}
 
 			for (ufbx_node* Instance : Mesh->instances)
@@ -321,7 +325,7 @@ namespace UE::Interchange::Private
 					}
 					if (!ExistingMorphTargetNode)
 					{
-						UInterchangeMeshNode* MorphTargetNode = CreateMeshNode(NodeContainer, MorphTargetNameId.Name, MorphTargetNameId.UniqueID);
+						UInterchangeMeshNode* MorphTargetNode = CreateMeshNode(NodeContainer, MorphTargetNameId.Label, MorphTargetNameId.UniqueID);
 						const bool bIsMorphTarget = true;
 						MorphTargetNode->SetMorphTarget(bIsMorphTarget);
 						MorphTargetNode->SetMorphTargetName(ShapeName);
@@ -432,6 +436,7 @@ namespace UE::Interchange::Private
 		constexpr bool bKeepExistingAttribute = true;
 		Attributes.Register(bKeepExistingAttribute);
 
+		// #ufbx_todo Revisit how FMeshDescription is filled up
 #ifdef WITH_MESH_DESCRIPTION_BUILDER
 		FMeshDescriptionBuilder Builder;
 		Builder.SetMeshDescription(&MeshDescription);
@@ -463,16 +468,20 @@ namespace UE::Interchange::Private
 		if (bMakeMeshGroupsPerMaterial)
 		{
 
-			for (int MaterialIndex = 0; MaterialIndex < Mesh.material_parts.count; ++MaterialIndex)
+			for (int32 MaterialIndex = 0; MaterialIndex < Mesh.material_parts.count; ++MaterialIndex)
 			{
-				ufbx_material* Material = MeshUtils::GetMeshMaterial(Mesh, MaterialIndex);
-				FPolygonGroupID GroupID = Builder.AppendPolygonGroup(Material ? Parser.GetMaterialName(*Material) : NAME_None);
-				ufbx_mesh_part& MeshPart = Mesh.material_parts[MaterialIndex];
-
-				const int32 FaceCount = MeshPart.num_faces;
-				for (int32 FaceIndex = 0; FaceIndex < FaceCount; ++FaceIndex)
+				// #ufbx_todo: maybe should specify slot name as just part's index? Why use materials name? Expect that tests verify this.
+				const ufbx_mesh_part& MeshPart = Mesh.material_parts[MaterialIndex];
+				if (MeshPart.num_faces > 0)
 				{
-					FaceCollector.AddFace(GroupID, Mesh.faces[MeshPart.face_indices.data[FaceIndex]]);
+					const ufbx_material* Material = MeshUtils::GetMeshMaterial(Mesh, MaterialIndex);
+					const FPolygonGroupID GroupID = Builder.AppendPolygonGroup(Parser.GetMaterialSlotName(Material));
+
+					const int32 FaceCount = MeshPart.num_faces;
+					for (int32 FaceIndex = 0; FaceIndex < FaceCount; ++FaceIndex)
+					{
+						FaceCollector.AddFace(GroupID, Mesh.faces[MeshPart.face_indices.data[FaceIndex]]);
+					}
 				}
 			}
 		}
@@ -488,7 +497,28 @@ namespace UE::Interchange::Private
 
 		if (FaceCollector.NumDegenerateTriangles)
 		{
-			Parser.AddMessage<UInterchangeResultMeshWarning_Generic>()->Text = FText::Format(LOCTEXT("PrimitiveFoundDegenerateTriangles", "Mesh '{0}' has {1} degenerate triangles"), FText::FromString(Parser.GetMeshName(Mesh.element)), FaceCollector.NumDegenerateTriangles);;
+			Parser.AddMessage<UInterchangeResultMeshWarning_Generic>()->Text = FText::Format(LOCTEXT("PrimitiveFoundDegenerateTriangles", "Mesh '{0}' has {1} degenerate triangles"), FText::FromString(Parser.GetMeshLabel(Mesh.element)), FaceCollector.NumDegenerateTriangles);;
+		}
+
+		// Provide edge smoothing
+		TEdgeAttributesRef<bool> EdgeHardnesses = Attributes.GetEdgeHardnesses();
+		for (int EdgeIndex = 0; EdgeIndex < Mesh.edges.count; ++EdgeIndex)
+		{
+			const ufbx_edge& Edge = Mesh.edges[EdgeIndex];
+
+			const FVertexID EdgeVertexA = FaceCollector.VertexMapping[Mesh.vertex_indices[Edge.a]];
+			const FVertexID EdgeVertexB = FaceCollector.VertexMapping[Mesh.vertex_indices[Edge.b]];
+
+			FEdgeID EdgeID = MeshDescription.GetVertexPairEdge(EdgeVertexA, EdgeVertexB);
+			// shouldn't edge already be here?
+			if (EdgeID == INDEX_NONE)
+			{
+				EdgeID = MeshDescription.CreateEdge(EdgeVertexA, EdgeVertexB);
+			}
+			if (EdgeIndex < Mesh.edge_smoothing.count)
+			{
+				EdgeHardnesses[EdgeID] = !Mesh.edge_smoothing[EdgeIndex];
+			}
 		}
 
 		HandleMesh(FaceCollector);
@@ -552,10 +582,10 @@ namespace UE::Interchange::Private
 				double Weight = SkinWeight.weight;
 
 				int32 BoneIndex = -1;
-				if (!SortedJoints.Find(BoneNode, BoneIndex))
+				if (BoneNode && !SortedJoints.Find(BoneNode, BoneIndex))
 				{
 					BoneIndex = SortedJoints.Add(BoneNode);
-					FString JointNodeUniqueID = Parser.GetElementName(BoneNode->element, true);
+					FString JointNodeUniqueID = Parser.GetBoneNodeUid(*BoneNode);
 					OutJointUniqueNames.Add(JointNodeUniqueID);
 				}
 
diff --git a/Engine/Plugins/Interchange/Runtime/Source/Parsers/Fbx/Private/Ufbx/UfbxParser.cpp b/Engine/Plugins/Interchange/Runtime/Source/Parsers/Fbx/Private/Ufbx/UfbxParser.cpp
--- a/Engine/Plugins/Interchange/Runtime/Source/Parsers/Fbx/Private/Ufbx/UfbxParser.cpp
+++ b/Engine/Plugins/Interchange/Runtime/Source/Parsers/Fbx/Private/Ufbx/UfbxParser.cpp
@@ -42,16 +42,43 @@ FUfbxParser::~FUfbxParser()
 	Reset();
 }
 
+FString FUfbxParser::GetElementNameRaw(const ufbx_element& Element) const
+{
+	return Convert::ToUnrealString(Element.name);
+}
+
+FString FUfbxParser::GetElementNameDeduplicated(const ufbx_element& Element, bool bIsJoint) const
+{
+	if (FString* Found = ElementToName.Find(&Element))
+	{
+		return *Found;
+	}
+
+	const FString BaseName = GetElementName(Element, bIsJoint);
+	FString Name = BaseName;
+	for (int32 Count = 0;;++Count)
+	{
+		if (!NameToElement.Contains(Name))
+		{
+			break;
+		}
+		Name = BaseName + TEXT("_ncl_") + FString::FromInt(Count+1);
+	}
+	NameToElement.Add(Name, &Element);
+	ElementToName.Add(&Element, Name);
+	return Name;
+}
+
 // inspired by UE::Interchange::Private::FFbxHelper::GetFbxObjectName
 FString FUfbxParser::GetElementName(const ufbx_element& Element, bool bIsJoint) const
 {
-	FString ElementName = Convert::ToUnrealString(Element.name);
+	FString ElementName = FUfbxParser::GetElementNameRaw(Element);
 	UE::Interchange::SanitizeName(ElementName, bIsJoint);
 	return ElementName;
 }
 
 // #ufbx_todo: see FFbxHelper::GetMeshName
-FString FUfbxParser::GetMeshName(const ufbx_element& Mesh) const
+FString FUfbxParser::GetMeshLabel(const ufbx_element& Mesh) const
 {
 	return UTF8_TO_TCHAR(Mesh.name.data);
 }
@@ -61,7 +88,7 @@ FString FUfbxParser::GetNodeUid(const ufbx_node& Node) const
 	// #ufbx_todo: check that element_id is enough
 	// reference FBX parser's uses UE::Interchange::Private::FFbxHelper::GetFbxNodeHierarchyName to create NodeUid
 
-	FString NodeUid = GetElementName(Node.element) + TEXT("_") + FString::FromInt(Node.element_id);
+	FString NodeUid = GetElementNameDeduplicated(Node.element);
 	
 	if (Node.parent)
 	{
@@ -77,20 +104,26 @@ FString FUfbxParser::GetNodeUid(const ufbx_node& Node) const
 	return NodeUid;
 }
 
-FName FUfbxParser::GetMaterialName(const ufbx_material& Material) const
+FString FUfbxParser::GetBoneNodeUid(const ufbx_node& Node) const
+{
+	// #ufbx_todo: need make sure not dup
+	return GetElementNameDeduplicated(Node.element, true);
+}
+
+FName FUfbxParser::GetMaterialLabel(const ufbx_material& Material) const
 {
-	// #ufbx_todo: ensure uniqueness when GetElementName returns same value for different materials
 	return *GetElementName(Material.element);
 }
 
-FString FUfbxParser::GetMaterialUid(const FString& MaterialName) const
+FString FUfbxParser::GetMaterialUid(const ufbx_material& Material) const
 {
-	return TEXT("\\Material\\") + MaterialName;
+	// Since materials names are not unique use its ufbx type_id and keep original name in the Uid for simpler debugging
+	return  FString::Printf(TEXT("\\Material\\%d_%s"), Material.element.typed_id, *GetElementNameRaw(Material.element));
 }
 
 FString FUfbxParser::GetMeshUid(const ufbx_element& Mesh) const
 {
-	return FString::FromInt(Mesh.element_id);
+	return GetElementNameDeduplicated(Mesh);
 }
 
 void FUfbxParser::ConvertProperty(const ufbx_prop& Prop, UInterchangeBaseNode* SceneNode)
@@ -183,6 +216,9 @@ void FUfbxParser::SetResultContainer(UInterchangeResultsContainer* Result)
 
 void FUfbxParser::Reset()
 {
+	NameToElement.Reset();
+	ElementToName.Reset();
+
 	ElementIdToSceneNode.Reset();
 
 	PayloadContexts.Reset();
diff --git a/Engine/Plugins/Interchange/Runtime/Source/Parsers/Fbx/Private/Ufbx/UfbxParser.h b/Engine/Plugins/Interchange/Runtime/Source/Parsers/Fbx/Private/Ufbx/UfbxParser.h
--- a/Engine/Plugins/Interchange/Runtime/Source/Parsers/Fbx/Private/Ufbx/UfbxParser.h
+++ b/Engine/Plugins/Interchange/Runtime/Source/Parsers/Fbx/Private/Ufbx/UfbxParser.h
@@ -135,12 +135,40 @@ namespace UE::Interchange::Private
 			Item->SourceAssetName = SourceFilename;
 		}
 
+		/**
+		 * Get sanitized element name, not unique
+		 * @param Element 
+		 * @param bIsJoint 
+		 * @return 
+		 */
 		FString GetElementName(const ufbx_element& Element, bool bIsJoint = false) const;
-		FString GetMeshName(const ufbx_element& Mesh) const;
+		FString GetElementNameRaw(const ufbx_element& Element) const;
+		/**
+		 * Returns name build for element making sure it's unique for each element
+		 * - name based on sanitized name of the element so it can be used in any place when sanitized string is needed
+		 * @param Element *
+		 * @return 
+		 */
+		FString GetElementNameDeduplicated(const ufbx_element& Element, bool bIsJoint = false) const;
+
 		FString GetNodeUid(const ufbx_node& Node) const;
+		FString GetBoneNodeUid(const ufbx_node& Node) const;
+
 		FString GetMeshUid(const ufbx_element& Node) const;
-		FName GetMaterialName(const ufbx_material& Material) const;
-		FString GetMaterialUid(const FString& Name) const;
+		FString GetMeshLabel(const ufbx_element& Mesh) const;
+
+		FString GetMaterialUid(const ufbx_material& Material) const;
+		FName GetMaterialLabel(const ufbx_material& Material) const;
+
+		FName GetMaterialSlotName(const ufbx_material* Material) const
+		{
+			if (!Material)
+			{
+				return NAME_None;
+			}
+
+			return {Material ? FName(*GetElementNameDeduplicated(Material->element)) : NAME_None};
+		}
 
 		static void ConvertProperty(const ufbx_prop& Prop, UInterchangeBaseNode* SceneNode);
 
@@ -235,5 +263,9 @@ namespace UE::Interchange::Private
 		TMap<uint32_t, UInterchangeSceneNode*> ElementIdToSceneNode;
 		FPayloadContexts PayloadContexts;
 
+		// Keeps track of unique names generated for elements
+		mutable TMap<FString, const ufbx_element*> NameToElement;
+		mutable TMap<const ufbx_element*, FString> ElementToName;
+
 	};
 }
diff --git a/Engine/Plugins/Interchange/Runtime/Source/Parsers/Fbx/Private/Ufbx/UfbxScene.cpp b/Engine/Plugins/Interchange/Runtime/Source/Parsers/Fbx/Private/Ufbx/UfbxScene.cpp
--- a/Engine/Plugins/Interchange/Runtime/Source/Parsers/Fbx/Private/Ufbx/UfbxScene.cpp
+++ b/Engine/Plugins/Interchange/Runtime/Source/Parsers/Fbx/Private/Ufbx/UfbxScene.cpp
@@ -21,7 +21,7 @@ namespace UE::Interchange::Private
 			return TrackNodeUid;
 		}
 
-		UInterchangeSceneNode* CreateTransformNode(FUfbxParser& Parser, UInterchangeBaseNodeContainer& NodeContainer, const FString& NodeName, const FString& NodeUID, const FTransform& LocalTransform, const FString& ParentNodeUID)
+		UInterchangeSceneNode* CreateTransformNode(FUfbxParser& Parser, UInterchangeBaseNodeContainer& NodeContainer, const FString& NodeDisplayLabel, const FString& NodeUID, const FTransform& LocalTransform, const FString& ParentNodeUID)
 		{
 			UInterchangeSceneNode* TransformNode = NewObject<UInterchangeSceneNode>(&NodeContainer, NAME_None);
 			if (!ensure(TransformNode))
@@ -30,7 +30,7 @@ namespace UE::Interchange::Private
 				Message->Text = LOCTEXT("TransformNodeAllocationError", "Unable to allocate a node when importing FBX.");
 				return nullptr;
 			}
-			NodeContainer.SetupNode(TransformNode, NodeUID, NodeName, EInterchangeNodeContainerType::TranslatedScene, ParentNodeUID);
+			NodeContainer.SetupNode(TransformNode, NodeUID, NodeDisplayLabel, EInterchangeNodeContainerType::TranslatedScene, ParentNodeUID);
 			constexpr bool bResetCache = false;
 			TransformNode->SetCustomLocalTransform(&NodeContainer, LocalTransform, bResetCache);
 
@@ -42,7 +42,7 @@ namespace UE::Interchange::Private
 		{
 			Convert::FNodeNameAndUid NodeNameAndUid(Parser, ParentSceneNode, Node);
 			UInterchangeSceneNode* SceneNode = CreateTransformNode(Parser,
-				NodeContainer, NodeNameAndUid.Name, NodeNameAndUid.UniqueID, Convert::ConvertTransform(Node.local_transform), ParentSceneNode ? ParentSceneNode->GetUniqueID() : TEXT(""));
+				NodeContainer, NodeNameAndUid.Label, NodeNameAndUid.UniqueID, Convert::ConvertTransform(Node.local_transform), ParentSceneNode ? ParentSceneNode->GetUniqueID() : TEXT(""));
 
 			Parser.ElementIdToSceneNode.Add(Node.element_id, SceneNode);
 
@@ -91,9 +91,6 @@ namespace UE::Interchange::Private
 
 		UInterchangeSceneNode* SceneNode = *SceneNodeFound;
 
-		// #ufbx_todo: Node->geometry_transform
-		// #ufbx_todo: possible to set SetCustomGeometricTransform, SetCustomPivotNodeTransform?
-
 		FTransform GeometryTransform = Convert::ConvertTransform(InNode.geometry_transform);
 		if (!GeometryTransform.Equals(FTransform::Identity))
 		{
@@ -135,7 +132,7 @@ namespace UE::Interchange::Private
 			// To recognize Node as a Joint, see e.g. UInterchangePipelineMeshesUtilities::CreateInterchangePipelineMeshesUtilities
 			UnrealNode->AddSpecializedType(FSceneNodeStaticData::GetJointSpecializeTypeString());
 
-			FString JointNodeName = Parser.GetElementName(InNode.element, true);
+			FString JointNodeName = Parser.GetBoneNodeUid(InNode);
 			UnrealNode->SetDisplayLabel(JointNodeName);
 
 			if (InNode.bind_pose)
