[Interchange] Animation Sequence reimport now scales existing float curves, matching Legacy FBX importer's behavior.

#rnx

diff --git a/Engine/Plugins/Interchange/Runtime/Source/Import/Private/Animation/InterchangeAnimSequenceFactory.cpp b/Engine/Plugins/Interchange/Runtime/Source/Import/Private/Animation/InterchangeAnimSequenceFactory.cpp
--- a/Engine/Plugins/Interchange/Runtime/Source/Import/Private/Animation/InterchangeAnimSequenceFactory.cpp
+++ b/Engine/Plugins/Interchange/Runtime/Source/Import/Private/Animation/InterchangeAnimSequenceFactory.cpp
@@ -492,6 +492,7 @@ namespace UE::Interchange::Private
 		AnimSequenceFactoryNode->GetSceneNodeAnimationPayloadKeys(PayloadKeys);
 		
 		const bool bShouldTransact = bIsReimporting;
+		float PreviousSequenceLength = bIsReimporting ? AnimSequence->GetPlayLength() : 0;
 
 		bool bImportBoneTracks = false;
 		AnimSequenceFactoryNode->GetCustomImportBoneTracks(bImportBoneTracks);
@@ -530,6 +531,23 @@ namespace UE::Interchange::Private
 			Controller.SetFrameRate(ResampleFrameRate, bShouldTransact);
 			Controller.SetNumberOfFrames(FrameCount, bShouldTransact);
 
+			if (bIsReimporting && 
+				PreviousSequenceLength > MINIMUM_ANIMATION_LENGTH && 
+				AnimSequence->GetDataModel()->GetNumberOfFloatCurves() > 0)
+			{
+				// The sequence already existed when we began the import. We need to scale the key times for all curves to match the new 
+				// duration before importing over them. This is to catch any user-added curves
+				float ScaleFactor = AnimSequence->GetPlayLength() / PreviousSequenceLength;
+				if (!FMath::IsNearlyEqual(ScaleFactor, 1.f))
+				{
+					for (const FFloatCurve& Curve : AnimSequence->GetDataModel()->GetFloatCurves())
+					{
+						const FAnimationCurveIdentifier CurveId(Curve.GetName(), ERawCurveTrackTypes::RCT_Float);
+						Controller.ScaleCurve(CurveId, 0.f, ScaleFactor, bShouldTransact);
+					}
+				}
+			}
+
 			for (TTuple< const UInterchangeSceneNode*, UE::Interchange::FAnimationPayloadData>& AnimationPayload : BoneTrackData.PreProcessedAnimationPayloads)
 			{
 				NonAnimatedSkeletonNodes.Remove(AnimationPayload.Key->GetUniqueID());
