Interchange: Use transforms placed directly on the LOD scene nodes when producing multiple LOD meshes.

diff --git a/Engine/Plugins/Interchange/Runtime/Source/Import/Private/Mesh/InterchangeStaticMeshFactory.cpp b/Engine/Plugins/Interchange/Runtime/Source/Import/Private/Mesh/InterchangeStaticMeshFactory.cpp
--- a/Engine/Plugins/Interchange/Runtime/Source/Import/Private/Mesh/InterchangeStaticMeshFactory.cpp
+++ b/Engine/Plugins/Interchange/Runtime/Source/Import/Private/Mesh/InterchangeStaticMeshFactory.cpp
@@ -176,7 +176,7 @@ void UInterchangeStaticMeshFactory::CreatePayloadTasks(const FImportAssetObjectP
 			}
 		};
 
-		auto AddMeshPayloads = [&Arguments, MeshTranslatorPayloadInterface, bBakeMeshes, bBakePivotMeshes, &GlobalOffsetTransform, &PayloadTasks, bAsync, &PayloadAttributes, &GetNodeAndGlobalTransform]
+		auto AddMeshPayloads = [&Arguments, MeshTranslatorPayloadInterface, bBakeMeshes, bBakePivotMeshes, &GlobalOffsetTransform, &PayloadTasks, bAsync, &PayloadAttributes, &GetNodeAndGlobalTransform, LodCount]
 			(const TArray<FString>& MeshUids, TMap<FInterchangeMeshPayLoadKey, UE::Interchange::FMeshPayload>& PayloadPerKey, const TMap<FString, FString>* ColliderToRenderUid = nullptr)
 			{
 				PayloadPerKey.Reserve(MeshUids.Num());
@@ -187,14 +187,14 @@ void UInterchangeStaticMeshFactory::CreatePayloadTasks(const FImportAssetObjectP
 
 				for (const FString& MeshUid : MeshUids)
 				{
-					FTransform GlobalMeshTransform;
+					FTransform TransformToBake;
 					const UInterchangeMeshNode* MeshNode = nullptr;
-					GetNodeAndGlobalTransform(MeshUid, bBakeTheseMeshes, bBakeThesePivotMeshes, MeshNode, GlobalMeshTransform);
+					GetNodeAndGlobalTransform(MeshUid, bBakeTheseMeshes, bBakeThesePivotMeshes, MeshNode, TransformToBake);
 
 					if (!bBakeMeshes && ColliderToRenderUid)
 					{
 						const FString& ColliderMeshUid = MeshUid;
-						const FTransform& ColliderLocalToGlobal = GlobalMeshTransform;
+						const FTransform& ColliderLocalToGlobal = TransformToBake;
 
 						if (const FString* RenderMeshUid = ColliderToRenderUid->Find(ColliderMeshUid))
 						{
@@ -202,7 +202,20 @@ void UInterchangeStaticMeshFactory::CreatePayloadTasks(const FImportAssetObjectP
 							const UInterchangeMeshNode* RenderMeshNode = nullptr;
 							GetNodeAndGlobalTransform(*RenderMeshUid, bBakeTheseMeshes, bBakeThesePivotMeshes, RenderMeshNode, RenderLocalToGlobal);
 
-							GlobalMeshTransform = ColliderLocalToGlobal * RenderLocalToGlobal.Inverse();
+							TransformToBake = ColliderLocalToGlobal * RenderLocalToGlobal.Inverse();
+						}
+					}
+					else if ((!bBakeMeshes && !bBakePivotMeshes) && LodCount > 1)
+					{
+						// If we're not fully baking the whole thing, but are still importing LODs, we need to bake the scene node transform
+						// from each LOD into the mesh, as there's otherwise no other place to put it (as each LOD may have a separate transform,
+						// and so it can't be placed on the component)
+						if (const UInterchangeSceneNode* SceneNode = Cast<UInterchangeSceneNode>(Arguments.NodeContainer->GetNode(MeshUid)))
+						{
+							FTransform SceneNodeLocalTransform;
+							SceneNode->GetCustomLocalTransform(SceneNodeLocalTransform);
+
+							TransformToBake = TransformToBake * SceneNodeLocalTransform;
 						}
 					}
 
@@ -222,15 +235,15 @@ void UInterchangeStaticMeshFactory::CreatePayloadTasks(const FImportAssetObjectP
 					FInterchangeMeshPayLoadKey& PayLoadKey = OptionalPayLoadKey.GetValue();
 					
 					FInterchangeMeshPayLoadKey GlobalPayLoadKey = PayLoadKey;
-					GlobalPayLoadKey.UniqueId += FInterchangeMeshPayLoadKey::GetTransformString(GlobalMeshTransform);
+					GlobalPayLoadKey.UniqueId += FInterchangeMeshPayLoadKey::GetTransformString(TransformToBake);
 					if (!PayloadPerKey.Contains(GlobalPayLoadKey))
 					{
 						UE::Interchange::FMeshPayload& Payload = PayloadPerKey.FindOrAdd(GlobalPayLoadKey);
-						Payload.Transform = GlobalMeshTransform;
+						Payload.Transform = TransformToBake;
 						Payload.MeshName = PayLoadKey.UniqueId;
-						PayloadAttributes.RegisterAttribute(UE::Interchange::FAttributeKey{ UE::Interchange::MeshPayload::Attributes::MeshGlobalTransform }, GlobalMeshTransform);
+						PayloadAttributes.RegisterAttribute(UE::Interchange::FAttributeKey{ UE::Interchange::MeshPayload::Attributes::MeshGlobalTransform }, TransformToBake);
 						TSharedPtr<UE::Interchange::FInterchangeTaskLambda, ESPMode::ThreadSafe> TaskGetMeshPayload = MakeShared<UE::Interchange::FInterchangeTaskLambda, ESPMode::ThreadSafe>(bAsync ? UE::Interchange::EInterchangeTaskThread::AsyncThread : UE::Interchange::EInterchangeTaskThread::GameThread
-							, [&Payload, GlobalMeshTransform, MeshTranslatorPayloadInterface, PayLoadKey, PayloadAttributes]()
+							, [&Payload, TransformToBake, MeshTranslatorPayloadInterface, PayLoadKey, PayloadAttributes]()
 							{
 								TRACE_CPUPROFILER_EVENT_SCOPE(UInterchangeStaticMeshFactory::GetMeshPayloadDataTask);
 								if (ensure(!Payload.PayloadData.IsSet()))
