USD: Fix a couple of issues reported by BugHawk from cl 46415229

diff --git a/Engine/Plugins/Interchange/Extensions/OpenUSD/Source/Import/Private/InterchangeUsdTranslator.cpp b/Engine/Plugins/Interchange/Extensions/OpenUSD/Source/Import/Private/InterchangeUsdTranslator.cpp
--- a/Engine/Plugins/Interchange/Extensions/OpenUSD/Source/Import/Private/InterchangeUsdTranslator.cpp
+++ b/Engine/Plugins/Interchange/Extensions/OpenUSD/Source/Import/Private/InterchangeUsdTranslator.cpp
@@ -6332,7 +6332,7 @@ namespace UE::InterchangeUsdTranslator::Private
 		{
 			const UE::Interchange::FAnimationPayloadQuery* Query = Queries[QueryIndex];
 
-			const FString& IndexStr = PayloadKeyTokens[1];
+			const FString IndexStr = Query->PayloadKey.UniqueId.RightChop(SkeletonPrimPath.Len() + 1);	// Also skip the '\'
 			int32 JointIndex = INDEX_NONE;
 			bool bLexed = LexTryParseString(JointIndex, *IndexStr);
 			if (!bLexed)
diff --git a/Engine/Plugins/Interchange/Runtime/Source/Common/Public/Usd/InterchangeUsdDefinitions.h b/Engine/Plugins/Interchange/Runtime/Source/Common/Public/Usd/InterchangeUsdDefinitions.h
--- a/Engine/Plugins/Interchange/Runtime/Source/Common/Public/Usd/InterchangeUsdDefinitions.h
+++ b/Engine/Plugins/Interchange/Runtime/Source/Common/Public/Usd/InterchangeUsdDefinitions.h
@@ -41,7 +41,7 @@ namespace UE::Interchange::USD
 	const FString UVIndexParameterSuffix = TEXT("UVIndex");
 
 	// In order to make it easy to find the skeleton root bone scene node UID from a prim path, for the root bone in particular
-	// we always use an UID like <skeleton prim path>/RootBoneUid (e.g. "/MySkelRoot/MySkeleton/Root")
+	// we always use an UID based on <skeleton prim path>/RootBoneUidSuffix (e.g. "\Bone\/MySkelRoot/MySkeleton/Root")
 	//
 	// Note that this is unrelated to the bone's *name*, which is still named after the actual USD bone name
 	const FString RootBoneUidSuffix = TEXT("Root");
