USD Interchange: Fix SVT asset names ending up as "Null" if the volume filename itself was just a number.

diff --git a/Engine/Plugins/Interchange/Runtime/Source/Pipelines/Private/InterchangeSparseVolumeTexturePipeline.cpp b/Engine/Plugins/Interchange/Runtime/Source/Pipelines/Private/InterchangeSparseVolumeTexturePipeline.cpp
--- a/Engine/Plugins/Interchange/Runtime/Source/Pipelines/Private/InterchangeSparseVolumeTexturePipeline.cpp
+++ b/Engine/Plugins/Interchange/Runtime/Source/Pipelines/Private/InterchangeSparseVolumeTexturePipeline.cpp
@@ -461,6 +461,33 @@ void UInterchangeSparseVolumeTexturePipeline::ExecutePipeline(
 		SplitNumberedSuffix(FileName, Prefix, NumberSuffix);
 		FString DisplayLabel = RemoveTrailingSeparators(Prefix);	// e.g. "tornado"
 
+		// If the volume name is purely a number or something else weird (e.g. MaterialEggs scenes) then the above
+		// sanitizing process has likely produced a fully empty string. Here let's try using one of the grid names
+		// as the SVT asset name, if we can find one
+		if (DisplayLabel.IsEmpty())
+		{
+			TArray<FString> GridNodeUids;
+			FirstVolume->GetCustomGridDependecies(GridNodeUids);
+
+			for (int32 GridIndex = 0; GridIndex < GridNodeUids.Num(); ++GridIndex)
+			{
+				const UInterchangeVolumeGridNode* GridNode = Cast<UInterchangeVolumeGridNode>(BaseNodeContainer->GetNode(GridNodeUids[GridIndex]));
+				if (GridNode)
+				{
+					DisplayLabel = GridNode->GetDisplayLabel();
+					if (!DisplayLabel.IsEmpty())
+					{
+						break;
+					}
+				}
+			}
+		}
+		// Just give up and use the "Volume" string as the asset name instead, if that still didn't work
+		if (DisplayLabel.IsEmpty())
+		{
+			DisplayLabel = TEXT("Volume");
+		}
+
 		const FString FactoryNodeUid = UInterchangeFactoryBaseNode::BuildFactoryNodeUid(FirstVolume->GetUniqueID());
 
 		UInterchangeSparseVolumeTextureFactoryNode* FactoryNode = CreateTextureFactoryNode(DisplayLabel, FactoryNodeUid, BaseNodeContainer);
