Interchange - FBX - Fixed name clashes that can occur after asset names or bone names get sanitized

FBX was already doing some work to avoid collision by adding a suffix to identical names. However the name used to compare was not sanitized, which means that after sanitization collision might still happen causing all sorts of issue.
The fix is to sanitize before adding any sort of suffix. Similar to what GLTF is doing.

This is a local fix for FBX, but a bigger refactor is needed at this point in order to:
- Make renaming uniform between all translator (right now FBX adds a number, while GLTF add a hash) (UE-319008)
- Make renaming local specific assets instead of uniquely global (like having multiple skeletal meshes with similar bone names) (UE-319012)
- Reduce translators overhead
- Fix inconsistencies within translators themselves (for instance, FBX will not rename the same if "import collision" is enabled or not)

#tests Ran Interchange FBX tests locally - got same result as before this CL
[RN] Fixed risk of missing content when importing FBX containing assets or bones requiring name sanitization.

diff --git a/Engine/Plugins/Interchange/Runtime/Source/Parsers/Fbx/Private/FbxAPI.cpp b/Engine/Plugins/Interchange/Runtime/Source/Parsers/Fbx/Private/FbxAPI.cpp
--- a/Engine/Plugins/Interchange/Runtime/Source/Parsers/Fbx/Private/FbxAPI.cpp
+++ b/Engine/Plugins/Interchange/Runtime/Source/Parsers/Fbx/Private/FbxAPI.cpp
@@ -11,6 +11,7 @@
 #include "FbxMaterial.h"
 #include "FbxMesh.h"
 #include "FbxScene.h"
+#include "InterchangeHelper.h"
 #include "InterchangeTextureNode.h"
 #if WITH_ENGINE
 #include "Mesh/InterchangeMeshPayload.h"
@@ -315,7 +316,7 @@ namespace UE
 				{
 					if (ObjectName.Contains(TEXT(":")))
 					{
-						ObjectName = ObjectName.Replace(TEXT(":"), TEXT("_"));
+						ObjectName.ReplaceInline(TEXT(":"), TEXT("_"));
 						Object->SetName(TCHAR_TO_UTF8(*ObjectName));
 					}
 				}
@@ -383,7 +384,12 @@ namespace UE
 				auto MakeFbxObjectNameUnique = [bKeepFbxNamespaceClosure = bKeepFbxNamespace](FbxObject* Object, TMap<FString, int32>& Names)
 					{
 						FString ObjectName = UTF8_TO_TCHAR(Object->GetName());
+						//Manage namespaces should be done before sanitizing since it relies on the ":" character
 						ManageNamespace(bKeepFbxNamespaceClosure, ObjectName, Object);
+
+						//Make sure to use the sanitized name to avoid name collisions later on
+						UE::Interchange::SanitizeName(ObjectName);
+
 						if (int32* Count = Names.Find(ObjectName))
 						{
 							(*Count)++;
diff --git a/Engine/Plugins/Interchange/Runtime/Source/Parsers/GLTFCore/Private/GLTF/GLTFAsset.cpp b/Engine/Plugins/Interchange/Runtime/Source/Parsers/GLTFCore/Private/GLTF/GLTFAsset.cpp
--- a/Engine/Plugins/Interchange/Runtime/Source/Parsers/GLTFCore/Private/GLTF/GLTFAsset.cpp
+++ b/Engine/Plugins/Interchange/Runtime/Source/Parsers/GLTFCore/Private/GLTF/GLTFAsset.cpp
@@ -204,8 +204,7 @@ namespace GLTF
 
 				UE::Interchange::SanitizeName(Node.Name, bIsJoint);
 
-				FString LowerCaseName = Node.Name.ToLower();
-				if (Names.Contains(LowerCaseName))
+				if (Names.Contains(Node.Name))
 				{
 					//to make the name unique we encode the UID into the name:
 					Node.Name += TEXT("_") + FMD5::HashAnsiString(*Node.UniqueId);
diff --git a/Engine/Source/Runtime/Interchange/Core/Private/InterchangeHelper.cpp b/Engine/Source/Runtime/Interchange/Core/Private/InterchangeHelper.cpp
--- a/Engine/Source/Runtime/Interchange/Core/Private/InterchangeHelper.cpp
+++ b/Engine/Source/Runtime/Interchange/Core/Private/InterchangeHelper.cpp
@@ -21,7 +21,7 @@ namespace UE::Interchange
 		static FAutoConsoleVariableRef CVarInterchangeMaxAssetPathLength(
 			TEXT("Interchange.MaxAssetPathLength"),
 			GInterchangeMaxAssetPathLength,
-			TEXT("Interchange will try to limit asset path length to this value. Default: 260")
+			TEXT("Interchange will try to limit asset path length to this value. Default: 160")
 		);
 
 		int32 GetUsableMaxAssetPathLength()
