[Interchange][FBX] UFBX - custom curves

- curve/split_curve support
- doesn't import tangets from ascii files well - seems like an issue in ufbx
- bonus: Material with empy emissive doesn;'t create this connection

diff --git a/Engine/Plugins/Interchange/Runtime/Source/Parsers/Fbx/Private/Ufbx/UfbxAnimation.cpp b/Engine/Plugins/Interchange/Runtime/Source/Parsers/Fbx/Private/Ufbx/UfbxAnimation.cpp
--- a/Engine/Plugins/Interchange/Runtime/Source/Parsers/Fbx/Private/Ufbx/UfbxAnimation.cpp
+++ b/Engine/Plugins/Interchange/Runtime/Source/Parsers/Fbx/Private/Ufbx/UfbxAnimation.cpp
@@ -302,6 +302,21 @@ namespace UE::Interchange::Private
 			return TrackNodeUid;
 		}
 
+		template<typename T>
+		void ConvertStepCurveKeyframes(const ufbx_anim_curve& Curve, TArray<float>& OutTimes, TOptional<TArray<T>>& OutValues)
+		{
+			const int32 Count = Curve.keyframes.count;
+			OutTimes.Reserve(Count);
+			TArray<T>& Values = OutValues.Emplace();
+			Values.Reserve(Count);
+
+			for (const ufbx_keyframe& Keyframe : Curve.keyframes)
+			{
+				OutTimes.Add(static_cast<float>(Keyframe.time));
+				Values.Add(static_cast<T>(Keyframe.value));
+			}
+		}
+
 	}
 
 	void FUfbxAnimation::FetchSkinnedAnimation(const UE::Interchange::FAnimationPayloadQuery& PayloadQuery, const FSkeletalAnimationPayloadContext& SkeletalAnimationPayload, FAnimationPayloadData& PayloadData)
@@ -384,6 +399,53 @@ namespace UE::Interchange::Private
 		return true;
 	}
 
+	bool FUfbxAnimation::FetchPropertyAnimationCurves(const FUfbxParser& Parser, const FPropertyAnimationPayloadContext& Animation, TArray<FInterchangeCurve>& OutCurves)
+	{
+		for (const ufbx_anim_curve* Curve : Animation.AnimValue->curves)
+		{
+			if (Curve)
+			{
+				AnimUtils::ConvertCurve(*Curve, 1.f, OutCurves.AddDefaulted_GetRef().Keys);
+			}
+		}
+		return true;
+	}
+
+	bool FUfbxAnimation::FetchPropertyAnimationStepCurves(const FUfbxParser& Parser, const FPropertyAnimationPayloadContext& Animation, TArray<FInterchangeStepCurve>& OutCurves)
+	{
+		for (const ufbx_anim_curve* Curve : Animation.AnimValue->curves)
+		{
+			if (Curve)
+			{
+				
+
+				switch (Animation.Prop->type)
+				{
+				case UFBX_PROP_BOOLEAN:
+					{
+						FInterchangeStepCurve& DestinationCurve = OutCurves.AddDefaulted_GetRef();
+						AnimUtils::ConvertStepCurveKeyframes(*Curve, DestinationCurve.KeyTimes, DestinationCurve.BooleanKeyValues);
+					}
+					break;
+				case UFBX_PROP_INTEGER:
+					{
+						FInterchangeStepCurve& DestinationCurve = OutCurves.AddDefaulted_GetRef();
+						AnimUtils::ConvertStepCurveKeyframes(*Curve, DestinationCurve.KeyTimes, DestinationCurve.IntegerKeyValues);
+					}
+					break;
+				case UFBX_PROP_NUMBER:
+					{
+						FInterchangeStepCurve& DestinationCurve = OutCurves.AddDefaulted_GetRef();
+						AnimUtils::ConvertStepCurveKeyframes(*Curve, DestinationCurve.KeyTimes, DestinationCurve.FloatKeyValues);
+					}
+					break;
+				}
+
+			}
+		}
+		return true;
+	}
+
 	void FUfbxAnimation::AddAnimation(FUfbxParser& Parser, FUfbxScene& Scene, UInterchangeBaseNodeContainer& NodeContainer,
 	                                  const TMap<ufbx_blend_shape*, ufbx_mesh*> MeshPerBlendShape)
 	{
@@ -642,6 +704,75 @@ namespace UE::Interchange::Private
 		}
 
 	}
+
+	bool FUfbxAnimation::AddNodeAttributeCurvesAnimation(FUfbxParser& Parser
+		, const FString NodeUid
+		, const ufbx_prop& Prop
+		, const ufbx_anim_value& AnimValue
+		, TOptional<FString>& OutPayloadKey
+		, TOptional<bool>& OutIsStepCurve
+		, FString& OutCurveName)
+	{
+		FString PropertyName = Convert::ToUnrealString(Prop.name);
+		const FString PayLoadKey = NodeUid + PropertyName + TEXT("_AnimationPayloadKey");
+		if (!ensure(!Parser.PayloadContexts.Contains(PayLoadKey)))
+		{
+			return false;
+		}
+
+		bool bIsStepCurve;
+		switch (Prop.type)
+		{
+		case UFBX_PROP_NUMBER:
+		case UFBX_PROP_VECTOR:
+		case UFBX_PROP_COLOR:
+		case UFBX_PROP_COLOR_WITH_ALPHA:
+		case UFBX_PROP_TRANSLATION:
+		case UFBX_PROP_ROTATION:
+		case UFBX_PROP_SCALING:
+		case UFBX_PROP_DISTANCE:
+		case UFBX_PROP_COMPOUND:
+			bIsStepCurve = false;
+			break;
+		default:
+			bIsStepCurve = true;
+		};
+
+		bool bHasCurve = false;
+		//Only curves with Constant interpolation on all keys are deemed as StepCurves.
+		for (const ufbx_anim_curve* Curve: AnimValue.curves)
+		{
+			if (Curve)
+			{
+				bHasCurve = true;
+				for (const ufbx_keyframe& Keyframe : Curve->keyframes)
+				{
+					if (Keyframe.interpolation != UFBX_INTERPOLATION_CONSTANT_PREV 
+						&& Keyframe.interpolation != UFBX_INTERPOLATION_CONSTANT_NEXT)
+					{
+						bIsStepCurve = false;
+					}
+				}
+			}
+		}
+
+		if (!bHasCurve)
+		{
+			// Property even though has an anim_value but still not animated
+			return false;
+		}
+
+		OutIsStepCurve = bIsStepCurve;
+		Parser.PayloadContexts.Add(PayLoadKey, FPropertyAnimationPayloadContext{
+			.Prop = &Prop,
+			.AnimValue = &AnimValue,
+			.bIsStepAnimation = bIsStepCurve
+		});
+		OutPayloadKey = PayLoadKey;
+		OutCurveName = Convert::ToUnrealString(AnimValue.name);
+
+		return true;
+	}
 }
 
 #undef LOCTEXT_NAMESPACE
diff --git a/Engine/Plugins/Interchange/Runtime/Source/Parsers/Fbx/Private/Ufbx/UfbxAnimation.h b/Engine/Plugins/Interchange/Runtime/Source/Parsers/Fbx/Private/Ufbx/UfbxAnimation.h
--- a/Engine/Plugins/Interchange/Runtime/Source/Parsers/Fbx/Private/Ufbx/UfbxAnimation.h
+++ b/Engine/Plugins/Interchange/Runtime/Source/Parsers/Fbx/Private/Ufbx/UfbxAnimation.h
@@ -12,6 +12,7 @@ namespace UE::Interchange
 }
 
 struct FInterchangeCurve;
+struct FInterchangeStepCurve;
 
 namespace UE::Interchange::Private
 {
@@ -25,6 +26,18 @@ namespace UE::Interchange::Private
 		static bool FetchMorphTargetAnimation(const FUfbxParser& Parser, const FMorphAnimationPayloadContext& Animation, TArray<FInterchangeCurve>& OutCurves);
 		static bool FetchRigidAnimation(const FUfbxParser& Parser, const FRigidAnimationPayloadContext& Animation, TArray<FInterchangeCurve>& OutCurves);
 
+		static bool FetchPropertyAnimationCurves(const FUfbxParser& Parser, const FPropertyAnimationPayloadContext& Animation, TArray<FInterchangeCurve>& OutCurves);
+		static bool FetchPropertyAnimationStepCurves(const FUfbxParser& Parser, const FPropertyAnimationPayloadContext& Animation, TArray<FInterchangeStepCurve>& OutCurves);
+
 		static void AddAnimation(FUfbxParser& Parser, FUfbxScene& Scene, UInterchangeBaseNodeContainer& NodeContainer, const TMap<ufbx_blend_shape*, ufbx_mesh*> MeshPerBlendShape);
+
+		static bool AddNodeAttributeCurvesAnimation(FUfbxParser& Parser
+			, const FString NodeUid
+			, const ufbx_prop& Prop
+			, const ufbx_anim_value& AnimValue
+			, TOptional<FString>& OutPayloadKey
+			, TOptional<bool>& OutIsStepCurve
+			, FString& OutCurveName);
+
 	};
 }
diff --git a/Engine/Plugins/Interchange/Runtime/Source/Parsers/Fbx/Private/Ufbx/UfbxMaterial.cpp b/Engine/Plugins/Interchange/Runtime/Source/Parsers/Fbx/Private/Ufbx/UfbxMaterial.cpp
--- a/Engine/Plugins/Interchange/Runtime/Source/Parsers/Fbx/Private/Ufbx/UfbxMaterial.cpp
+++ b/Engine/Plugins/Interchange/Runtime/Source/Parsers/Fbx/Private/Ufbx/UfbxMaterial.cpp
@@ -155,13 +155,16 @@ namespace UE::Interchange::Private
 
 			FMaterialMapConverter MaterialMapConverter(Parser, NodeContainer, *MaterialInstanceNode);
 			
-			auto SetColorFromMap = [&MaterialInstanceNode](FName InputName, const ufbx_material_map& Map, float Factor)
+			auto SetColorFromMap = [&MaterialInstanceNode](FName InputName, const ufbx_material_map& Map, float Factor, TFunctionRef<bool(const FLinearColor&)> Filter = [](const FLinearColor&){return true;})
 			{
 				if (Map.value_components >=3)
 				{
 					const FString InputAttributeKey = InputName.ToString();
 					FLinearColor PropertyValue(Convert::ConvertVec4(Map.value_vec4) * Factor);
-					MaterialInstanceNode->AddVectorParameterValue(InputAttributeKey, PropertyValue);
+					if (Filter(PropertyValue))
+					{
+						MaterialInstanceNode->AddVectorParameterValue(InputAttributeKey, PropertyValue);
+					}
 				}
 			};
 
@@ -240,7 +243,7 @@ namespace UE::Interchange::Private
 
 				bHasInput = true;
 			}
-
+			
 			// Emissive
 			if (Material->fbx.emission_color.has_value)
 			{
@@ -250,7 +253,7 @@ namespace UE::Interchange::Private
 
 				if (!SetTextureFromMap(InputName, Map, Factor))
 				{
-					SetColorFromMap(InputName, Map, Factor);
+					SetColorFromMap(InputName, Map, Factor, [](const FLinearColor& Color) { return !Color.IsAlmostBlack(); });
 				}
 
 				bHasInput = true;
diff --git a/Engine/Plugins/Interchange/Runtime/Source/Parsers/Fbx/Private/Ufbx/UfbxParser.cpp b/Engine/Plugins/Interchange/Runtime/Source/Parsers/Fbx/Private/Ufbx/UfbxParser.cpp
--- a/Engine/Plugins/Interchange/Runtime/Source/Parsers/Fbx/Private/Ufbx/UfbxParser.cpp
+++ b/Engine/Plugins/Interchange/Runtime/Source/Parsers/Fbx/Private/Ufbx/UfbxParser.cpp
@@ -126,13 +126,11 @@ FString FUfbxParser::GetMeshUid(const ufbx_element& Mesh) const
 	return GetElementNameDeduplicated(Mesh);
 }
 
-void FUfbxParser::ConvertProperty(const ufbx_prop& Prop, UInterchangeBaseNode* SceneNode)
+void FUfbxParser::ConvertProperty(const ufbx_prop& Prop, UInterchangeBaseNode* SceneNode, const TOptional<FString>& PayloadKey)
 {
 	if (Prop.flags & UFBX_PROP_FLAG_USER_DEFINED)
 	{
 		FString PropertyName = Convert::ToUnrealString(Prop.name);
-		TOptional<FString> PayloadKey;
-		// #ufbx_todo: animated prop; PayloadKey will be added then  
 
 		switch (Prop.type)
 		{
@@ -339,7 +337,7 @@ bool FUfbxParser::FetchPayloadData(const FString& PayloadKey, const FString& Pay
 		break;
 	case FPayloadContext::AnimationRigid:
 		{
-			FRigidAnimationPayloadContext& Animation = const_cast<FRigidAnimationPayloadContext&>(PayloadContexts.GetRigidAnimation(PayloadContext));
+			const FRigidAnimationPayloadContext& Animation = PayloadContexts.GetRigidAnimation(PayloadContext);
 
 			TArray<FInterchangeCurve> InterchangeCurves;
 
@@ -351,8 +349,39 @@ bool FUfbxParser::FetchPayloadData(const FString& PayloadKey, const FString& Pay
 				FFileHelper::SaveArrayToFile(Buffer, *PayloadFilepath);
 				return true;
 			}
+		}
+		break;
+	case FPayloadContext::AnimationProperty:
+		{
+			const FPropertyAnimationPayloadContext& Animation = PayloadContexts.GetPropertyAnimation(PayloadContext);
+
+			if (Animation.bIsStepAnimation)
+			{
+				TArray<FInterchangeStepCurve> InterchangeStepCurves;
+
+				if (FUfbxAnimation::FetchPropertyAnimationStepCurves(*this, Animation, InterchangeStepCurves))
+				{
+					TArray64<uint8> Buffer;
+					FMemoryWriter64 Ar(Buffer);
+					Ar << InterchangeStepCurves;
+					FFileHelper::SaveArrayToFile(Buffer, *PayloadFilepath);
+					return true;
+				}
+			}
+			else
+			{
+				TArray<FInterchangeCurve> InterchangeCurves;
 
-			return true;
+				if (FUfbxAnimation::FetchPropertyAnimationCurves(*this, Animation, InterchangeCurves))
+				{
+					TArray64<uint8> Buffer;
+					FMemoryWriter64 Ar(Buffer);
+					Ar << InterchangeCurves;
+					FFileHelper::SaveArrayToFile(Buffer, *PayloadFilepath);
+					return true;
+				}
+				
+			}
 		}
 		break;
 	}
diff --git a/Engine/Plugins/Interchange/Runtime/Source/Parsers/Fbx/Private/Ufbx/UfbxParser.h b/Engine/Plugins/Interchange/Runtime/Source/Parsers/Fbx/Private/Ufbx/UfbxParser.h
--- a/Engine/Plugins/Interchange/Runtime/Source/Parsers/Fbx/Private/Ufbx/UfbxParser.h
+++ b/Engine/Plugins/Interchange/Runtime/Source/Parsers/Fbx/Private/Ufbx/UfbxParser.h
@@ -26,7 +26,8 @@ namespace UE::Interchange::Private
 			Morph,
 			AnimationSkinned,
 			AnimationMorph,
-			AnimationRigid
+			AnimationRigid,
+			AnimationProperty
 		};
 
 		FPayloadContext() = delete;
@@ -73,6 +74,13 @@ namespace UE::Interchange::Private
 		const ufbx_anim_prop* ScalingProp = nullptr;
 	};
 
+	struct FPropertyAnimationPayloadContext
+	{
+		const ufbx_prop* Prop;
+		const ufbx_anim_value* AnimValue;
+		bool bIsStepAnimation;
+	};
+
 	class FUfbxParser: public IFbxParser, public FNoncopyable
 	{
 	public:
@@ -170,7 +178,7 @@ namespace UE::Interchange::Private
 			return {Material ? FName(*GetElementNameDeduplicated(Material->element)) : NAME_None};
 		}
 
-		static void ConvertProperty(const ufbx_prop& Prop, UInterchangeBaseNode* SceneNode);
+		static void ConvertProperty(const ufbx_prop& Prop, UInterchangeBaseNode* SceneNode, const TOptional<FString>& PayloadKey = TOptional<FString>());
 
 		ufbx_scene* Scene = nullptr;
 
@@ -179,7 +187,13 @@ namespace UE::Interchange::Private
 			void Reset()
 			{
 				PayloadContexts.Reset();
+
 				MorphPayloads.Reset();
+
+				SkeletalAnimations.Reset();
+				MorphAnimations.Reset();
+				RigidAnimations.Reset();
+				PropertyAnimations.Reset();
 			}
 
 			void Add(const FString& PayloadKey, const FPayloadContext& Payload)
@@ -211,6 +225,12 @@ namespace UE::Interchange::Private
 				PayloadContexts.Add(PayloadKey, FPayloadContext(FPayloadContext::AnimationRigid, PayloadId));
 			}
 
+			void Add(const FString& PayloadKey, const FPropertyAnimationPayloadContext& Payload)
+			{
+				int32 PayloadId = PropertyAnimations.Add(Payload);
+				PayloadContexts.Add(PayloadKey, FPayloadContext(FPayloadContext::AnimationProperty, PayloadId));
+			}
+
 			bool Contains(const FString& PayloadKey)
 			{
 				return PayloadContexts.Contains(PayloadKey);
@@ -246,18 +266,19 @@ namespace UE::Interchange::Private
 				return RigidAnimations[Payload.Index];
 			}
 
+			const FPropertyAnimationPayloadContext& GetPropertyAnimation(const FPayloadContext& Payload)
+			{
+				return PropertyAnimations[Payload.Index];
+			}
+
 			TMap<FString, FPayloadContext> PayloadContexts;
 
-			// Payloads that a simple ufbx elements
 			TArray<FMorph> MorphPayloads;
-			TArray<FSkeletalAnimationPayloadContext> SkeletalAnimations;
 
+			TArray<FSkeletalAnimationPayloadContext> SkeletalAnimations;
 			TArray<FMorphAnimationPayloadContext> MorphAnimations;
 			TArray<FRigidAnimationPayloadContext> RigidAnimations;
-			
-			
-			
-
+			TArray<FPropertyAnimationPayloadContext> PropertyAnimations;
 		};
 
 		TMap<uint32_t, UInterchangeSceneNode*> ElementIdToSceneNode;
diff --git a/Engine/Plugins/Interchange/Runtime/Source/Parsers/Fbx/Private/Ufbx/UfbxScene.cpp b/Engine/Plugins/Interchange/Runtime/Source/Parsers/Fbx/Private/Ufbx/UfbxScene.cpp
--- a/Engine/Plugins/Interchange/Runtime/Source/Parsers/Fbx/Private/Ufbx/UfbxScene.cpp
+++ b/Engine/Plugins/Interchange/Runtime/Source/Parsers/Fbx/Private/Ufbx/UfbxScene.cpp
@@ -1,10 +1,13 @@
 // Copyright Epic Games, Inc. All Rights Reserved.
 
 #include "UfbxScene.h"
+
+#include "InterchangeFbxSettings.h"
 #include "UfbxConvert.h"
 #include "UfbxParser.h"
 
 #include "InterchangeMeshNode.h"
+#include "UfbxAnimation.h"
 #include "Nodes/InterchangeBaseNodeContainer.h"
 #include "Nodes/InterchangeUserDefinedAttribute.h"
 
@@ -79,6 +82,33 @@ namespace UE::Interchange::Private
 			return Bone;
 		}
 
+		bool IsPropTypeSupported(const ufbx_prop_type& PropType)
+		{
+			// Also see FUfbxParser::ConvertProperty
+			switch (PropType)
+			{
+			case UFBX_PROP_BOOLEAN:
+			case UFBX_PROP_INTEGER:
+			case UFBX_PROP_NUMBER:
+			case UFBX_PROP_VECTOR:
+			case UFBX_PROP_COLOR:
+			case UFBX_PROP_COLOR_WITH_ALPHA:
+			case UFBX_PROP_STRING:
+			case UFBX_PROP_DATE_TIME:
+				return true;
+			// #ufbx_todo: support other useful custom attributes among these types
+			case UFBX_PROP_TRANSLATION:
+			case UFBX_PROP_ROTATION:
+			case UFBX_PROP_SCALING:
+			case UFBX_PROP_DISTANCE:
+			case UFBX_PROP_COMPOUND:
+			case UFBX_PROP_BLOB:
+			case UFBX_PROP_REFERENCE:
+			case UFBX_PROP_TYPE_FORCE_32BIT:
+			default: ;
+			}
+			return false;
+		}
 	}
 
 	void FUfbxScene::ProcessNode(UInterchangeBaseNodeContainer& NodeContainer, const ufbx_node& InNode)
@@ -151,7 +181,48 @@ namespace UE::Interchange::Private
 
 		for (const ufbx_prop& Prop : InNode.props.props)
 		{
-			Parser.ConvertProperty(Prop, SceneNode);
+			if (SceneUtils::IsPropTypeSupported(Prop.type) &&
+				Prop.flags & UFBX_PROP_FLAG_USER_DEFINED)
+			{
+				TOptional<FString> PayloadKey;
+				if (Prop.flags & UFBX_PROP_FLAG_ANIMATED)
+				{
+					for (const ufbx_connection& Connection : InNode.element.connections_dst)
+					{
+						// Property and its connection have same name string pointer in ufbx
+						if (Connection.dst_prop.data == Prop.name.data)
+						{
+							if (ufbx_anim_value* AnimValue = ufbx_as_anim_value(Connection.src))
+							{
+								FString CurveName;
+								TOptional<bool> bIsStepCurve;
+								if (FUfbxAnimation::AddNodeAttributeCurvesAnimation(Parser, SceneNode->GetUniqueID(), Prop, *AnimValue, PayloadKey, bIsStepCurve, CurveName))
+								{
+									EInterchangeAnimationPayLoadType PayloadType = bIsStepCurve.GetValue() ? EInterchangeAnimationPayLoadType::STEPCURVE : EInterchangeAnimationPayLoadType::CURVE;
+
+									const UInterchangeFbxSettings* InterchangeFbxSettings = GetDefault<UInterchangeFbxSettings>();
+									if (EInterchangePropertyTracks PropertyTrack = InterchangeFbxSettings->GetPropertyTrack(CurveName); PropertyTrack != EInterchangePropertyTracks::None)
+									{
+										UInterchangeAnimationTrackNode* AnimTrackNode = NewObject< UInterchangeAnimationTrackNode >(&NodeContainer);
+										const FString AnimTrackNodeName = SceneNode->GetDisplayLabel() + CurveName;
+										const FString AnimTrackNodeUid = TEXT("\\AnimationTrack\\") + AnimTrackNodeName;
+
+										NodeContainer.SetupNode(AnimTrackNode, AnimTrackNodeUid, AnimTrackNodeName, EInterchangeNodeContainerType::TranslatedAsset);
+
+										AnimTrackNode->SetCustomActorDependencyUid(SceneNode->GetUniqueID());
+										AnimTrackNode->SetCustomAnimationPayloadKey(PayloadKey.GetValue(), PayloadType);
+										AnimTrackNode->SetCustomPropertyTrack(PropertyTrack);
+									}
+
+									SceneNode->SetAnimationCurveTypeForCurveName(CurveName, PayloadType);
+								}
+							}
+							break;
+						}
+					}
+				}
+				Parser.ConvertProperty(Prop, SceneNode, PayloadKey);
+			}
 		}
 
 	}
