[Interchange & Legacy FBX Importer] UseT0AsRefPose now updates normals, tangents and binormalSigns if present

#rnx

diff --git a/Engine/Plugins/Interchange/Runtime/Source/Import/Private/Mesh/InterchangeSkeletalMeshFactory.cpp b/Engine/Plugins/Interchange/Runtime/Source/Import/Private/Mesh/InterchangeSkeletalMeshFactory.cpp
--- a/Engine/Plugins/Interchange/Runtime/Source/Import/Private/Mesh/InterchangeSkeletalMeshFactory.cpp
+++ b/Engine/Plugins/Interchange/Runtime/Source/Import/Private/Mesh/InterchangeSkeletalMeshFactory.cpp
@@ -252,6 +252,11 @@ namespace UE
 				FSkeletalMeshAttributes Attributes(MeshDescription);
 				TVertexAttributesRef<FVector3f> VertexPositions = Attributes.GetVertexPositions();
 				FSkinWeightsVertexAttributesRef VertexSkinWeights = Attributes.GetVertexSkinWeights();
+				TVertexInstanceAttributesRef<FVector3f> VertexInstanceNormals = Attributes.GetVertexInstanceNormals();
+				TVertexInstanceAttributesRef<FVector3f> VertexInstanceTangents = Attributes.GetVertexInstanceTangents();
+				TVertexInstanceAttributesRef<float> VertexInstanceBinormalSigns = Attributes.GetVertexInstanceBinormalSigns();
+				TVertexInstanceAttributesConstRef<FVertexID> VertexInstanceVertexIndices = Attributes.GetVertexInstanceVertexIndices();
+				
 
 				for (FVertexID VertexID : MeshDescription.Vertices().GetElementIDs())
 				{
@@ -371,22 +376,96 @@ namespace UE
 					}
 				}
 
-				for (FVertexID VertexID : MeshDescription.Vertices().GetElementIDs())
+				if (VertexInstanceNormals.IsValid())
 				{
-					const int32 VertexIndex = VertexID.GetValue();
-					const FVector lSrcVertex = FVector(DestinationVertexPositions[VertexIndex]);
-					FVector3f& lDstVertex = DestinationVertexPositions[VertexIndex];
-					double Weight = SkinWeights[VertexIndex];
+					for (FVertexID VertexID : MeshDescription.Vertices().GetElementIDs())
+					{
+						const int32 VertexIndex = VertexID.GetValue();
+						const FVector lSrcVertex = FVector(DestinationVertexPositions[VertexIndex]);
+						FVector3f& lDstVertex = DestinationVertexPositions[VertexIndex];
+						double Weight = SkinWeights[VertexIndex];
+
+						// Deform the vertex if there was at least a link with an influence on the vertex,
+						if (!FMath::IsNearlyZero(Weight))
+						{
+							//Apply skinning of all joints
+							lDstVertex = FVector4f(SkinDeformations[VertexIndex].TransformPosition(lSrcVertex));
+							//Normalized, in case the weight is different then 1
+							lDstVertex /= Weight;
+							//Set the new vertex position in the mesh description
+							VertexPositions[VertexID] = FVector3f(FVector4f(SceneNodeTransform.TransformFVector4(FVector(lDstVertex))));
+
+							//Vertex Instance Normals:
+							TArrayView<const FVertexInstanceID> VertexInstanceIDs = MeshDescription.GetVertexVertexInstanceIDs(VertexID);
+							for (const FVertexInstanceID VertexInstanceID : VertexInstanceIDs)
+							{
+								FVector3f& Normal = VertexInstanceNormals[VertexInstanceID];
+
+								if (Normal.IsZero())
+								{
+									continue;
+								}
+
+								FVector NormalTransformed = SkinDeformations[VertexIndex].TransformVector(FVector(Normal));
+								NormalTransformed /= Weight;
+								Normal = FVector4f(SceneNodeTransform.TransformFVector4(NormalTransformed)).GetSafeNormal();
 
-					// Deform the vertex if there was at least a link with an influence on the vertex,
-					if (!FMath::IsNearlyZero(Weight))
+								if (!VertexInstanceTangents.IsValid())
+								{
+									continue;
+								}
+
+								FVector3f& Tangent = VertexInstanceTangents[VertexInstanceID];
+								
+								if (Tangent.IsZero())
+								{
+									continue;
+								}
+
+								float& BinormalSign = VertexInstanceBinormalSigns[VertexInstanceID];
+
+								//Calculate original binormal vector to correct handidness after transform applications:
+								FVector3f OriginalBinormal = FVector3f::CrossProduct(Normal, Tangent) * BinormalSign;
+
+								//Calculate transformed Tangent:
+								FVector TangentTransformed = SkinDeformations[VertexIndex].TransformVector(FVector(Tangent));
+								TangentTransformed /= Weight;
+								Tangent = FVector4f(SceneNodeTransform.TransformFVector4(TangentTransformed)).GetSafeNormal();
+
+								//Calclate transformed binormal vector:
+								FVector BinormalTransformed = SkinDeformations[VertexIndex].TransformVector(FVector(OriginalBinormal));
+								BinormalTransformed /= Weight;
+								BinormalTransformed = SceneNodeTransform.TransformFVector4(BinormalTransformed).GetSafeNormal();
+
+								//Get Candidate Binormal from transformed Tangent and Normal:
+								FVector3f CandidateBinormal = FVector3f::CrossProduct(Normal, Tangent).GetSafeNormal();
+
+								//Get the handidness corrected:
+								// As in: Check if BinormalTransformed and the CandidateBinormal = (NormalTransformed x TangentTransformed) are aligned, if yes +1, if no -1.
+								BinormalSign = (FVector3f::DotProduct(CandidateBinormal, FVector3f(BinormalTransformed)) >= 0.0f) ? +1.0f : -1.0f;
+							}
+						}
+					}
+				}
+				else
+				{
+					for (FVertexID VertexID : MeshDescription.Vertices().GetElementIDs())
 					{
-						//Apply skinning of all joints
-						lDstVertex = FVector4f(SkinDeformations[VertexIndex].TransformPosition(lSrcVertex));
-						//Normalized, in case the weight is different then 1
-						lDstVertex /= Weight;
-						//Set the new vertex position in the mesh description
-						VertexPositions[VertexID] = FVector3f(FVector4f(SceneNodeTransform.TransformFVector4(FVector(lDstVertex))));
+						const int32 VertexIndex = VertexID.GetValue();
+						const FVector lSrcVertex = FVector(DestinationVertexPositions[VertexIndex]);
+						FVector3f& lDstVertex = DestinationVertexPositions[VertexIndex];
+						double Weight = SkinWeights[VertexIndex];
+
+						// Deform the vertex if there was at least a link with an influence on the vertex,
+						if (!FMath::IsNearlyZero(Weight))
+						{
+							//Apply skinning of all joints
+							lDstVertex = FVector4f(SkinDeformations[VertexIndex].TransformPosition(lSrcVertex));
+							//Normalized, in case the weight is different then 1
+							lDstVertex /= Weight;
+							//Set the new vertex position in the mesh description
+							VertexPositions[VertexID] = FVector3f(FVector4f(SceneNodeTransform.TransformFVector4(FVector(lDstVertex))));
+						}
 					}
 				}
 			}
diff --git a/Engine/Source/Editor/UnrealEd/Private/Fbx/FbxSkeletalMeshImport.cpp b/Engine/Source/Editor/UnrealEd/Private/Fbx/FbxSkeletalMeshImport.cpp
--- a/Engine/Source/Editor/UnrealEd/Private/Fbx/FbxSkeletalMeshImport.cpp
+++ b/Engine/Source/Editor/UnrealEd/Private/Fbx/FbxSkeletalMeshImport.cpp
@@ -464,6 +464,72 @@ void FFbxImporter::SkinControlPointsToPose(FSkeletalMeshImportData& ImportData,
 			} 
 		}
 		
+		{
+			FbxAMatrix TotalMatrixForNormal = MeshMatrix.Inverse().Transpose();
+
+			FbxLayer* BaseLayer = FbxMesh->GetLayer(0);
+			FbxLayerElementNormal* LayerElementNormal = BaseLayer->GetNormals();
+			FbxLayerElementTangent* LayerElementTangent = BaseLayer->GetTangents();
+			FbxLayerElementBinormal* LayerElementBinormal = BaseLayer->GetBinormals();
+
+			bool OddNegativeScale = IsOddNegativeScale(MeshMatrix);
+
+			FbxLayerElement::EReferenceMode NormalReferenceMode(FbxLayerElement::eDirect);
+			FbxLayerElement::EMappingMode NormalMappingMode(FbxLayerElement::eByControlPoint);
+			if (LayerElementNormal)
+			{
+				NormalReferenceMode = LayerElementNormal->GetReferenceMode();
+				NormalMappingMode = LayerElementNormal->GetMappingMode();
+			}
+
+			const int32 TriangleCount = FbxMesh->GetPolygonCount();
+			const int32 ExistingTriangleCount = ImportData.Faces.Num();
+			const int32 StartTriangleIndex = ExistingTriangleCount - TriangleCount;
+			for (int32 TriangleIndex = 0; TriangleIndex < TriangleCount; ++TriangleIndex)
+			{
+				SkeletalMeshImportData::FTriangle& Triangle = ImportData.Faces[StartTriangleIndex + TriangleIndex];
+
+				for (int32 VertexIndex = 0; VertexIndex < 3; ++VertexIndex)
+				{
+					int32 TmpIndex = TriangleIndex * 3 + VertexIndex;
+
+					int32 UnrealVertexIndex = OddNegativeScale ? 2 - VertexIndex : VertexIndex;
+
+					int32 ControlPointIndex = FbxMesh->GetPolygonVertex(TriangleIndex, VertexIndex);
+
+					int32 NormalMapIndex = (NormalMappingMode == FbxLayerElement::eByControlPoint) ? ControlPointIndex : TmpIndex;
+					int32 NormalValueIndex = (NormalReferenceMode == FbxLayerElement::eDirect) ? NormalMapIndex : LayerElementNormal->GetIndexArray().GetAt(NormalMapIndex);
+					int32 TangentMapIndex = TmpIndex;
+
+					FbxAMatrix ClusterDeformation = ClusterDeformations[ControlPointIndex].Inverse().Transpose();
+
+					if (ImportData.bHasNormals)
+					{
+						if (ImportData.bHasTangents)
+						{
+							FbxVector4 TangentX = LayerElementTangent->GetDirectArray().GetAt(TangentMapIndex);
+							TangentX = ClusterDeformation.MultT(TangentX);
+							TangentX = TotalMatrixForNormal.MultT(TangentX);
+							Triangle.TangentX[UnrealVertexIndex] = (FVector3f)Converter.ConvertDir(TangentX);
+							Triangle.TangentX[UnrealVertexIndex].Normalize();
+
+							FbxVector4 TangentY = LayerElementBinormal->GetDirectArray().GetAt(TangentMapIndex);
+							TangentY = ClusterDeformation.MultT(TangentY);
+							TangentY = TotalMatrixForNormal.MultT(TangentY);
+							Triangle.TangentY[UnrealVertexIndex] = (FVector3f)Converter.ConvertDir(TangentY);
+							Triangle.TangentY[UnrealVertexIndex].Normalize();
+						}
+
+						FbxVector4 TangentZ = LayerElementNormal->GetDirectArray().GetAt(NormalValueIndex);
+						TangentZ = ClusterDeformation.MultT(TangentZ);
+						TangentZ = TotalMatrixForNormal.MultT(TangentZ);
+						Triangle.TangentZ[UnrealVertexIndex] = (FVector3f)Converter.ConvertDir(TangentZ);
+						Triangle.TangentZ[UnrealVertexIndex].Normalize();
+					}
+				}
+			}
+		}
+
 		// change the vertex position
 		int32 ExistPointNum = ImportData.Points.Num();
 		int32 StartPointIndex = ExistPointNum - VertexCount;
