Fixed reimport strategy not apply on RecomputeNormals property
- This was true for any boolean attribute associated to a property with a bit-field specification, in our case uint8 PropertyName:1.
- Root cause: The unspecialized call to SetAttribute was converted to SetAttribute<uint8>(...) which type did not match the one of the attribute, bool.
- Solution: Add an AttributeType argument to IMPLEMENT_MESH_BUILD_ASSET_TO_VALUE and call SetAttribute<AttributeType>(...) instead.

#rnx

diff --git a/Engine/Plugins/Interchange/Runtime/Source/FactoryNodes/Private/InterchangeMeshFactoryNode.cpp b/Engine/Plugins/Interchange/Runtime/Source/FactoryNodes/Private/InterchangeMeshFactoryNode.cpp
--- a/Engine/Plugins/Interchange/Runtime/Source/FactoryNodes/Private/InterchangeMeshFactoryNode.cpp
+++ b/Engine/Plugins/Interchange/Runtime/Source/FactoryNodes/Private/InterchangeMeshFactoryNode.cpp
@@ -41,19 +41,19 @@ if (GetAttribute<AttributeType>(Macro_Custom##AttributeName##Key.Key, ValueData)
 }																					\
 return bResult;
 
-#define IMPLEMENT_MESH_BUILD_ASSET_TO_VALUE(AttributeName, PropertyName)	\
+#define IMPLEMENT_MESH_BUILD_ASSET_TO_VALUE(AttributeName, AttributeType, PropertyName)	\
 if (UStaticMesh* StaticMesh = Cast<UStaticMesh>(Asset))								\
 {																					\
 	if (StaticMesh->GetNumSourceModels() > 0)										\
 	{																				\
-		return SetAttribute(Macro_Custom##AttributeName##Key.Key, StaticMesh->GetSourceModel(0).BuildSettings.PropertyName);	\
+		return SetAttribute<AttributeType>(Macro_Custom##AttributeName##Key.Key, StaticMesh->GetSourceModel(0).BuildSettings.PropertyName);	\
 	}																				\
 }																					\
 else if (USkeletalMesh* SkeletalMesh = Cast<USkeletalMesh>(Asset))					\
 {																					\
 	if (FSkeletalMeshLODInfo* LodInfo = SkeletalMesh->GetLODInfo(0))				\
 	{																				\
-		return SetAttribute(Macro_Custom##AttributeName##Key.Key, LodInfo->BuildSettings.PropertyName);	\
+		return SetAttribute<AttributeType>(Macro_Custom##AttributeName##Key.Key, LodInfo->BuildSettings.PropertyName);	\
 	}																				\
 }																					\
 return false;
@@ -63,7 +63,7 @@ return false;
 #define IMPLEMENT_MESH_BUILD_VALUE_TO_ASSET(AttributeName, AttributeType, PropertyName)	\
 return false;
 
-#define IMPLEMENT_MESH_BUILD_ASSET_TO_VALUE(AttributeName, PropertyName)	\
+#define IMPLEMENT_MESH_BUILD_ASSET_TO_VALUE(AttributeName, AttributeType, PropertyName)	\
 return false;
 #endif // else WITH_EDITOR
 
@@ -388,7 +388,7 @@ bool UInterchangeMeshFactoryNode::ApplyCustomRecomputeNormalsToAsset(UObject* As
 
 bool UInterchangeMeshFactoryNode::FillCustomRecomputeNormalsFromAsset(UObject* Asset)
 {
-	IMPLEMENT_MESH_BUILD_ASSET_TO_VALUE(RecomputeNormals, bRecomputeNormals);
+	IMPLEMENT_MESH_BUILD_ASSET_TO_VALUE(RecomputeNormals, bool, bRecomputeNormals);
 }
 
 bool UInterchangeMeshFactoryNode::GetCustomRecomputeTangents(bool& AttributeValue) const
@@ -408,7 +408,7 @@ bool UInterchangeMeshFactoryNode::ApplyCustomRecomputeTangentsToAsset(UObject* A
 
 bool UInterchangeMeshFactoryNode::FillCustomRecomputeTangentsFromAsset(UObject* Asset)
 {
-	IMPLEMENT_MESH_BUILD_ASSET_TO_VALUE(RecomputeTangents, bRecomputeTangents);
+	IMPLEMENT_MESH_BUILD_ASSET_TO_VALUE(RecomputeTangents, bool, bRecomputeTangents);
 }
 
 bool UInterchangeMeshFactoryNode::GetCustomUseMikkTSpace(bool& AttributeValue) const
@@ -428,7 +428,7 @@ bool UInterchangeMeshFactoryNode::ApplyCustomUseMikkTSpaceToAsset(UObject* Asset
 
 bool UInterchangeMeshFactoryNode::FillCustomUseMikkTSpaceFromAsset(UObject* Asset)
 {
-	IMPLEMENT_MESH_BUILD_ASSET_TO_VALUE(UseMikkTSpace, bUseMikkTSpace);
+	IMPLEMENT_MESH_BUILD_ASSET_TO_VALUE(UseMikkTSpace, bool, bUseMikkTSpace);
 }
 
 bool UInterchangeMeshFactoryNode::GetCustomComputeWeightedNormals(bool& AttributeValue) const
@@ -448,7 +448,7 @@ bool UInterchangeMeshFactoryNode::ApplyCustomComputeWeightedNormalsToAsset(UObje
 
 bool UInterchangeMeshFactoryNode::FillCustomComputeWeightedNormalsFromAsset(UObject* Asset)
 {
-	IMPLEMENT_MESH_BUILD_ASSET_TO_VALUE(ComputeWeightedNormals, bComputeWeightedNormals);
+	IMPLEMENT_MESH_BUILD_ASSET_TO_VALUE(ComputeWeightedNormals, bool, bComputeWeightedNormals);
 }
 
 bool UInterchangeMeshFactoryNode::GetCustomUseHighPrecisionTangentBasis(bool& AttributeValue) const
@@ -468,7 +468,7 @@ bool UInterchangeMeshFactoryNode::ApplyCustomUseHighPrecisionTangentBasisToAsset
 
 bool UInterchangeMeshFactoryNode::FillCustomUseHighPrecisionTangentBasisFromAsset(UObject* Asset)
 {
-	IMPLEMENT_MESH_BUILD_ASSET_TO_VALUE(UseHighPrecisionTangentBasis, bUseHighPrecisionTangentBasis);
+	IMPLEMENT_MESH_BUILD_ASSET_TO_VALUE(UseHighPrecisionTangentBasis, bool, bUseHighPrecisionTangentBasis);
 }
 
 bool UInterchangeMeshFactoryNode::GetCustomUseFullPrecisionUVs(bool& AttributeValue) const
@@ -488,7 +488,7 @@ bool UInterchangeMeshFactoryNode::ApplyCustomUseFullPrecisionUVsToAsset(UObject*
 
 bool UInterchangeMeshFactoryNode::FillCustomUseFullPrecisionUVsFromAsset(UObject* Asset)
 {
-	IMPLEMENT_MESH_BUILD_ASSET_TO_VALUE(UseFullPrecisionUVs, bUseFullPrecisionUVs);
+	IMPLEMENT_MESH_BUILD_ASSET_TO_VALUE(UseFullPrecisionUVs, bool, bUseFullPrecisionUVs);
 }
 
 bool UInterchangeMeshFactoryNode::GetCustomUseBackwardsCompatibleF16TruncUVs(bool& AttributeValue) const
@@ -507,7 +507,7 @@ bool UInterchangeMeshFactoryNode::ApplyCustomUseBackwardsCompatibleF16TruncUVsTo
 }
 bool UInterchangeMeshFactoryNode::FillCustomUseBackwardsCompatibleF16TruncUVsFromAsset(UObject* Asset)
 {
-	IMPLEMENT_MESH_BUILD_ASSET_TO_VALUE(UseBackwardsCompatibleF16TruncUVs, bUseBackwardsCompatibleF16TruncUVs);
+	IMPLEMENT_MESH_BUILD_ASSET_TO_VALUE(UseBackwardsCompatibleF16TruncUVs, bool, bUseBackwardsCompatibleF16TruncUVs);
 }
 
 bool UInterchangeMeshFactoryNode::GetCustomRemoveDegenerates(bool& AttributeValue) const
@@ -527,7 +527,7 @@ bool UInterchangeMeshFactoryNode::ApplyCustomRemoveDegeneratesToAsset(UObject* A
 
 bool UInterchangeMeshFactoryNode::FillCustomRemoveDegeneratesFromAsset(UObject* Asset)
 {
-	IMPLEMENT_MESH_BUILD_ASSET_TO_VALUE(RemoveDegenerates, bRemoveDegenerates);
+	IMPLEMENT_MESH_BUILD_ASSET_TO_VALUE(RemoveDegenerates, bool, bRemoveDegenerates);
 }
 
 void UInterchangeMeshFactoryNode::CopyWithObject(const UInterchangeFactoryBaseNode* SourceNode, UObject* Object)
diff --git a/Engine/Plugins/Interchange/Runtime/Source/FactoryNodes/Private/InterchangeSkeletalMeshFactoryNode.cpp b/Engine/Plugins/Interchange/Runtime/Source/FactoryNodes/Private/InterchangeSkeletalMeshFactoryNode.cpp
--- a/Engine/Plugins/Interchange/Runtime/Source/FactoryNodes/Private/InterchangeSkeletalMeshFactoryNode.cpp
+++ b/Engine/Plugins/Interchange/Runtime/Source/FactoryNodes/Private/InterchangeSkeletalMeshFactoryNode.cpp
@@ -30,12 +30,12 @@ if (GetAttribute<AttributeType>(Macro_Custom##AttributeName##Key.Key, ValueData)
 }																					\
 return bResult;
 
-#define IMPLEMENT_SKELETALMESH_BUILD_ASSET_TO_VALUE(AttributeName, PropertyName)					\
+#define IMPLEMENT_SKELETALMESH_BUILD_ASSET_TO_VALUE(AttributeName, AttributeType, PropertyName)					\
 if (USkeletalMesh* SkeletalMesh = Cast<USkeletalMesh>(Asset))						\
 {																					\
 	if (FSkeletalMeshLODInfo* LodInfo = SkeletalMesh->GetLODInfo(0))				\
 	{																				\
-		return SetAttribute(Macro_Custom##AttributeName##Key.Key, LodInfo->BuildSettings.PropertyName);	\
+		return SetAttribute<AttributeType>(Macro_Custom##AttributeName##Key.Key, LodInfo->BuildSettings.PropertyName);	\
 	}																				\
 }																					\
 return false;
@@ -45,7 +45,7 @@ return false;
 #define IMPLEMENT_SKELETAL_BUILD_VALUE_TO_ASSET(AttributeName, AttributeType, PropertyName)	\
 return false;
 
-#define IMPLEMENT_SKELETALMESH_BUILD_ASSET_TO_VALUE(AttributeName, PropertyName)					\
+#define IMPLEMENT_SKELETALMESH_BUILD_ASSET_TO_VALUE(AttributeName, AttributeType, PropertyName)					\
 return false;
 
 #endif //else WITH_EDIOTR
@@ -151,7 +151,7 @@ bool UInterchangeSkeletalMeshFactoryNode::ApplyCustomUseHighPrecisionSkinWeights
 }
 bool UInterchangeSkeletalMeshFactoryNode::FillCustomUseHighPrecisionSkinWeightsFromAsset(UObject* Asset)
 {
-	IMPLEMENT_SKELETALMESH_BUILD_ASSET_TO_VALUE(UseHighPrecisionSkinWeights, bUseHighPrecisionSkinWeights);
+	IMPLEMENT_SKELETALMESH_BUILD_ASSET_TO_VALUE(UseHighPrecisionSkinWeights, bool, bUseHighPrecisionSkinWeights);
 }
 
 bool UInterchangeSkeletalMeshFactoryNode::GetCustomThresholdPosition(float& AttributeValue) const
@@ -168,7 +168,7 @@ bool UInterchangeSkeletalMeshFactoryNode::ApplyCustomThresholdPositionToAsset(UO
 }
 bool UInterchangeSkeletalMeshFactoryNode::FillCustomThresholdPositionFromAsset(UObject* Asset)
 {
-	IMPLEMENT_SKELETALMESH_BUILD_ASSET_TO_VALUE(ThresholdPosition, ThresholdPosition);
+	IMPLEMENT_SKELETALMESH_BUILD_ASSET_TO_VALUE(ThresholdPosition, float, ThresholdPosition);
 }
 
 bool UInterchangeSkeletalMeshFactoryNode::GetCustomThresholdTangentNormal(float& AttributeValue) const
@@ -185,7 +185,7 @@ bool UInterchangeSkeletalMeshFactoryNode::ApplyCustomThresholdTangentNormalToAss
 }
 bool UInterchangeSkeletalMeshFactoryNode::FillCustomThresholdTangentNormalFromAsset(UObject* Asset)
 {
-	IMPLEMENT_SKELETALMESH_BUILD_ASSET_TO_VALUE(ThresholdTangentNormal, ThresholdTangentNormal);
+	IMPLEMENT_SKELETALMESH_BUILD_ASSET_TO_VALUE(ThresholdTangentNormal, float, ThresholdTangentNormal);
 }
 
 bool UInterchangeSkeletalMeshFactoryNode::GetCustomThresholdUV(float& AttributeValue) const
@@ -202,7 +202,7 @@ bool UInterchangeSkeletalMeshFactoryNode::ApplyCustomThresholdUVToAsset(UObject*
 }
 bool UInterchangeSkeletalMeshFactoryNode::FillCustomThresholdUVFromAsset(UObject* Asset)
 {
-	IMPLEMENT_SKELETALMESH_BUILD_ASSET_TO_VALUE(ThresholdUV, ThresholdUV);
+	IMPLEMENT_SKELETALMESH_BUILD_ASSET_TO_VALUE(ThresholdUV, float, ThresholdUV);
 }
 
 bool UInterchangeSkeletalMeshFactoryNode::GetCustomMorphThresholdPosition(float& AttributeValue) const
@@ -219,7 +219,7 @@ bool UInterchangeSkeletalMeshFactoryNode::ApplyCustomMorphThresholdPositionToAss
 }
 bool UInterchangeSkeletalMeshFactoryNode::FillCustomMorphThresholdPositionFromAsset(UObject* Asset)
 {
-	IMPLEMENT_SKELETALMESH_BUILD_ASSET_TO_VALUE(MorphThresholdPosition, MorphThresholdPosition);
+	IMPLEMENT_SKELETALMESH_BUILD_ASSET_TO_VALUE(MorphThresholdPosition, float, MorphThresholdPosition);
 }
 
 bool UInterchangeSkeletalMeshFactoryNode::GetCustomBoneInfluenceLimit(int32& AttributeValue) const
@@ -236,7 +236,7 @@ bool UInterchangeSkeletalMeshFactoryNode::ApplyCustomBoneInfluenceLimitToAsset(U
 }
 bool UInterchangeSkeletalMeshFactoryNode::FillCustomBoneInfluenceLimitFromAsset(UObject* Asset)
 {
-	IMPLEMENT_SKELETALMESH_BUILD_ASSET_TO_VALUE(BoneInfluenceLimit, BoneInfluenceLimit);
+	IMPLEMENT_SKELETALMESH_BUILD_ASSET_TO_VALUE(BoneInfluenceLimit, int32, BoneInfluenceLimit);
 }
 
 bool UInterchangeSkeletalMeshFactoryNode::GetCustomMergeMorphTargetShapeWithSameName(bool& AttributeValue) const
diff --git a/Engine/Plugins/Interchange/Runtime/Source/FactoryNodes/Private/InterchangeStaticMeshFactoryNode.cpp b/Engine/Plugins/Interchange/Runtime/Source/FactoryNodes/Private/InterchangeStaticMeshFactoryNode.cpp
--- a/Engine/Plugins/Interchange/Runtime/Source/FactoryNodes/Private/InterchangeStaticMeshFactoryNode.cpp
+++ b/Engine/Plugins/Interchange/Runtime/Source/FactoryNodes/Private/InterchangeStaticMeshFactoryNode.cpp
@@ -30,12 +30,12 @@ if (GetAttribute<AttributeType>(Macro_Custom##AttributeName##Key.Key, ValueData)
 return bResult;
 
 //We take only the lod zero to get back the value
-#define IMPLEMENT_STATICMESH_BUILD_ASSET_TO_VALUE(AttributeName, PropertyName)	\
+#define IMPLEMENT_STATICMESH_BUILD_ASSET_TO_VALUE(AttributeName, AttributeType, PropertyName)	\
 if (UStaticMesh* StaticMesh = Cast<UStaticMesh>(Asset))								\
 {																					\
 	if (StaticMesh->GetNumSourceModels() > 0)										\
 	{																				\
-		return SetAttribute(Macro_Custom##AttributeName##Key.Key, StaticMesh->GetSourceModel(0).BuildSettings.PropertyName);	\
+		return SetAttribute<AttributeType>(Macro_Custom##AttributeName##Key.Key, StaticMesh->GetSourceModel(0).BuildSettings.PropertyName);	\
 	}																				\
 }																					\
 return false;
@@ -45,7 +45,7 @@ return false;
 #define IMPLEMENT_STATICMESH_BUILD_VALUE_TO_ASSET(AttributeName, AttributeType, PropertyName)	\
 return false;
 
-#define IMPLEMENT_STATICMESH_BUILD_ASSET_TO_VALUE(AttributeName, PropertyName)	\
+#define IMPLEMENT_STATICMESH_BUILD_ASSET_TO_VALUE(AttributeName, AttributeType, PropertyName)	\
 return false;
 
 #endif // else WITH_EDITOR
@@ -285,7 +285,7 @@ bool UInterchangeStaticMeshFactoryNode::ApplyCustomBuildReversedIndexBufferToAss
 }
 bool UInterchangeStaticMeshFactoryNode::FillCustomBuildReversedIndexBufferFromAsset(UObject* Asset)
 {
-	IMPLEMENT_STATICMESH_BUILD_ASSET_TO_VALUE(BuildReversedIndexBuffer, bBuildReversedIndexBuffer);
+	IMPLEMENT_STATICMESH_BUILD_ASSET_TO_VALUE(BuildReversedIndexBuffer, bool, bBuildReversedIndexBuffer);
 }
 
 bool UInterchangeStaticMeshFactoryNode::GetCustomGenerateLightmapUVs(bool& AttributeValue) const
@@ -302,7 +302,7 @@ bool UInterchangeStaticMeshFactoryNode::ApplyCustomGenerateLightmapUVsToAsset(UO
 }
 bool UInterchangeStaticMeshFactoryNode::FillCustomGenerateLightmapUVsFromAsset(UObject* Asset)
 {
-	IMPLEMENT_STATICMESH_BUILD_ASSET_TO_VALUE(GenerateLightmapUVs, bGenerateLightmapUVs);
+	IMPLEMENT_STATICMESH_BUILD_ASSET_TO_VALUE(GenerateLightmapUVs, bool, bGenerateLightmapUVs);
 }
 
 bool UInterchangeStaticMeshFactoryNode::GetCustomGenerateDistanceFieldAsIfTwoSided(bool& AttributeValue) const
@@ -320,7 +320,7 @@ bool UInterchangeStaticMeshFactoryNode::ApplyCustomGenerateDistanceFieldAsIfTwoS
 }
 bool UInterchangeStaticMeshFactoryNode::FillCustomGenerateDistanceFieldAsIfTwoSidedFromAsset(UObject* Asset)
 {
-	IMPLEMENT_STATICMESH_BUILD_ASSET_TO_VALUE(GenerateDistanceFieldAsIfTwoSided, bGenerateDistanceFieldAsIfTwoSided);
+	IMPLEMENT_STATICMESH_BUILD_ASSET_TO_VALUE(GenerateDistanceFieldAsIfTwoSided, bool, bGenerateDistanceFieldAsIfTwoSided);
 }
 
 bool UInterchangeStaticMeshFactoryNode::GetCustomSupportFaceRemap(bool& AttributeValue) const
@@ -337,7 +337,7 @@ bool UInterchangeStaticMeshFactoryNode::ApplyCustomSupportFaceRemapToAsset(UObje
 }
 bool UInterchangeStaticMeshFactoryNode::FillCustomSupportFaceRemapFromAsset(UObject* Asset)
 {
-	IMPLEMENT_STATICMESH_BUILD_ASSET_TO_VALUE(SupportFaceRemap, bSupportFaceRemap);
+	IMPLEMENT_STATICMESH_BUILD_ASSET_TO_VALUE(SupportFaceRemap, bool, bSupportFaceRemap);
 }
 
 bool UInterchangeStaticMeshFactoryNode::GetCustomMinLightmapResolution(int32& AttributeValue) const
@@ -354,7 +354,7 @@ bool UInterchangeStaticMeshFactoryNode::ApplyCustomMinLightmapResolutionToAsset(
 }
 bool UInterchangeStaticMeshFactoryNode::FillCustomMinLightmapResolutionFromAsset(UObject* Asset)
 {
-	IMPLEMENT_STATICMESH_BUILD_ASSET_TO_VALUE(MinLightmapResolution, MinLightmapResolution);
+	IMPLEMENT_STATICMESH_BUILD_ASSET_TO_VALUE(MinLightmapResolution, int32, MinLightmapResolution);
 }
 
 bool UInterchangeStaticMeshFactoryNode::GetCustomSrcLightmapIndex(int32& AttributeValue) const
@@ -371,7 +371,7 @@ bool UInterchangeStaticMeshFactoryNode::ApplyCustomSrcLightmapIndexToAsset(UObje
 }
 bool UInterchangeStaticMeshFactoryNode::FillCustomSrcLightmapIndexFromAsset(UObject* Asset)
 {
-	IMPLEMENT_STATICMESH_BUILD_ASSET_TO_VALUE(SrcLightmapIndex, SrcLightmapIndex);
+	IMPLEMENT_STATICMESH_BUILD_ASSET_TO_VALUE(SrcLightmapIndex, int32, SrcLightmapIndex);
 }
 
 bool UInterchangeStaticMeshFactoryNode::GetCustomDstLightmapIndex(int32& AttributeValue) const
@@ -388,7 +388,7 @@ bool UInterchangeStaticMeshFactoryNode::ApplyCustomDstLightmapIndexToAsset(UObje
 }
 bool UInterchangeStaticMeshFactoryNode::FillCustomDstLightmapIndexFromAsset(UObject* Asset)
 {
-	IMPLEMENT_STATICMESH_BUILD_ASSET_TO_VALUE(DstLightmapIndex, DstLightmapIndex);
+	IMPLEMENT_STATICMESH_BUILD_ASSET_TO_VALUE(DstLightmapIndex, int32, DstLightmapIndex);
 }
 
 bool UInterchangeStaticMeshFactoryNode::GetCustomBuildScale3D(FVector& AttributeValue) const
@@ -405,7 +405,7 @@ bool UInterchangeStaticMeshFactoryNode::ApplyCustomBuildScale3DToAsset(UObject*
 }
 bool UInterchangeStaticMeshFactoryNode::FillCustomBuildScale3DFromAsset(UObject* Asset)
 {
-	IMPLEMENT_STATICMESH_BUILD_ASSET_TO_VALUE(BuildScale3D, BuildScale3D);
+	IMPLEMENT_STATICMESH_BUILD_ASSET_TO_VALUE(BuildScale3D, FVector, BuildScale3D);
 }
 
 bool UInterchangeStaticMeshFactoryNode::GetCustomDistanceFieldResolutionScale(float& AttributeValue) const
@@ -422,7 +422,7 @@ bool UInterchangeStaticMeshFactoryNode::ApplyCustomDistanceFieldResolutionScaleT
 }
 bool UInterchangeStaticMeshFactoryNode::FillCustomDistanceFieldResolutionScaleFromAsset(UObject* Asset)
 {
-	IMPLEMENT_STATICMESH_BUILD_ASSET_TO_VALUE(DistanceFieldResolutionScale, DistanceFieldResolutionScale);
+	IMPLEMENT_STATICMESH_BUILD_ASSET_TO_VALUE(DistanceFieldResolutionScale, float, DistanceFieldResolutionScale);
 }
 
 bool UInterchangeStaticMeshFactoryNode::GetCustomDistanceFieldReplacementMesh(FSoftObjectPath& AttributeValue) const
@@ -480,5 +480,5 @@ bool UInterchangeStaticMeshFactoryNode::ApplyCustomMaxLumenMeshCardsToAsset(UObj
 }
 bool UInterchangeStaticMeshFactoryNode::FillCustomMaxLumenMeshCardsFromAsset(UObject* Asset)
 {
-	IMPLEMENT_STATICMESH_BUILD_ASSET_TO_VALUE(MaxLumenMeshCards, MaxLumenMeshCards);
+	IMPLEMENT_STATICMESH_BUILD_ASSET_TO_VALUE(MaxLumenMeshCards, int32, MaxLumenMeshCards);
 }
\ No newline at end of file
