Fixed crash due to null pointer

#rnx
#DisabledValidations CheckBalancedMacros

diff --git a/Engine/Plugins/Interchange/Runtime/Source/Import/Private/Fbx/InterchangeFbxTranslator.cpp b/Engine/Plugins/Interchange/Runtime/Source/Import/Private/Fbx/InterchangeFbxTranslator.cpp
--- a/Engine/Plugins/Interchange/Runtime/Source/Import/Private/Fbx/InterchangeFbxTranslator.cpp
+++ b/Engine/Plugins/Interchange/Runtime/Source/Import/Private/Fbx/InterchangeFbxTranslator.cpp
@@ -276,9 +276,13 @@ bool UInterchangeFbxTranslator::Translate(UInterchangeBaseNodeContainer& BaseNod
 {
 	TRACE_CPUPROFILER_EVENT_SCOPE(UInterchangeFbxTranslator::Translate);
 	using namespace UE::Interchange::Private;
-	// This is to ensure that the keys are not duplicated when the Analytics Event is actually sent.
-	// Changing the Translator Settings in the import dialog will call the Translate function on the Translators again.
-	AnalyticsHelper->ClearAnalyticsEventData(UE::Interchange::Private::FBXAnalyticsIdentifier);
+
+	if (AnalyticsHelper)
+	{
+		// This is to ensure that the keys are not duplicated when the Analytics Event is actually sent.
+		// Changing the Translator Settings in the import dialog will call the Translate function on the Translators again.
+		AnalyticsHelper->ClearAnalyticsEventData(FBXAnalyticsIdentifier);
+	}
 
 	if (HasAnyFlags(RF_ClassDefaultObject | RF_ArchetypeObject))
 	{
