USD: Tiny tweak on primvar baking tooltip to not look like the text is clipped

diff --git a/Engine/Plugins/Interchange/Runtime/Source/Common/Public/Usd/InterchangeUsdDefinitions.h b/Engine/Plugins/Interchange/Runtime/Source/Common/Public/Usd/InterchangeUsdDefinitions.h
--- a/Engine/Plugins/Interchange/Runtime/Source/Common/Public/Usd/InterchangeUsdDefinitions.h
+++ b/Engine/Plugins/Interchange/Runtime/Source/Common/Public/Usd/InterchangeUsdDefinitions.h
@@ -15,7 +15,7 @@ enum class EInterchangeUsdPrimvar : uint8
 	/** Store only primvars in the Mesh Description used for baking to textures (basically <geompropvalue> node from MaterialX shadergraphs that are converted to <image>)*/
 	Bake,
 
-	/** Store all primvars in the MeshDescription, regardless*/
+	/** Store all primvars in the MeshDescription */
 	All
 };
 
