[Interchange][FBX] UFBX - Skeleton Root fixes

- fixed to match FBX SDK's as done in UE-325794
- fixes Editor.Interchange.Import.Asset.FBX.SKM_Cube_Blue test
- support for relatile texture file path in absolute path source field instead of relative

diff --git a/Engine/Plugins/Interchange/Runtime/Source/Parsers/Fbx/Private/Ufbx/UfbxMaterial.cpp b/Engine/Plugins/Interchange/Runtime/Source/Parsers/Fbx/Private/Ufbx/UfbxMaterial.cpp
--- a/Engine/Plugins/Interchange/Runtime/Source/Parsers/Fbx/Private/Ufbx/UfbxMaterial.cpp
+++ b/Engine/Plugins/Interchange/Runtime/Source/Parsers/Fbx/Private/Ufbx/UfbxMaterial.cpp
@@ -54,9 +54,9 @@ namespace UE::Interchange::Private
 		{
 		}
 
-		FString FindTextureFile(ufbx_texture* Texture)
+		FString FindTextureFile(const ufbx_texture& Texture)
 		{
-			FString TextureFilepath = Convert::ToUnrealString(Texture->absolute_filename);
+			FString TextureFilepath = Convert::ToUnrealString((&Texture)->absolute_filename);
 			FPaths::NormalizeFilename(TextureFilepath);
 			if(FPaths::FileExists(TextureFilepath))
 			{
@@ -64,15 +64,22 @@ namespace UE::Interchange::Private
 			}
 			FString FileBasePath = FPaths::GetPath(Parser.SourceFilename);
 
-			FString RelativeToFBXTextureFilepath = FileBasePath / Convert::ToUnrealString(Texture->relative_filename);
+			FString RelativeToFBXTextureFilepath = FileBasePath / Convert::ToUnrealString(Texture.relative_filename);
 			FPaths::NormalizeFilename(RelativeToFBXTextureFilepath);
 			if( FPaths::FileExists(RelativeToFBXTextureFilepath) )
 			{
 				return RelativeToFBXTextureFilepath;
 			}
 
+			FString FilenameToFBXTextureFilepath = FileBasePath / Convert::ToUnrealString(Texture.filename);
+			FPaths::NormalizeFilename(FilenameToFBXTextureFilepath);
+			if( FPaths::FileExists(FilenameToFBXTextureFilepath) )
+			{
+				return FilenameToFBXTextureFilepath;
+			}
+
 			// Some fbx files do not store the actual absolute filename as absolute and it is actually relative.  Try to get it relative to the FBX file we are importing
-			FString AbsoluteAsRelativeToFBXTextureFilepath = Convert::ToUnrealString(Texture->filename);
+			FString AbsoluteAsRelativeToFBXTextureFilepath = FileBasePath / Convert::ToUnrealString(Texture.absolute_filename);
 			FPaths::NormalizeFilename(AbsoluteAsRelativeToFBXTextureFilepath );
 			if( FPaths::FileExists(AbsoluteAsRelativeToFBXTextureFilepath ) )
 			{
@@ -82,12 +89,12 @@ namespace UE::Interchange::Private
 			if (TextureFilepath.IsEmpty())
 			{
 				// Make a name so it could be used as a key or in a message
-				return FString::Printf(TEXT("Texture_%d"), Texture->element_id);
+				return FString::Printf(TEXT("Texture_%d"), Texture.element_id);
 			}
 			return TextureFilepath;
 		}
 
-		const UInterchangeTexture2DNode* CreateTexture2DNode(ufbx_texture* Texture)
+		const UInterchangeTexture2DNode* CreateTexture2DNode(const ufbx_texture& Texture)
 		{
 			const FString TextureFilename = FindTextureFile(Texture);
 			if (TextureFilename.IsEmpty() || !FPaths::FileExists(TextureFilename))
@@ -177,13 +184,13 @@ namespace UE::Interchange::Private
 
 			auto SetTextureFromMap = [&MaterialInstanceNode, &MaterialMapConverter, &bHasUVTransform](FName InputName, const ufbx_material_map& Map, float Factor) -> bool
 			{
-				ufbx_texture* Texture = Map.texture_enabled ? Map.texture : nullptr;
+				const ufbx_texture* Texture = Map.texture_enabled ? Map.texture : nullptr;
 				if (!Texture)
 				{
 					return false;	
 				}
 
-				const UInterchangeTexture2DNode* TextureNode = MaterialMapConverter.CreateTexture2DNode(Texture);
+				const UInterchangeTexture2DNode* TextureNode = MaterialMapConverter.CreateTexture2DNode(*Texture);
 
 				if (!TextureNode)
 				{
diff --git a/Engine/Plugins/Interchange/Runtime/Source/Parsers/Fbx/Private/Ufbx/UfbxParser.cpp b/Engine/Plugins/Interchange/Runtime/Source/Parsers/Fbx/Private/Ufbx/UfbxParser.cpp
--- a/Engine/Plugins/Interchange/Runtime/Source/Parsers/Fbx/Private/Ufbx/UfbxParser.cpp
+++ b/Engine/Plugins/Interchange/Runtime/Source/Parsers/Fbx/Private/Ufbx/UfbxParser.cpp
@@ -279,8 +279,13 @@ bool FUfbxParser::LoadFbxFile(const FString& Filename, UInterchangeBaseNodeConta
 	// SourceNode->SetCustomSourceFrameRateNumerator(FrameRate);
 	// SourceNode->SetCustomSourceFrameRateDenominator(Denominator);
 
-	// #ufbx_todo: check
-	// SourceNode->SetCustomUseLegacySkeletalMeshBakeTransform(true);
+	// Fbx legacy has a special way to bake the skeletal mesh that do not fit the interchange standard
+	// The interchange skeletal mesh factory will read this to use the proper bake transform so it match legacy behavior.
+	// This fix the issue with blender armature bone skip
+	SourceNode->SetCustomUseLegacySkeletalMeshBakeTransform(true);
+
+	//Fbx legacy does not allow Scene Root Nodes to be part of the skeletons (to be joints).
+	SourceNode->SetCustomAllowSceneRootAsJoint(false);
 
 	// #ufbx_todo: FileDetails
 
diff --git a/Engine/Plugins/Interchange/Runtime/Source/Parsers/Fbx/Private/Ufbx/UfbxScene.cpp b/Engine/Plugins/Interchange/Runtime/Source/Parsers/Fbx/Private/Ufbx/UfbxScene.cpp
--- a/Engine/Plugins/Interchange/Runtime/Source/Parsers/Fbx/Private/Ufbx/UfbxScene.cpp
+++ b/Engine/Plugins/Interchange/Runtime/Source/Parsers/Fbx/Private/Ufbx/UfbxScene.cpp
@@ -41,7 +41,7 @@ namespace UE::Interchange::Private
 			return TransformNode;
 		}
 
-		void InitHierarchyRecursively(FUfbxParser& Parser, UInterchangeBaseNodeContainer& NodeContainer, ufbx_node& Node, UInterchangeSceneNode* ParentSceneNode)
+		void InitHierarchyRecursively(FUfbxParser& Parser, UInterchangeBaseNodeContainer& NodeContainer, const ufbx_node& Node, UInterchangeSceneNode* ParentSceneNode)
 		{
 			Convert::FNodeNameAndUid NodeNameAndUid(Parser, ParentSceneNode, Node);
 			UInterchangeSceneNode* SceneNode = CreateTransformNode(Parser,
@@ -58,15 +58,34 @@ namespace UE::Interchange::Private
 		// Inspired by UE::Interchange::Private::FFbxScene::FindCommonJointRootNode, to match how FbxParser creates skeleton
 		// This function makes root of the skeleton to be a node just above topmost bone(node participating in skinning)
 		// Not sure why this is done like this in FBX SDK 
-		ufbx_node* FindSkeletonRoot(ufbx_node* Bone)
+		const ufbx_node* FindSkeletonRoot(const FUfbxParser& Parser, const ufbx_node* Bone)
 		{
 			while (Bone && Bone->parent)
 			{
 				ufbx_node* ParentNode = Bone->parent;
+
+				bool bIsBlenderArmatureBone = false;
+				ufbx_scene& Scene = *Parser.Scene;
+				if (Scene.metadata.exporter == UFBX_EXPORTER_BLENDER_ASCII ||
+					Scene.metadata.exporter == UFBX_EXPORTER_BLENDER_BINARY)
+				{
+					//Hack to support armature dummy node from blender
+					//Users do not want the null attribute node named armature which is the parent of the real root bone in blender fbx file
+					//This is a hack since if a rigid mesh group root node is named "armature" it will be skip
+					const FString ParentName = Parser.GetElementNameRaw(ParentNode->element);
+
+					const ufbx_node* GrandFather = ParentNode->parent;
+					
+					bIsBlenderArmatureBone = (GrandFather == nullptr || GrandFather == Scene.root_node) && (ParentName.Compare(TEXT("armature"), ESearchCase::IgnoreCase) == 0);
+				}
+				// Don't drop Armature bone with multiple nodes under armature since this will remove extra bones from the skeleton
+				const bool bIgnoreBlenderArmatureBone = ParentNode->children.count > 1;
+
 				if (ParentNode->parent &&
-					(ParentNode->attrib_type == UFBX_ELEMENT_EMPTY
-						|| ParentNode->attrib_type == UFBX_ELEMENT_MESH
-						|| ParentNode->attrib_type == UFBX_ELEMENT_BONE))
+					((ParentNode->attrib_type == UFBX_ELEMENT_EMPTY && (!bIsBlenderArmatureBone||bIgnoreBlenderArmatureBone)) || 
+						ParentNode->attrib_type == UFBX_ELEMENT_MESH || 
+						ParentNode->attrib_type == UFBX_ELEMENT_BONE)
+					&& ParentNode != Scene.root_node)
 				{
 					if (ParentNode->mesh && ParentNode->mesh->skin_deformers.count > 0)
 					{
@@ -121,7 +140,13 @@ namespace UE::Interchange::Private
 
 		UInterchangeSceneNode* SceneNode = *SceneNodeFound;
 
-		FTransform GeometryTransform = Convert::ConvertTransform(InNode.geometry_transform);
+		const bool bIsRootNode = &InNode == Parser.Scene->root_node;
+		if (bIsRootNode)
+		{
+			SceneNode->SetCustomIsSceneRoot(true);
+		}
+
+		const FTransform GeometryTransform = Convert::ConvertTransform(InNode.geometry_transform);
 		if (!GeometryTransform.Equals(FTransform::Identity))
 		{
 			SceneNode->SetCustomGeometricTransform(GeometryTransform);
@@ -242,7 +267,7 @@ namespace UE::Interchange::Private
 
 		SceneUtils::InitHierarchyRecursively(Parser, NodeContainer, *RootNode, nullptr);
 
-		for (ufbx_node* Node : Parser.Scene->nodes)
+		for (const ufbx_node* Node : Parser.Scene->nodes)
 		{
 			if (!ensureMsgf(Node, TEXT("Not expected ufbx scene nodes to contain a nullptr")))
 			{
@@ -251,7 +276,7 @@ namespace UE::Interchange::Private
 			
 			if (Node->bone)
 			{
-				if (uint32 SkeletonRoot = SceneUtils::FindSkeletonRoot(Node)->element_id)
+				if (uint32 SkeletonRoot = SceneUtils::FindSkeletonRoot(Parser, Node)->element_id)
 				{
 					CommonJointRootNodes.FindOrAdd(SkeletonRoot);
 				}
@@ -267,7 +292,7 @@ namespace UE::Interchange::Private
 	FString FUfbxScene::GetSkeletonRoot(ufbx_node* Node)
 	{
 		FString* SkeletonRootFound = SkeletonRootPerBone.Find(Node);
-		return SkeletonRootFound ? *SkeletonRootFound : Parser.GetNodeUid(*SceneUtils::FindSkeletonRoot(Node));
+		return SkeletonRootFound ? *SkeletonRootFound : Parser.GetNodeUid(*SceneUtils::FindSkeletonRoot(Parser, Node));
 	}
 }
 
