[Interchange] Optimization to disable recompute tangents on skeletal meshes when Nanite will  be enabled

diff --git a/Engine/Plugins/Interchange/Runtime/Source/Pipelines/Private/InterchangeGenericSkeletalMeshPipeline.cpp b/Engine/Plugins/Interchange/Runtime/Source/Pipelines/Private/InterchangeGenericSkeletalMeshPipeline.cpp
--- a/Engine/Plugins/Interchange/Runtime/Source/Pipelines/Private/InterchangeGenericSkeletalMeshPipeline.cpp
+++ b/Engine/Plugins/Interchange/Runtime/Source/Pipelines/Private/InterchangeGenericSkeletalMeshPipeline.cpp
@@ -339,6 +339,46 @@ UInterchangeSkeletalMeshFactoryNode* UInterchangeGenericMeshPipeline::CreateSkel
 	{
 		return nullptr;
 	}
+
+	auto HasNaniteAssemblyDependencies = [this, &MeshUidsPerLodIndex](const int32 LodIndex)
+	{
+		if (!PipelineMeshesUtilities || !PipelineMeshesUtilities->HasAssemblyMeshDependencies())
+		{
+			return false;
+		}
+
+		const TArray<FString>* MeshUids = MeshUidsPerLodIndex.Find(LodIndex);
+		if (ensure(MeshUids))
+		{
+			for (const FString& MeshUid : *MeshUids)
+			{
+				if (const UInterchangeMeshNode* MeshNode = Cast<const UInterchangeMeshNode>(BaseNodeContainer->GetNode(MeshUid)))
+				{
+					if (MeshNode->GetAssemblyPartDependenciesCount() > 0)
+					{
+						return true;
+					}
+				}
+
+				if (const UInterchangeSceneNode* SceneNode = Cast<const UInterchangeSceneNode>(BaseNodeContainer->GetNode(MeshUid)))
+				{
+					FString MeshNodeUid;
+					if (SceneNode->GetCustomAssetInstanceUid(MeshNodeUid))
+					{
+						if (const UInterchangeMeshNode* MeshNode = Cast<UInterchangeMeshNode>(BaseNodeContainer->GetNode(MeshNodeUid)))
+						{
+							if (MeshNode->GetAssemblyPartDependenciesCount() > 0)
+							{
+								return true;
+							}
+						}
+					}
+				}
+			}
+		}
+
+		return false;
+	};
 	
 	auto GetFirstNodeInfo = [this, &MeshUidsPerLodIndex](const int32 LodIndex, FString& OutFirstMeshNodeUid, int32& OutSceneNodeCount)->const UInterchangeBaseNode*
 	{
@@ -513,9 +553,18 @@ UInterchangeSkeletalMeshFactoryNode* UInterchangeGenericMeshPipeline::CreateSkel
 		SkeletalMeshFactoryNode->SetEnabled(false);
 	}
 
+	// When Nanite is enabled, recomputed tangents are unused. Nanite generates its own tangents
+	// unless explicit tangent data is provided, so we can safely skip recomputation here.
+	bool bRecomputeTangents = CommonMeshesProperties->bRecomputeTangents;
+	if (bRecomputeTangents && HasNaniteAssemblyDependencies(BaseLodIndex))
+	{
+		UE_LOG(LogInterchangePipeline, Log, TEXT("UInterchangeGenericMeshPipeline: Disabling unnecessary tangent recomputation on Nanite skeletalmesh %s"), *DisplayLabel);
+		bRecomputeTangents = false;
+	}
+
 	//Common meshes build options
 	SkeletalMeshFactoryNode->SetCustomRecomputeNormals(CommonMeshesProperties->bRecomputeNormals);
-	SkeletalMeshFactoryNode->SetCustomRecomputeTangents(CommonMeshesProperties->bRecomputeTangents);
+	SkeletalMeshFactoryNode->SetCustomRecomputeTangents(bRecomputeTangents);
 	SkeletalMeshFactoryNode->SetCustomUseMikkTSpace(CommonMeshesProperties->bUseMikkTSpace);
 	SkeletalMeshFactoryNode->SetCustomComputeWeightedNormals(CommonMeshesProperties->bComputeWeightedNormals);
 	SkeletalMeshFactoryNode->SetCustomUseHighPrecisionTangentBasis(CommonMeshesProperties->bUseHighPrecisionTangentBasis);
