[Interchange-FBX] Animation displaystring generation consistency fixes.

#rnx

diff --git a/Engine/Plugins/Interchange/Runtime/Source/Parsers/Fbx/Private/FbxScene.cpp b/Engine/Plugins/Interchange/Runtime/Source/Parsers/Fbx/Private/FbxScene.cpp
--- a/Engine/Plugins/Interchange/Runtime/Source/Parsers/Fbx/Private/FbxScene.cpp
+++ b/Engine/Plugins/Interchange/Runtime/Source/Parsers/Fbx/Private/FbxScene.cpp
@@ -899,15 +899,12 @@ namespace UE
 				if (UInterchangeSceneNode* UnrealNode = const_cast<UInterchangeSceneNode*>(Cast< UInterchangeSceneNode >(NodeContainer.GetNode(NodeUniqueID))))
 				{
 					bool HasSkeletonAttribute = false;
-					auto ApplySkeletonAttribute = [&SDKScene, &HasSkeletonAttribute, &SkeletonRootNodeUids, &NodeUniqueID, &SkeletalAnimationTrackNode, &UnrealNode, &AnimationIndex, &NodeContainer, &Node]()
+					auto ApplySkeletonAttribute = [this, &SDKScene, &HasSkeletonAttribute, &SkeletonRootNodeUids, &NodeUniqueID, &SkeletalAnimationTrackNode, &UnrealNode, &AnimationIndex, &NodeContainer, &Node]()
 					{
 						HasSkeletonAttribute = true;
 						if (SkeletonRootNodeUids.Contains(NodeUniqueID))
 						{
 							FbxAnimStack* CurrentAnimationStack = (FbxAnimStack*)SDKScene->GetSrcObject<FbxAnimStack>(AnimationIndex);
-							FString DisplayString = FFbxConvert::MakeString(CurrentAnimationStack->GetName());
-							constexpr bool bIsJointFalse = false;
-							UE::Interchange::SanitizeName(DisplayString, bIsJointFalse);
 							FString TrackNodeUid = CreateTrackNodeUid(UnrealNode->GetUniqueID(), AnimationIndex);
 							NodeContainer.BreakableIterateNodesOfType<UInterchangeSkeletalAnimationTrackNode>([&TrackNodeUid, &SkeletalAnimationTrackNode](const FString& NodeUid, UInterchangeSkeletalAnimationTrackNode* Node)
 							{
@@ -925,7 +922,7 @@ namespace UE
 								//In this specific instance, the processing of AnimationSequence and their existance are dependent on the "up-coming" nodes 
 								// (aka the root skeleton node should create the SkeletalAnimationTrackNode, but it really only should exist if at least 1 of the joints of the Skeleton has animation.)
 								// (this could potentially be done in a neater way by passing down the stack the SkeleotnNodeUid instead of a potentially 'fake' UInterchangeSkeletalAnimationTrackNode.)
-								SkeletalAnimationTrackNode->InitializeNode(TrackNodeUid, DisplayString, EInterchangeNodeContainerType::TranslatedAsset);
+								SkeletalAnimationTrackNode->InitializeNode(TrackNodeUid, Parser.GetFbxHelper()->GetFbxObjectName(CurrentAnimationStack), EInterchangeNodeContainerType::TranslatedAsset);
 							}
 
 							double FrameRate = FbxTime::GetFrameRate(SDKScene->GetGlobalSettings().GetTimeMode());
@@ -1288,9 +1285,7 @@ namespace UE
 						if(!SkeletalAnimationTrackNode)
 						{
 							SkeletalAnimationTrackNode = NewObject< UInterchangeSkeletalAnimationTrackNode >(&NodeContainer);
-							FString DisplayString = FFbxConvert::MakeString(CurrentAnimationStack->GetName());
-
-							NodeContainer.SetupNode(SkeletalAnimationTrackNode, TrackNodeUid, DisplayString, EInterchangeNodeContainerType::TranslatedAsset);
+							NodeContainer.SetupNode(SkeletalAnimationTrackNode, TrackNodeUid, Parser.GetFbxHelper()->GetFbxObjectName(CurrentAnimationStack), EInterchangeNodeContainerType::TranslatedAsset);
 
 							double FrameRate = FbxTime::GetFrameRate(SDKScene->GetGlobalSettings().GetTimeMode());
 							SkeletalAnimationTrackNode->SetCustomAnimationSampleRate(FrameRate);
