[Interchange][FBX] uFBX Morph animation fixes

- fixes tests
Editor.Interchange.Import.Asset.FBX.MorphTargets
Editor.Interchange.Import.Asset.FBX.ANIM_MorphTarget

diff --git a/Engine/Plugins/Interchange/Runtime/Source/Parsers/Fbx/Private/Ufbx/UfbxAnimation.cpp b/Engine/Plugins/Interchange/Runtime/Source/Parsers/Fbx/Private/Ufbx/UfbxAnimation.cpp
--- a/Engine/Plugins/Interchange/Runtime/Source/Parsers/Fbx/Private/Ufbx/UfbxAnimation.cpp
+++ b/Engine/Plugins/Interchange/Runtime/Source/Parsers/Fbx/Private/Ufbx/UfbxAnimation.cpp
@@ -6,6 +6,8 @@
 #include "UfbxConvert.h"
 
 #include "InterchangeCommonAnimationPayload.h"
+#include "InterchangeMeshNode.h"
+#include "UfbxMesh.h"
 #include "Algo/AnyOf.h"
 #include "Nodes/InterchangeBaseNodeContainer.h"
 
@@ -446,9 +448,9 @@ namespace UE::Interchange::Private
 		return true;
 	}
 
-	void FUfbxAnimation::AddAnimation(FUfbxParser& Parser, FUfbxScene& Scene, UInterchangeBaseNodeContainer& NodeContainer,
-	                                  const TMap<ufbx_blend_shape*, ufbx_mesh*> MeshPerBlendShape)
+	void FUfbxAnimation::AddAnimation(FUfbxParser& Parser, FUfbxScene& Scene, const FUfbxMesh& Meshes, UInterchangeBaseNodeContainer& NodeContainer)
 	{
+		const TMap<const ufbx_blend_shape*, const ufbx_mesh*> MeshPerBlendShape = Meshes.GetBlendShapes();
 		double FrameRate = Parser.Scene->settings.frames_per_second;
 
 		TArray<FString> TransformAnimTrackNodeUids;
@@ -482,19 +484,13 @@ namespace UE::Interchange::Private
 			}
 
 			TConstArrayView<ufbx_baked_node> BakedNodes(BakedAnim->nodes.data, BakedAnim->nodes.count);
-			if (BakedNodes.IsEmpty())
-			{
-				continue;
-			}
 
 			TConstArrayView<ufbx_node*> FbxNodes(Parser.Scene->nodes.data, Parser.Scene->nodes.count);
 
 			FString DisplayString = Convert::ToUnrealString(Stack->name);
 
-			auto GetTrackNode = [&](ufbx_node* Node)
+			auto GetTrackNode = [&](FString SkeletonRootNodeUid)
 			{
-				FString SkeletonRootNodeUid = Scene.GetSkeletonRoot(Node); 
-
 				UInterchangeSkeletalAnimationTrackNode*& SkeletalAnimationTrackNode = TrackNodePerSkeletonRoot.FindOrAdd(SkeletonRootNodeUid);
 				if (!SkeletalAnimationTrackNode)
 				{
@@ -517,9 +513,9 @@ namespace UE::Interchange::Private
 				return SkeletalAnimationTrackNode;
 			};
 
-			for (int32 BakedNodeIndex = 0; BakedNodeIndex < BakedAnim->nodes.count; ++BakedNodeIndex)
+			for (int32 BakedNodeIndex = 0; BakedNodeIndex < BakedNodes.Num(); ++BakedNodeIndex)
 			{
-				const ufbx_baked_node& BakedNode = BakedAnim->nodes[BakedNodeIndex];
+				const ufbx_baked_node& BakedNode = BakedNodes[BakedNodeIndex];
 
 				if (!ensureMsgf(FbxNodes.IsValidIndex(BakedNode.typed_id), TEXT("Unexpected type_id encountered in ufbx baked anim nodes")))
 				{
@@ -531,7 +527,7 @@ namespace UE::Interchange::Private
 				{
 					FString NodeUid = Parser.GetNodeUid(*Node);
 
-					UInterchangeSkeletalAnimationTrackNode* SkeletalAnimationTrackNode = GetTrackNode(Node);
+					UInterchangeSkeletalAnimationTrackNode* SkeletalAnimationTrackNode = GetTrackNode(Scene.GetSkeletonRoot(Node));
 
 					// #ufbx_todo: Need check animated parents - whole subtree has to be animated too?
 					// see UE::Interchange::Private::FFbxAnimation::AddSkeletalTransformAnimation
@@ -557,6 +553,7 @@ namespace UE::Interchange::Private
 				{
 					const ufbx_anim_prop& Prop = Layer->anim_props[AnimPropIndex];
 
+					// Morph animation is DeformPercent
 					if (FCStringAnsi::Strcmp(UFBX_DeformPercent, Prop.prop_name.data)==0)
 					{
 						if (ufbx_blend_channel* Channel = ufbx_as_blend_channel(Prop.element))
@@ -582,11 +579,42 @@ namespace UE::Interchange::Private
 
 							const ufbx_mesh* Mesh = *MeshFound;
 							TSet<UInterchangeSkeletalAnimationTrackNode*> Tracks;
-							for (ufbx_skin_deformer* Deformer : Mesh->skin_deformers)
+
+							const UInterchangeMeshNode*const* MeshNodeFound = Meshes.MeshToMeshNode.Find(Mesh);
+							if (MeshNodeFound)
 							{
-								for (ufbx_skin_cluster* Cluster : Deformer->clusters)
+								const UInterchangeMeshNode* MeshNode = *MeshNodeFound;
+
+								if (MeshNode->IsSkinnedMesh())
 								{
-									Tracks.Add(GetTrackNode(Cluster->bone_node));
+									for (ufbx_skin_deformer* Deformer : Mesh->skin_deformers)
+									{
+										for (ufbx_skin_cluster* Cluster : Deformer->clusters)
+										{
+											Tracks.Add(GetTrackNode(Scene.GetSkeletonRoot(Cluster->bone_node)));
+										}
+									}
+								}
+								else
+								{
+									//Find MeshInstances: where CustomAssetInstanceUid == MeshNode->GetUniqueID
+									// For every occurence create a Track node MeshNode->GetUniqueID
+									// #ufbx_todo: maybe better to keep map of instances per mesh, instead of iterating...
+									TSet<UInterchangeSceneNode*> SkeletonNodes;
+									NodeContainer.IterateNodesOfType<UInterchangeSceneNode>([&](const FString& NodeUid, UInterchangeSceneNode* SceneNode)
+										{
+											FString AssetInstanceUid;
+											if (SceneNode->GetCustomAssetInstanceUid(AssetInstanceUid) && AssetInstanceUid == MeshNode->GetUniqueID())
+											{
+												
+												SkeletonNodes.Add(SceneNode);
+											}
+										});
+									// Make track nodes separately to prevent modifying NodeContainer while iterating it
+									for (UInterchangeSceneNode* SceneNode : SkeletonNodes)
+									{
+										Tracks.Add(GetTrackNode(SceneNode->GetUniqueID()));
+									}
 								}
 							}
 
@@ -597,8 +625,14 @@ namespace UE::Interchange::Private
 
 							FMorphAnimationPayloadContext MorphPayloadContext;
 							MorphPayloadContext.Curve = MorphCurve;
-							Parser.PayloadContexts.Add(PayloadKey, MorphPayloadContext);
 
+							for (const ufbx_blend_keyframe& Keyframe : Channel->keyframes)
+							{
+								MorphPayloadContext.InbetweenFullWeights.Add(Keyframe.target_weight);
+								MorphPayloadContext.InbetweenCurveNames.Add(Parser.GetElementNameDeduplicated(Keyframe.shape->element));
+							}
+
+							Parser.PayloadContexts.Add(PayloadKey, MorphPayloadContext);
 						}
 					}
 
diff --git a/Engine/Plugins/Interchange/Runtime/Source/Parsers/Fbx/Private/Ufbx/UfbxAnimation.h b/Engine/Plugins/Interchange/Runtime/Source/Parsers/Fbx/Private/Ufbx/UfbxAnimation.h
--- a/Engine/Plugins/Interchange/Runtime/Source/Parsers/Fbx/Private/Ufbx/UfbxAnimation.h
+++ b/Engine/Plugins/Interchange/Runtime/Source/Parsers/Fbx/Private/Ufbx/UfbxAnimation.h
@@ -18,6 +18,7 @@ namespace UE::Interchange::Private
 {
 	struct FSkeletalAnimationPayloadContext;
 	class FUfbxScene;
+	class FUfbxMesh;
 
 	class FUfbxAnimation
 	{
@@ -29,7 +30,7 @@ namespace UE::Interchange::Private
 		static bool FetchPropertyAnimationCurves(const FUfbxParser& Parser, const FPropertyAnimationPayloadContext& Animation, TArray<FInterchangeCurve>& OutCurves);
 		static bool FetchPropertyAnimationStepCurves(const FUfbxParser& Parser, const FPropertyAnimationPayloadContext& Animation, TArray<FInterchangeStepCurve>& OutCurves);
 
-		static void AddAnimation(FUfbxParser& Parser, FUfbxScene& Scene, UInterchangeBaseNodeContainer& NodeContainer, const TMap<ufbx_blend_shape*, ufbx_mesh*> MeshPerBlendShape);
+		static void AddAnimation(FUfbxParser& Parser, FUfbxScene& Scene, const FUfbxMesh& Meshes, UInterchangeBaseNodeContainer& NodeContainer);
 
 		static bool AddNodeAttributeCurvesAnimation(FUfbxParser& Parser
 			, const FString NodeUid
diff --git a/Engine/Plugins/Interchange/Runtime/Source/Parsers/Fbx/Private/Ufbx/UfbxMaterial.cpp b/Engine/Plugins/Interchange/Runtime/Source/Parsers/Fbx/Private/Ufbx/UfbxMaterial.cpp
--- a/Engine/Plugins/Interchange/Runtime/Source/Parsers/Fbx/Private/Ufbx/UfbxMaterial.cpp
+++ b/Engine/Plugins/Interchange/Runtime/Source/Parsers/Fbx/Private/Ufbx/UfbxMaterial.cpp
@@ -56,7 +56,9 @@ namespace UE::Interchange::Private
 
 		FString FindTextureFile(const ufbx_texture& Texture)
 		{
-			FString TextureFilepath = Convert::ToUnrealString((&Texture)->absolute_filename);
+			// #ufbx_todo: maybe just for each of "filename" always check absolute/relative? Just need to decide on the order...
+
+			FString TextureFilepath = Convert::ToUnrealString(Texture.absolute_filename);
 			FPaths::NormalizeFilename(TextureFilepath);
 			if(FPaths::FileExists(TextureFilepath))
 			{
diff --git a/Engine/Plugins/Interchange/Runtime/Source/Parsers/Fbx/Private/Ufbx/UfbxMesh.cpp b/Engine/Plugins/Interchange/Runtime/Source/Parsers/Fbx/Private/Ufbx/UfbxMesh.cpp
--- a/Engine/Plugins/Interchange/Runtime/Source/Parsers/Fbx/Private/Ufbx/UfbxMesh.cpp
+++ b/Engine/Plugins/Interchange/Runtime/Source/Parsers/Fbx/Private/Ufbx/UfbxMesh.cpp
@@ -192,6 +192,8 @@ namespace UE::Interchange::Private
 
 			UInterchangeMeshNode* MeshNode = CreateMeshNode(NodeContainer, MeshNameAndUid.Label, MeshNameAndUid.UniqueID);
 
+			MeshToMeshNode.Add(Mesh, MeshNode);
+
 			if (Mesh->skin_deformers.count > 0)
 			{
 				TArray<FString> JointNodeUniqueIDs;
@@ -272,79 +274,72 @@ namespace UE::Interchange::Private
 
 			MeshNode->SetPayLoadKey(PayLoadKey, MeshNode->IsSkinnedMesh() ? EInterchangeMeshPayLoadType::SKELETAL : EInterchangeMeshPayLoadType::STATIC);
 
-			// WIP: morphs
 			for (ufbx_blend_deformer* Deformer : Mesh->blend_deformers)
 			{
-				for (ufbx_blend_channel* Channel : Deformer->channels)
+				const FString MorphTargetName = Parser.GetElementNameDeduplicated(Deformer->element);
+				for (int32 ChannelIndex = 0; ChannelIndex < Deformer->channels.count; ++ChannelIndex)
 				{
-					FString ChannelName = Convert::ToUnrealString(Channel->name);
-					ufbx_blend_shape* Shape = Channel->target_shape;
-					// #ufbx_todo: FBX SDK has multiple targets per channel via FbxBlendShapeChannel::GetTargetShapeCount(see comment on FbxBlendShapeChannel)
-					// check how uFBX handles this
-					
-					FString ShapeName = Parser.GetElementName(Shape->element);
-
-					// #ufbx_todo: 
-					// FString ShapeName;
-					// if (CurrentChannelMorphTargetCount > 1)
-					// {
-					// 	ShapeName = Parser.GetFbxHelper()->GetFbxObjectName(Shape);
-					// }
-					// else
-					// {
-					// 	if (bMightBeBadMAXFile)
-					// 	{
-					// 		ShapeName = Parser.GetFbxHelper()->GetFbxObjectName(Shape);
-					// 	}
-					// 	else
-					// 	{
-					// 		// Maya concatenates the number of the shape to the end of its name, so instead use the name of the channel
-					// 		ShapeName = ChannelName;
-					// 	}
-					// }
-
-					if (ShapeName.IsEmpty())
+					const ufbx_blend_channel* Channel = Deformer->channels[ChannelIndex];
+
+					const int32 CurrentChannelMorphTargetCount = Channel->keyframes.count;
+
+					FString ChannelName = Parser.GetElementNameDeduplicated(Channel->element);
+					// Maya adds the name of the MorphTarget and an underscore to the front of the channel name, so remove it
+					if (ChannelName.StartsWith(MorphTargetName) && (ChannelName.Compare(MorphTargetName, ESearchCase::IgnoreCase) != 0))
 					{
-						// #ufbx_todo:
-						ensure(false);
-						// ShapeName = TEXT("MorphTarget_") + FString::FromInt(ChannelIndex) + TEXT("_") + FString::FromInt(ChannelMorphTargetIndex);
+						ChannelName.RightInline(ChannelName.Len() - (MorphTargetName.Len() + 1), EAllowShrinking::No);
 					}
 
-					Convert::FMeshNameAndUid MorphTargetNameId(Parser, *Shape);
-					const UInterchangeMeshNode* ExistingMorphTargetNode = Cast<const UInterchangeMeshNode>(NodeContainer.GetNode(MorphTargetNameId.UniqueID));
-					if (ExistingMorphTargetNode)
+					for (int32 ChannelMorphTargetIndex = 0; ChannelMorphTargetIndex < Channel->keyframes.count; ++ChannelMorphTargetIndex)
 					{
-						int32 UniqueId = 1;
-						FString NameClash = "_ncl_";
-						FString UniqueIDBase = MorphTargetNameId.UniqueID;
-						while (ExistingMorphTargetNode)
+						const ufbx_blend_keyframe& Keyframe = Channel->keyframes[ChannelMorphTargetIndex];
+						const ufbx_blend_shape* Shape = Keyframe.shape;
+
+						FString ShapeName = CurrentChannelMorphTargetCount > 1
+							? Parser.GetElementNameDeduplicated(Shape->element)
+							// Maya concatenates the number of the shape to the end of its name, so instead use the name of the channel
+							: ChannelName;
+
+						if (ShapeName.IsEmpty())
 						{
-							MorphTargetNameId.UniqueID = UniqueIDBase + NameClash + FString::FromInt(UniqueId++);
-							ExistingMorphTargetNode = Cast<const UInterchangeMeshNode>(NodeContainer.GetNode(MorphTargetNameId.UniqueID));
+							ShapeName = TEXT("MorphTarget_") + FString::FromInt(ChannelIndex) + TEXT("_") + FString::FromInt(ChannelMorphTargetIndex);
 						}
-					}
-					if (!ExistingMorphTargetNode)
-					{
-						UInterchangeMeshNode* MorphTargetNode = CreateMeshNode(NodeContainer, MorphTargetNameId.Label, MorphTargetNameId.UniqueID);
-						const bool bIsMorphTarget = true;
-						MorphTargetNode->SetMorphTarget(bIsMorphTarget);
-						MorphTargetNode->SetMorphTargetName(ShapeName);
-
-						//Create a Mesh node dependency, so the mesh node can retrieve is associate morph target
-						MeshNode->SetMorphTargetDependencyUid(MorphTargetNameId.UniqueID);
-						FString MorphTargetPayLoadKey = MorphTargetNameId.UniqueID;
-						if (ensure(!Parser.PayloadContexts.Contains(MorphTargetPayLoadKey)))
+
+						Convert::FMeshNameAndUid MorphTargetNameId(Parser, *Shape);
+						const UInterchangeMeshNode* ExistingMorphTargetNode = Cast<const UInterchangeMeshNode>(NodeContainer.GetNode(MorphTargetNameId.UniqueID));
+						if (ExistingMorphTargetNode)
 						{
-							Parser.PayloadContexts.Add(MorphTargetPayLoadKey, FMorph{Mesh->typed_id, Shape->typed_id});
+							int32 UniqueId = 1;
+							const FString NameClash = "_ncl_";
+							const FString UniqueIDBase = MorphTargetNameId.UniqueID;
+							while (ExistingMorphTargetNode)
+							{
+								MorphTargetNameId.UniqueID = UniqueIDBase + NameClash + FString::FromInt(UniqueId++);
+								ExistingMorphTargetNode = Cast<const UInterchangeMeshNode>(NodeContainer.GetNode(MorphTargetNameId.UniqueID));
+							}
 						}
-						MorphTargetNode->SetPayLoadKey(MorphTargetPayLoadKey, EInterchangeMeshPayLoadType::MORPHTARGET);
-					}
 
-					if (ensure(!SkeletonRootPerBlendShape.Contains(Shape)))
-					{
-						SkeletonRootPerBlendShape.Add(Shape, Mesh);	
+						if (!ExistingMorphTargetNode)
+						{
+							UInterchangeMeshNode* MorphTargetNode = CreateMeshNode(NodeContainer, MorphTargetNameId.Label, MorphTargetNameId.UniqueID);
+							const bool bIsMorphTarget = true;
+							MorphTargetNode->SetMorphTarget(bIsMorphTarget);
+							MorphTargetNode->SetMorphTargetName(ShapeName);
+
+							//Create a Mesh node dependency, so the mesh node can retrieve is associate morph target
+							MeshNode->SetMorphTargetDependencyUid(MorphTargetNameId.UniqueID);
+							const FString MorphTargetPayLoadKey = MorphTargetNameId.UniqueID;
+							if (ensure(!Parser.PayloadContexts.Contains(MorphTargetPayLoadKey)))
+							{
+								Parser.PayloadContexts.Add(MorphTargetPayLoadKey, FMorph{Mesh->typed_id, Shape->typed_id});
+							}
+							MorphTargetNode->SetPayLoadKey(MorphTargetPayLoadKey, EInterchangeMeshPayLoadType::MORPHTARGET);
+						}
+						if (ensure(!SkeletonRootPerBlendShape.Contains(Shape)))
+						{
+							SkeletonRootPerBlendShape.Add(Shape, Mesh);	
+						}
 					}
-					
 				}
 			}
 		}
diff --git a/Engine/Plugins/Interchange/Runtime/Source/Parsers/Fbx/Private/Ufbx/UfbxMesh.h b/Engine/Plugins/Interchange/Runtime/Source/Parsers/Fbx/Private/Ufbx/UfbxMesh.h
--- a/Engine/Plugins/Interchange/Runtime/Source/Parsers/Fbx/Private/Ufbx/UfbxMesh.h
+++ b/Engine/Plugins/Interchange/Runtime/Source/Parsers/Fbx/Private/Ufbx/UfbxMesh.h
@@ -31,17 +31,19 @@ namespace UE::Interchange::Private
 		static bool FetchSkinnedMesh(FUfbxParser& Parser, FMeshDescription& MeshDescription, ufbx_element* Element, const FTransform& MeshGlobalTransform, TArray<FString>& OutJointUniqueNames);
 		static bool FetchBlendShape(FUfbxParser& Parser, FMeshDescription& MeshDescription, const ufbx_mesh& Mesh, const ufbx_blend_shape& BlendShape, const FTransform& MeshGlobalTransform);
 
-		const TMap<ufbx_blend_shape*, ufbx_mesh*>& GetBlendShapes()
+		const TMap<const ufbx_blend_shape*, const ufbx_mesh*>& GetBlendShapes() const
 		{
 			return SkeletonRootPerBlendShape;
 		}
 
+		TMap<const ufbx_mesh*, UInterchangeMeshNode*> MeshToMeshNode;
+
 	private:
 		UInterchangeMeshNode* CreateMeshNode(UInterchangeBaseNodeContainer& NodeContainer, const FString& NodeDisplayLabel, const FString& NodeUniqueID);
 
 		FUfbxParser& Parser;
 
-		TMap<ufbx_blend_shape*, ufbx_mesh*> SkeletonRootPerBlendShape;
+		TMap<const ufbx_blend_shape*, const ufbx_mesh*> SkeletonRootPerBlendShape;
 	};
 
 }
\ No newline at end of file
diff --git a/Engine/Plugins/Interchange/Runtime/Source/Parsers/Fbx/Private/Ufbx/UfbxParser.cpp b/Engine/Plugins/Interchange/Runtime/Source/Parsers/Fbx/Private/Ufbx/UfbxParser.cpp
--- a/Engine/Plugins/Interchange/Runtime/Source/Parsers/Fbx/Private/Ufbx/UfbxParser.cpp
+++ b/Engine/Plugins/Interchange/Runtime/Source/Parsers/Fbx/Private/Ufbx/UfbxParser.cpp
@@ -307,7 +307,7 @@ void FUfbxParser::FillContainerWithFbxScene(UInterchangeBaseNodeContainer& NodeC
 	Meshes.AddAllMeshes(NodeContainer);
 
 	SceneConverter.ProcessNodes(NodeContainer);
-	FUfbxAnimation::AddAnimation(*this, SceneConverter, NodeContainer, Meshes.GetBlendShapes());
+	FUfbxAnimation::AddAnimation(*this, SceneConverter, Meshes, NodeContainer);
 }
 
 bool FUfbxParser::FetchPayloadData(const FString& PayloadKey, const FString& PayloadFilepath)
