[Interchange][FBX] uFBX option only displayed when CVar is set

Interchange.FeatureFlags.Import.FBX.uFBX
False by default

diff --git a/Engine/Plugins/Interchange/Runtime/Source/Import/Private/Fbx/InterchangeFbxTranslator.cpp b/Engine/Plugins/Interchange/Runtime/Source/Import/Private/Fbx/InterchangeFbxTranslator.cpp
--- a/Engine/Plugins/Interchange/Runtime/Source/Import/Private/Fbx/InterchangeFbxTranslator.cpp
+++ b/Engine/Plugins/Interchange/Runtime/Source/Import/Private/Fbx/InterchangeFbxTranslator.cpp
@@ -45,6 +45,14 @@ static FAutoConsoleVariableRef CCvarInterchangeEnableFBXLevelImport(
 	TEXT("Whether support for FBX level import is enabled."),
 	ECVF_Default);
 
+// Temporarily hide the uFBX parser if LUF is used until it is production ready
+static bool GInterchangeFBXAllowUFBX = false;
+static FAutoConsoleVariableRef CCvarInterchangeFBXAllowUFBX(
+	TEXT("Interchange.FeatureFlags.Import.FBX.uFBX"),
+	GInterchangeFBXAllowUFBX,
+	TEXT("Whether to allow enabling uFBX for Interchange FBX import."),
+	ECVF_Default);
+
 static bool GInterchangeFBXTranslatorUseWorker = false;
 static FAutoConsoleVariableRef CCvarInterchangeFBXTranslatorUseWorker(
 	TEXT("Interchange.FeatureFlags.Translator.UseWorker.FBX"),
@@ -154,8 +162,7 @@ void UInterchangeFbxTranslatorSettings::PostDuplicate(bool bDuplicateForPIE)
 	bUsingLUFCoordinateSystem = AxisDisplayInfo::GetAxisDisplayCoordinateSystem() == EAxisList::LeftUpForward;
 	CoordinateSystemPolicy = bConvertScene ? (bForceFrontXAxis ? EInterchangeCoordinateSystemPolicy::MatchUpForwardAxes : EInterchangeCoordinateSystemPolicy::MatchUpAxis) : EInterchangeCoordinateSystemPolicy::KeepXYZAxes;
 
-	// Temporarily hide the uFBX parser if LUF is used until it is production ready
-	bDisplayUfbxParser = !bUsingLUFCoordinateSystem;
+	bDisplayUfbxParser = !bUsingLUFCoordinateSystem && GInterchangeFBXAllowUFBX;
 }
 
 #if WITH_EDITOR
