Interchange - Prevent static meshes from being built more than once

After investigation, it should be safe to leverage `UStaticMesh::FCommitMeshDescriptionParams.bUseHashAsGuid` since the option simply forces BulkData's Guid to be computed from the PayloadContentId instead of generating a random Guid - This helps `FStaticMeshRenderData::Cache` generate a DerivedDataKey that points to the latest commited/serialized BulkData and skip a build, which is the desired behavior.

#rnx

diff --git a/Engine/Plugins/Interchange/Runtime/Source/Import/Private/Mesh/InterchangeStaticMeshFactory.cpp b/Engine/Plugins/Interchange/Runtime/Source/Import/Private/Mesh/InterchangeStaticMeshFactory.cpp
--- a/Engine/Plugins/Interchange/Runtime/Source/Import/Private/Mesh/InterchangeStaticMeshFactory.cpp
+++ b/Engine/Plugins/Interchange/Runtime/Source/Import/Private/Mesh/InterchangeStaticMeshFactory.cpp
@@ -44,7 +44,7 @@
 #include "EditorFramework/AssetImportData.h"
 #endif //WITH_EDITORONLY_DATA
 
-static bool GInterchangeUseHashAsGuid = false;
+static bool GInterchangeUseHashAsGuid = true;
 static FAutoConsoleVariableRef CCvarInterchangeUseHashAsGuid(
 	TEXT("Interchange.FeatureFlags.Import.UseHashAsGuid"),
 	GInterchangeUseHashAsGuid,
