[Interchange] StaticMesh Factory UID generation fixes. This fixes expectation regarding 'instanced' static mesh imports.

#rnx

diff --git a/Engine/Plugins/Interchange/Runtime/Source/FactoryNodes/Private/InterchangeStaticMeshFactoryNode.cpp b/Engine/Plugins/Interchange/Runtime/Source/FactoryNodes/Private/InterchangeStaticMeshFactoryNode.cpp
--- a/Engine/Plugins/Interchange/Runtime/Source/FactoryNodes/Private/InterchangeStaticMeshFactoryNode.cpp
+++ b/Engine/Plugins/Interchange/Runtime/Source/FactoryNodes/Private/InterchangeStaticMeshFactoryNode.cpp
@@ -481,4 +481,14 @@ bool UInterchangeStaticMeshFactoryNode::ApplyCustomMaxLumenMeshCardsToAsset(UObj
 bool UInterchangeStaticMeshFactoryNode::FillCustomMaxLumenMeshCardsFromAsset(UObject* Asset)
 {
 	IMPLEMENT_STATICMESH_BUILD_ASSET_TO_VALUE(MaxLumenMeshCards, int32, MaxLumenMeshCards);
+}
+
+bool UInterchangeStaticMeshFactoryNode::SetCustomSubstituteUID(const FString& AttributeValue)
+{
+	IMPLEMENT_NODE_ATTRIBUTE_SETTER_NODELEGATE(SubstituteUID, FString)
+}
+
+bool UInterchangeStaticMeshFactoryNode::GetCustomSubstituteUID(FString& AttributeValue) const
+{
+	IMPLEMENT_NODE_ATTRIBUTE_GETTER(SubstituteUID, FString)
 }
\ No newline at end of file
diff --git a/Engine/Plugins/Interchange/Runtime/Source/FactoryNodes/Public/InterchangeStaticMeshFactoryNode.h b/Engine/Plugins/Interchange/Runtime/Source/FactoryNodes/Public/InterchangeStaticMeshFactoryNode.h
--- a/Engine/Plugins/Interchange/Runtime/Source/FactoryNodes/Public/InterchangeStaticMeshFactoryNode.h
+++ b/Engine/Plugins/Interchange/Runtime/Source/FactoryNodes/Public/InterchangeStaticMeshFactoryNode.h
@@ -216,6 +216,23 @@ public:
 	UFUNCTION(BlueprintCallable, Category = "Interchange | Node | StaticMesh")
 	UE_API bool SetCustomMaxLumenMeshCards(const int32& AttributeValue, bool bAddApplyDelegate = true);
 
+
+	/**
+	* Currently specifically used for LOD group nodes.
+	* If an LOD Group is identified as identical to another one (when bakeMesh is turned off),
+	*	then said LOD Group's asset won't be created and the substitute UID will be set to the FactoryNode that's identical to the one at hand.
+	*/
+	UFUNCTION(BlueprintCallable, Category = "Interchange | Node | StaticMesh")
+	UE_API bool SetCustomSubstituteUID(const FString& AttributeValue);
+
+	/**
+	* Gets the Substitute UID, said UID can be used to acquire an identical FactoryNode.
+	* (for more see SetCustomSubstituteUID)
+	*/
+	UFUNCTION(BlueprintCallable, Category = "Interchange | Node | StaticMesh")
+	UE_API bool GetCustomSubstituteUID(FString& AttributeValue) const;
+
+
 private:
 	UE_API virtual void FillAssetClassFromAttribute() override;
 	UE_API virtual bool SetNodeClassFromClassAttribute() override;
@@ -233,6 +250,7 @@ private:
 	const UE::Interchange::FAttributeKey Macro_CustomMaxLumenMeshCardsKey = UE::Interchange::FAttributeKey(TEXT("MaxLumenMeshCards"));
 	const UE::Interchange::FAttributeKey Macro_CustomBuildNaniteKey = UE::Interchange::FAttributeKey(TEXT("BuildNanite"));
 	const UE::Interchange::FAttributeKey Macro_CustomAutoComputeLODScreenSizesKey = UE::Interchange::FAttributeKey(TEXT("AutoComputeLODScreenSizes"));
+	const UE::Interchange::FAttributeKey Macro_CustomSubstituteUIDKey = UE::Interchange::FAttributeKey(TEXT("SubstituteUID"));
 
 
 	UE::Interchange::TArrayAttributeHelper<float> LODScreenSizes;
diff --git a/Engine/Plugins/Interchange/Runtime/Source/Import/Private/Scene/InterchangeStaticMeshActorFactory.cpp b/Engine/Plugins/Interchange/Runtime/Source/Import/Private/Scene/InterchangeStaticMeshActorFactory.cpp
--- a/Engine/Plugins/Interchange/Runtime/Source/Import/Private/Scene/InterchangeStaticMeshActorFactory.cpp
+++ b/Engine/Plugins/Interchange/Runtime/Source/Import/Private/Scene/InterchangeStaticMeshActorFactory.cpp
@@ -116,7 +116,19 @@ void UInterchangeStaticMeshActorFactory::SetupObject_GameThread(const FSetupObje
 			const UInterchangeStaticMeshFactoryNode* StaticMeshFactoryNode = TargetNodeUids.IsEmpty() ? nullptr : Cast<UInterchangeStaticMeshFactoryNode>(Arguments.NodeContainer->GetFactoryNode(TargetNodeUids[0]));
 			if (StaticMeshFactoryNode)
 			{
-				ApplyStaticMeshToComponent(StaticMeshFactoryNode, StaticMeshComponent, *Arguments.NodeContainer, Cast<UInterchangeMeshActorFactoryNode>(Arguments.FactoryNode));
+				UInterchangeMeshActorFactoryNode* ActorFactoryNode = Cast<UInterchangeMeshActorFactoryNode>(Arguments.FactoryNode);
+
+				ApplyStaticMeshToComponent(StaticMeshFactoryNode, StaticMeshComponent, *Arguments.NodeContainer, ActorFactoryNode);
+
+				FString SubstituteUID;
+				if (!StaticMeshComponent->GetStaticMesh() && StaticMeshFactoryNode->GetCustomSubstituteUID(SubstituteUID))
+				{
+					StaticMeshFactoryNode = Cast<UInterchangeStaticMeshFactoryNode>(Arguments.NodeContainer->GetFactoryNode(SubstituteUID));
+					if (StaticMeshFactoryNode)
+					{
+						ApplyStaticMeshToComponent(StaticMeshFactoryNode, StaticMeshComponent, *Arguments.NodeContainer, ActorFactoryNode);
+					}
+				}
 			}
 		}
 	}
diff --git a/Engine/Plugins/Interchange/Runtime/Source/Pipelines/Private/InterchangeGenericMeshPipeline.cpp b/Engine/Plugins/Interchange/Runtime/Source/Pipelines/Private/InterchangeGenericMeshPipeline.cpp
--- a/Engine/Plugins/Interchange/Runtime/Source/Pipelines/Private/InterchangeGenericMeshPipeline.cpp
+++ b/Engine/Plugins/Interchange/Runtime/Source/Pipelines/Private/InterchangeGenericMeshPipeline.cpp
@@ -456,6 +456,8 @@ void UInterchangeGenericMeshPipeline::ExecutePipeline(UInterchangeBaseNodeContai
 
 	//Create geometry cache factory nodes
 	ExecutePreImportPipelineGeometryCache();
+
+	FindAndProcessIdenticalStaticMeshLODGroups(InBaseNodeContainer);
 }
 
 void UInterchangeGenericMeshPipeline::ExecutePostFactoryPipeline(const UInterchangeBaseNodeContainer* InBaseNodeContainer, const FString& NodeKey, UObject* CreatedAsset, bool bIsAReimport)
diff --git a/Engine/Plugins/Interchange/Runtime/Source/Pipelines/Private/InterchangeGenericStaticMeshPipeline.cpp b/Engine/Plugins/Interchange/Runtime/Source/Pipelines/Private/InterchangeGenericStaticMeshPipeline.cpp
--- a/Engine/Plugins/Interchange/Runtime/Source/Pipelines/Private/InterchangeGenericStaticMeshPipeline.cpp
+++ b/Engine/Plugins/Interchange/Runtime/Source/Pipelines/Private/InterchangeGenericStaticMeshPipeline.cpp
@@ -14,6 +14,7 @@
 #include "InterchangeSceneNode.h"
 #include "InterchangeStaticMeshFactoryNode.h"
 #include "InterchangeStaticMeshLodDataNode.h"
+#include "Mesh/InterchangeMeshHelper.h"
 #include "Nodes/InterchangeBaseNode.h"
 #include "Nodes/InterchangeBaseNodeContainer.h"
 #include "Nodes/InterchangeSourceNode.h"
@@ -402,7 +403,9 @@ void UInterchangeGenericMeshPipeline::ExecutePreImportPipelineStaticMesh()
 						}
 						else
 						{
-							//Non instanced geometry cannot have lods
+							// #interchange_lod_rework
+							//Original presumption was: 'Non instanced geometry cannot have lods', which is not quite correct,
+							// will have to fix when refactoring LOD handling.
 							const FInterchangeMeshGeometry& MeshGeometry = PipelineMeshesUtilities->GetMeshGeometryByUid(MeshUid);
 							TArray<FString>& TranslatedNodes = MeshUidsPerLodIndex.FindOrAdd(LodIndexZero);
 							TranslatedNodes.Add(MeshGeometry.MeshUid);
@@ -529,17 +532,87 @@ bool UInterchangeGenericMeshPipeline::MakeMeshFactoryNodeUidAndDisplayLabel(cons
 		return false;
 	}
 
+	auto SetUidLabel = [&NewNodeUid, &DisplayLabel](const UInterchangeBaseNode* BaseNode)
+		{
+			NewNodeUid = BaseNode->GetUniqueID();
+			DisplayLabel = BaseNode->GetDisplayLabel();
+		};
+
+	auto SetUidLabelFromSceneNode = [&NewNodeUid, &DisplayLabel](const UInterchangeSceneNode* SceneNode, bool& bIsLodGroup)
+		{
+			bIsLodGroup = SceneNode->IsSpecializedTypeContains(UE::Interchange::FSceneNodeStaticData::GetLodGroupSpecializeTypeString());
+			//We can't have the ScenNode's id become a FactoryNodeUID as is, so we are adding a \\Mesh\\ prefix.
+			NewNodeUid = (bIsLodGroup ? TEXT("\\MeshLODGroupInstance\\") : TEXT("\\MeshInstance\\")) + SceneNode->GetUniqueID();
+			DisplayLabel = SceneNode->GetDisplayLabel();
+		};
+
 	if (!BaseMeshUid.IsEmpty())
 	{
-		if (const UInterchangeMeshLODContainerNode* LODContainerNode = Cast<const UInterchangeMeshLODContainerNode>(BaseNodeContainer->GetNode(BaseMeshUid)))
+		//BaseMeshUid represents either:
+		//	- UInterchangeSceneNode's UID that is a LOD Group.
+		//	- UInterchangeSceneNode's UID that is instantiating a UInterchangeMeshNode. (via CustomAssetInstanceUid).
+		//	- UInterchangeMeshNode's UID.
+		//	- UInterchangeMeshLODContainerNode's UID.
+		//	- in case of bCombineStaticMeshes == true this ID should be empty for now.
+
+		const UInterchangeBaseNode* BaseNode = BaseNodeContainer->GetNode(BaseMeshUid);
+		if (BaseNode)
 		{
-			NewNodeUid = LODContainerNode->GetUniqueID();
-			DisplayLabel = LODContainerNode->GetDisplayLabel();
+			if (const UInterchangeSceneNode* SceneNode = Cast<const UInterchangeSceneNode>(BaseNode))
+			{
+				//Regardless of the SceneNode being a LOD Group, or just instantiating a Mesh,
+				// if it was marked as the baseMeshUid we should use it's UID and Name:
+
+				bool bIsLodGroup = false;
+				SetUidLabelFromSceneNode(SceneNode, bIsLodGroup);
+
+				if (!CommonMeshesProperties->bBakeMeshes && !bIsLodGroup)
+				{
+					//We recursively check the AssetInstanceUid's chain to make sure we find LOD Group or MeshUID.
+					FString AssetInstanceUid;
+					while (SceneNode && SceneNode->GetCustomAssetInstanceUid(AssetInstanceUid))
+					{
+						const UInterchangeBaseNode* AssetInstanceNode = BaseNodeContainer->GetNode(AssetInstanceUid);
+						SceneNode = Cast<UInterchangeSceneNode>(AssetInstanceNode); //if its a sceneNode without a LOD group specialization we continue the search
+						if (SceneNode && SceneNode->IsSpecializedTypeContains(UE::Interchange::FSceneNodeStaticData::GetLodGroupSpecializeTypeString()))
+						{
+							//if it is a LOD Group specialization then we use the original NewNodeUid. (as then we want the 'instantiation of the LOD Group)
+							NewNodeUid = TEXT("\\MeshLODGroupInstance\\") + BaseMeshUid;
+							return true;
+						}
+						if (const UInterchangeMeshNode* MeshNode = Cast<UInterchangeMeshNode>(AssetInstanceNode))
+						{
+							//if it is a Mesh Node then we confirmed that it is an instanced Mesh
+							// but since we are not baking meshes we want the MeshNode
+							//It cannot be part of a LOD group as that would have broken out sooner due to the LodGroup check.
+							//For this case we are using MeshNode's UID and DisplayName.
+							SetUidLabel(AssetInstanceNode);
 
-			return true;
+							return true;
+						}
+					}
+				}
+				else
+				{
+					return true;
+				}
+			}
+			else if (const UInterchangeMeshLODContainerNode* LODContainerNode = Cast<const UInterchangeMeshLODContainerNode>(BaseNode))
+			{
+				SetUidLabel(BaseNode);
+
+				return true;
+			}
+			else if (const UInterchangeMeshNode* MeshNode = Cast<const UInterchangeMeshNode>(BaseNode))
+			{
+				SetUidLabel(BaseNode);
+
+				return true;
+			}
 		}
 	}
 
+	//If BaseMeshUID is not provided and or BaseMeshUID is not a LOD Group, fallback onto LOD=0 UID[0] to generate UID and Name.
 	for (const TPair<int32, TArray<FString>>& LodIndexAndMeshUids : MeshUidsPerLodIndex)
 	{
 		if (LodIndex == LodIndexAndMeshUids.Key)
@@ -548,37 +621,54 @@ bool UInterchangeGenericMeshPipeline::MakeMeshFactoryNodeUidAndDisplayLabel(cons
 			if (Uids.Num() > 0)
 			{
 				const FString& Uid = Uids[0];
-				const UInterchangeBaseNode* Node = BaseNodeContainer->GetNode(Uid);
+				const UInterchangeBaseNode* BaseNode = BaseNodeContainer->GetNode(Uid);
 
-				if (const UInterchangeMeshNode* MeshNode = Cast<const UInterchangeMeshNode>(Node))
+				if (const UInterchangeMeshNode* MeshNode = Cast<const UInterchangeMeshNode>(BaseNode))
 				{
-					DisplayLabel = Node->GetDisplayLabel();
-					NewNodeUid = Uid;
+					SetUidLabel(BaseNode);
+
 					return true;
 				}
 
-				if (const UInterchangeSceneNode* SceneNode = Cast<const UInterchangeSceneNode>(Node))
+				if (const UInterchangeSceneNode* SceneNode = Cast<const UInterchangeSceneNode>(BaseNode))
 				{
-					FString RefMeshUid;
-					if (SceneNode->GetCustomAssetInstanceUid(RefMeshUid))
+					bool bIsLodGroup = false;
+					SetUidLabelFromSceneNode(SceneNode, bIsLodGroup);
+
+					if (!CommonMeshesProperties->bBakeMeshes && !bIsLodGroup)
 					{
-						const UInterchangeBaseNode* MeshNode = BaseNodeContainer->GetNode(RefMeshUid);
-						if (MeshNode)
+						//We recursively check the AssetInstanceUid's chain to make sure we find LOD Group or MeshUID.
+						FString AssetInstanceUid;
+						while (SceneNode && SceneNode->GetCustomAssetInstanceUid(AssetInstanceUid))
 						{
-							DisplayLabel = MeshNode->GetDisplayLabel();
-							if (Uids.Num() == 1)
+							const UInterchangeBaseNode* AssetInstanceNode = BaseNodeContainer->GetNode(AssetInstanceUid);
+							SceneNode = Cast<UInterchangeSceneNode>(AssetInstanceNode); //if its a sceneNode without a LOD group specialization we continue the search
+							if (SceneNode && SceneNode->IsSpecializedTypeContains(UE::Interchange::FSceneNodeStaticData::GetLodGroupSpecializeTypeString()))
 							{
-								// If we are instancing one scene node, we want to use it to name the mesh
-								DisplayLabel = SceneNode->GetDisplayLabel();
+								//if it is a LOD Group specialization then we use the original NewNodeUid. (as then we want the 'instantiation' of the LOD Group)
+								NewNodeUid = TEXT("\\MeshLODGroupInstance\\") + Uid;
+								
+								return true;
 							}
+							if (const UInterchangeMeshNode* MeshNode = Cast<UInterchangeMeshNode>(AssetInstanceNode))
+							{
+								//if it is a Mesh Node then we confirmed that it is an instanced Mesh
+								// but since we are not baking meshes we want the MeshNode
+								//It cannot be part of a LOD group as that would have broken out sooner due to the LodGroup check.
+								//For this case we are using MeshNode's UID and DisplayName.
+								SetUidLabel(AssetInstanceNode);
 
-							NewNodeUid = RefMeshUid;
-							return true;
+								return true;
+							}
 						}
 					}
+					else
+					{
+						return true;
+					}
 				}
 
-				if (const UInterchangeInstancedStaticMeshComponentNode* ISMComponentNode = Cast<const UInterchangeInstancedStaticMeshComponentNode>(Node))
+				if (const UInterchangeInstancedStaticMeshComponentNode* ISMComponentNode = Cast<const UInterchangeInstancedStaticMeshComponentNode>(BaseNode))
 				{
 					FString RefMeshUid;
 					if (ISMComponentNode->GetCustomInstancedAssetUid(RefMeshUid))
@@ -586,8 +676,8 @@ bool UInterchangeGenericMeshPipeline::MakeMeshFactoryNodeUidAndDisplayLabel(cons
 						const UInterchangeBaseNode* MeshNode = BaseNodeContainer->GetNode(RefMeshUid);
 						if (MeshNode)
 						{
-							DisplayLabel = MeshNode->GetDisplayLabel();
-							NewNodeUid = RefMeshUid;
+							SetUidLabel(MeshNode);
+
 							return true;
 						}
 					}
@@ -620,14 +710,37 @@ UInterchangeStaticMeshFactoryNode* UInterchangeGenericMeshPipeline::CreateStatic
 		return nullptr;
 	}
 
+	// #interchange_lod_rework: if rework will be calling out LOD Groups directly this won't be needed.
+	auto CheckAndStoreBaseMeshUidForLODGroup = [this, BaseMeshUid](UInterchangeStaticMeshFactoryNode* StaticMeshFactoryNode)
+		{
+			if (const UInterchangeSceneNode* BaseMeshSceneNode = Cast<UInterchangeSceneNode>(BaseNodeContainer->GetNode(BaseMeshUid)))
+			{
+				if (BaseMeshSceneNode->IsSpecializedTypeContains(UE::Interchange::FSceneNodeStaticData::GetLodGroupSpecializeTypeString()))
+				{
+					StaticMeshFactoryNode->RegisterAttribute<FString>(UE::Interchange::FAttributeKey(FString(TEXT("LODGroupSceneNodeUid"))), BaseMeshUid);
+				}
+			}
+		};
+
 	const FString StaticMeshUid = UInterchangeFactoryBaseNode::BuildFactoryNodeUid(StaticMeshUid_MeshNamePart);
-	UInterchangeStaticMeshFactoryNode* StaticMeshFactoryNode = NewObject<UInterchangeStaticMeshFactoryNode>(BaseNodeContainer, NAME_None);
+
+	UInterchangeFactoryBaseNode* FactoryBaseNode = BaseNodeContainer->GetFactoryNode(StaticMeshUid);
+	UInterchangeStaticMeshFactoryNode* StaticMeshFactoryNode = Cast<UInterchangeStaticMeshFactoryNode>(FactoryBaseNode);
+
+	if (StaticMeshFactoryNode)
+	{
+		CheckAndStoreBaseMeshUidForLODGroup(StaticMeshFactoryNode);
+		return StaticMeshFactoryNode;
+	}
+
+	StaticMeshFactoryNode = NewObject<UInterchangeStaticMeshFactoryNode>(BaseNodeContainer, NAME_None);
 	if (!ensure(StaticMeshFactoryNode))
 	{
 		return nullptr;
 	}
 
 	StaticMeshFactoryNode->InitializeStaticMeshNode(StaticMeshUid, DisplayLabel, UStaticMesh::StaticClass()->GetName(), BaseNodeContainer);
+	CheckAndStoreBaseMeshUidForLODGroup(StaticMeshFactoryNode);
 
 	//Set the pipeline import sockets property
 	StaticMeshFactoryNode->SetCustomImportSockets(CommonMeshesProperties->bImportSockets);
@@ -852,3 +965,161 @@ void UInterchangeGenericMeshPipeline::AddLodDataToStaticMesh(UInterchangeStaticM
 		UE::Interchange::MeshesUtilities::ReorderSlotMaterialDependencies(*StaticMeshFactoryNode, *BaseNodeContainer);
 	}
 }
+
+// #interchange_lod_rework
+void UInterchangeGenericMeshPipeline::FindAndProcessIdenticalStaticMeshLODGroups(const UInterchangeBaseNodeContainer* InBaseNodeContainer)
+{
+	if (!CommonMeshesProperties->bBakeMeshes)
+	{
+		auto GetGlobalTransform = [&InBaseNodeContainer](
+			const FString& InMeshOrSceneNodeUid,
+			FTransform& OutGlobalTransform,
+			bool& bMeshNode,
+			FString& LODMeshAssetUid
+			)
+			{
+				bMeshNode = true;
+				const UInterchangeBaseNode* Node = InBaseNodeContainer->GetNode(InMeshOrSceneNodeUid);
+				const UInterchangeMeshNode* MeshNode = Cast<const UInterchangeMeshNode>(Node);
+				if (MeshNode == nullptr)
+				{
+					bMeshNode = false;
+					// MeshUid must refer to a scene node
+					const UInterchangeSceneNode* SceneNode = Cast<const UInterchangeSceneNode>(Node);
+					if (!ensure(SceneNode))
+					{
+						return;
+					}
+
+					// Get the transform from the scene node
+					FTransform SceneNodeGlobalTransform;
+					if (SceneNode->GetCustomGlobalTransform(InBaseNodeContainer, FTransform::Identity, SceneNodeGlobalTransform))
+					{
+						OutGlobalTransform = SceneNodeGlobalTransform;
+					}
+					UE::Interchange::Private::MeshHelper::AddSceneNodeGeometricAndPivotToGlobalTransform(OutGlobalTransform, SceneNode, true /*bakeMeshes*/, true /*bakePivotMeshes*/);
+					// And get the mesh node which it references
+					FString MeshDependencyUid;
+					SceneNode->GetCustomAssetInstanceUid(MeshDependencyUid);
+					LODMeshAssetUid = MeshDependencyUid;
+				}
+				else
+				{
+					LODMeshAssetUid = InMeshOrSceneNodeUid;
+				}
+			};
+
+		//Iterate through all UInterchangeStaticMeshFactoryNode and build a helper struct that will help us identify identical LOD Group Static Meshes.
+		struct FLODGroupIIHelper //Identical Identifier
+		{
+			TMap<FString, FTransform> Meshes;
+			UInterchangeStaticMeshFactoryNode* LODGroupStaticMeshFactoryNode;
+
+			TArray<UInterchangeStaticMeshFactoryNode*> DisabledIdenticals;
+
+			bool DataEquals(const FLODGroupIIHelper& LODGroupIIHelper)
+			{
+				if (Meshes.Num() != LODGroupIIHelper.Meshes.Num())
+				{
+					return false;
+				}
+
+				for (const TPair<FString, FTransform>& Entry : Meshes)
+				{
+					if (!LODGroupIIHelper.Meshes.Contains(Entry.Key))
+					{
+						return false;
+					}
+					if (!LODGroupIIHelper.Meshes[Entry.Key].Equals(Entry.Value))
+					{
+						return false;
+					}
+				}
+
+				return true;
+			}
+		};
+		TArray<FLODGroupIIHelper> LODGroupIdenticalIdentifier;
+		auto ProcessEntryCandidate = [&LODGroupIdenticalIdentifier](FLODGroupIIHelper& ToCompare)
+			{
+				for (FLODGroupIIHelper& LODGroupIIHelper : LODGroupIdenticalIdentifier)
+				{
+					//TODO/Note: we should compare all attributes of the UInterchangeStaticMeshFactoryNode to make sure that regardless of Mesh and Transforms they are identical and can be subsituted.
+					//For now however we are only checking Mesh and Transform.
+					//For now this is fine as the rest of the UInterchangeStaticMeshFactoryNode attributes are set from the Pipeline and Utilities properties which are shared across.
+					if (LODGroupIIHelper.DataEquals(ToCompare))
+					{
+						LODGroupIIHelper.DisabledIdenticals.Add(ToCompare.LODGroupStaticMeshFactoryNode);
+						ToCompare.LODGroupStaticMeshFactoryNode->SetEnabled(false);
+						ToCompare.LODGroupStaticMeshFactoryNode->SetCustomSubstituteUID(LODGroupIIHelper.LODGroupStaticMeshFactoryNode->GetUniqueID());
+						return;
+					}
+				}
+				LODGroupIdenticalIdentifier.Add(ToCompare);
+			};
+
+		InBaseNodeContainer->IterateNodesOfType<UInterchangeStaticMeshFactoryNode>(
+			[&InBaseNodeContainer, &GetGlobalTransform, &LODGroupIdenticalIdentifier, &ProcessEntryCandidate](const FString& NodeUid, UInterchangeStaticMeshFactoryNode* StaticMeshFactoryNode)
+			{
+				FString LODGroupSceneUid;
+				if (StaticMeshFactoryNode->GetAttribute<FString>(TEXT("LODGroupSceneNodeUid"), LODGroupSceneUid))
+				{
+					const UInterchangeSceneNode* LODGroupSceneNode = Cast<UInterchangeSceneNode>(InBaseNodeContainer->GetNode(LODGroupSceneUid));
+					if (!LODGroupSceneNode)
+					{
+						return;
+					}
+
+					FLODGroupIIHelper LODGroupIIHelper;
+					LODGroupIIHelper.LODGroupStaticMeshFactoryNode = StaticMeshFactoryNode;
+
+					//Get LOD Group's Global Transform:
+					FTransform LODGroupGlobalTransform;
+					bool bLODGroupMeshNode; //Dummy
+					FString LODMeshAssetUid; //Dummy
+					GetGlobalTransform(LODGroupSceneUid, LODGroupGlobalTransform, bLODGroupMeshNode, LODMeshAssetUid);
+					FTransform LODGroupGlobalTransformInverse = LODGroupGlobalTransform.Inverse();
+
+					TArray<FString> LodDataUniqueIds;
+					StaticMeshFactoryNode->GetLodDataUniqueIds(LodDataUniqueIds);
+
+					for (size_t LODIndex = 0; LODIndex < LodDataUniqueIds.Num(); LODIndex++)
+					{
+						const FString& LodUniqueId = LodDataUniqueIds[LODIndex];
+						const FString LODIndexString = FString::FromInt(LODIndex);
+						const UInterchangeStaticMeshLodDataNode* LodDataNode = Cast<UInterchangeStaticMeshLodDataNode>(InBaseNodeContainer->GetNode(LodUniqueId));
+						if (!ensure(LodDataNode))
+						{
+							continue;
+						}
+
+						TArray<FString> MeshUids;
+						LodDataNode->GetMeshUids(MeshUids);
+
+						for (const FString& MeshUid : MeshUids)
+						{
+							FTransform MeshGlobalTransform;
+							bool bMeshNode;
+							FString MeshAssetUid;
+							GetGlobalTransform(MeshUid, MeshGlobalTransform, bMeshNode, MeshAssetUid);
+
+							if (!MeshAssetUid.IsEmpty())
+							{
+								if (bMeshNode)
+								{
+									LODGroupIIHelper.Meshes.Add(LODIndexString + TEXT("_") + MeshAssetUid, FTransform::Identity);
+								}
+								else
+								{
+									//Note: Need the relative transform from the LODGroup not from the root for equal check purposes.
+									LODGroupIIHelper.Meshes.Add(LODIndexString + TEXT("_") + MeshAssetUid, LODGroupGlobalTransformInverse * MeshGlobalTransform);
+								}
+							}
+						}
+					}
+
+					ProcessEntryCandidate(LODGroupIIHelper);
+				}
+			});
+	}
+}
diff --git a/Engine/Plugins/Interchange/Runtime/Source/Pipelines/Private/InterchangePipelineMeshesUtilities.cpp b/Engine/Plugins/Interchange/Runtime/Source/Pipelines/Private/InterchangePipelineMeshesUtilities.cpp
--- a/Engine/Plugins/Interchange/Runtime/Source/Pipelines/Private/InterchangePipelineMeshesUtilities.cpp
+++ b/Engine/Plugins/Interchange/Runtime/Source/Pipelines/Private/InterchangePipelineMeshesUtilities.cpp
@@ -75,7 +75,7 @@ namespace UE::Private::InterchangeMeshPipeline
 						{
 							int32 LodNumber = FPlatformString::Atoi(*LODXNumber);
 							FString MatchName = ParentUniqueID;
-							
+
 							TArray<FString>& LodChildUids = TmpSceneMeshNodeUidsPerLodParentUidMap.FindOrAdd(MatchName);
 							//Add LOD at the correct index
 							if (LodNumber >= LodChildUids.Num())
@@ -107,6 +107,14 @@ namespace UE::Private::InterchangeMeshPipeline
 			if (LodChildUids.Num() > 1)
 			{
 				SceneMeshNodeUidsPerLodParentUidMap.Add(LodPrefixNodePair);
+
+				//Make sure that the Parent now has the LOD specialization, so that StaticMesh Creation can track the LOD Group instantiation as necessary.
+				// #interchange_lod_rework: rework as needed.
+				const UInterchangeSceneNode* SceneNodeConst = Cast<UInterchangeSceneNode>(BaseNodeContainer->GetNode(LodPrefixNodePair.Key));
+				if (UInterchangeSceneNode* SceneNode = const_cast<UInterchangeSceneNode*>(SceneNodeConst))
+				{
+					SceneNode->AddSpecializedType(UE::Interchange::FSceneNodeStaticData::GetLodGroupSpecializeTypeString());
+				}
 			}
 		}
 	}
diff --git a/Engine/Plugins/Interchange/Runtime/Source/Pipelines/Public/InterchangeGenericMeshPipeline.h b/Engine/Plugins/Interchange/Runtime/Source/Pipelines/Public/InterchangeGenericMeshPipeline.h
--- a/Engine/Plugins/Interchange/Runtime/Source/Pipelines/Public/InterchangeGenericMeshPipeline.h
+++ b/Engine/Plugins/Interchange/Runtime/Source/Pipelines/Public/InterchangeGenericMeshPipeline.h
@@ -494,6 +494,15 @@ private:
 	/* Geometry cache API END                                               */
 	/************************************************************************/
 
+	/**
+	* Works on Static Mesh Factory Nodes that are LOD Groups.
+	* Finds identical LOD groups. DisableS 'duplicate' StaticMeshFactoryNodes. Sets Substitute UIDs for disabled nodes.
+	* Finds Identical LOD group based on Meshes and their transforms related to the LOD Group Node for now.
+	* bBakeMeshes==false is a criteria.
+	* #interchange_lod_rework
+	*/
+	UE_API void FindAndProcessIdenticalStaticMeshLODGroups(const UInterchangeBaseNodeContainer* InBaseNodeContainer);
+
 public:
 
 		UE_API virtual void PostLoad() override;
