MaterialX:
- Add custom nodedefs path from MATERIALX_SEARCH_PATH environment variable

diff --git a/Engine/Plugins/Interchange/Runtime/Source/Import/Private/MaterialX/MaterialXUtils/MaterialXManager.cpp b/Engine/Plugins/Interchange/Runtime/Source/Import/Private/MaterialX/MaterialXUtils/MaterialXManager.cpp
--- a/Engine/Plugins/Interchange/Runtime/Source/Import/Private/MaterialX/MaterialXUtils/MaterialXManager.cpp
+++ b/Engine/Plugins/Interchange/Runtime/Source/Import/Private/MaterialX/MaterialXUtils/MaterialXManager.cpp
@@ -514,6 +514,12 @@ bool FMaterialXManager::Translate(const FString& Filename, UInterchangeBaseNodeC
 		mx::readFromXmlFile(Document, TCHAR_TO_UTF8(*Filename));
 		Document->importLibrary(MaterialXLibrary);
 
+		// The path to the folders containing the custom node nodedefs should also be indicated with the MATERIALX_SEARCH_PATH environment variable 
+		mx::DocumentPtr EnvPathLibrary = mx::createDocument();
+		mx::FileSearchPath EnvPath = mx::getEnvironmentPath();
+		mx::StringSet EnvLibraries = mx::loadLibraries({}, EnvPath, EnvPathLibrary);
+		Document->importLibrary(EnvPathLibrary);
+
 		bIsDocumentValid = Translate(Document, BaseNodeContainer, Translator);
 
 		if (Document->hasVersionString())
