[Interchange-FBX] Refactoring the handling of Node Names to be more aligned with Legacy FBX.

#rnx

diff --git a/Engine/Plugins/Interchange/Runtime/Source/Parsers/Fbx/Private/FbxAPI.cpp b/Engine/Plugins/Interchange/Runtime/Source/Parsers/Fbx/Private/FbxAPI.cpp
--- a/Engine/Plugins/Interchange/Runtime/Source/Parsers/Fbx/Private/FbxAPI.cpp
+++ b/Engine/Plugins/Interchange/Runtime/Source/Parsers/Fbx/Private/FbxAPI.cpp
@@ -69,7 +69,7 @@ namespace UE
 			{
 				if (!FbxHelper.IsValid())
 				{
-					FbxHelper = MakeShared<FFbxHelper>();
+					FbxHelper = MakeShared<FFbxHelper>(bKeepFbxNamespace);
 				}
 				check(FbxHelper.IsValid());
 				return FbxHelper;
@@ -122,23 +122,29 @@ namespace UE
 				FbxManager::GetFileFormatVersion(SDKMajor, SDKMinor, SDKRevision);
 
 				// Initialize the importer by providing a filename.
+				auto ReportOpenError = [this, &Filename]()
+					{
+						FFormatNamedArguments FilenameText
+						{
+							{ TEXT("Filename"), FText::FromString(Filename) }
+						};
+						UInterchangeResultError_Generic* Message = AddMessage<UInterchangeResultError_Generic>();
+						Message->Text = FText::Format(LOCTEXT("CannotOpenFBXFile", "Cannot open FBX file '{Filename}'."), FilenameText);
+					};
+
 				const bool bImportStatus = SDKImporter->Initialize(TCHAR_TO_UTF8(*Filename));
 				if (!bImportStatus)
 				{
-					FFormatNamedArguments FilenameText
-					{
-						{ TEXT("Filename"), FText::FromString(Filename) }
-					};
-					UInterchangeResultError_Generic* Message = AddMessage<UInterchangeResultError_Generic>();
-					Message->Text = FText::Format(LOCTEXT("CannotOpenFBXFile", "Cannot open FBX file '{Filename}'."), FilenameText);
+					ReportOpenError();
 					return false;
 				}
 
 				bool bStatus = SDKImporter->Import(SDKScene);
-
-				//To be able to re-import legacy fbx imported skeleton hierarchy we need to rename bones the same way legacy was doing, so this rename precede the CleanupFbxData renaming
-				constexpr bool bRemovePath = true;
-				EnsureNodeNameAreValid(FPaths::GetBaseFilename(Filename, bRemovePath));
+				if (!bStatus)
+				{
+					ReportOpenError();
+					return false;
+				}
 
 				//We always convert scene to UE axis and units
 				FbxAMatrix AxisConversionInverseMatrix;
@@ -310,95 +316,60 @@ namespace UE
 				return bResult;
 			}
 
-			void ManageNamespace(const bool bKeepFbxNamespace, FString& ObjectName, FbxObject* Object)
+			void FFbxParser::CleanupFbxData()
 			{
-				if (bKeepFbxNamespace)
-				{
-					if (ObjectName.Contains(TEXT(":")))
+				auto ReportAndSetNoName = [this](FbxObject* Object, const FString& ObjectName)
 					{
-						ObjectName.ReplaceInline(TEXT(":"), TEXT("_"));
 						Object->SetName(TCHAR_TO_UTF8(*ObjectName));
-					}
-				}
-				else
-				{
-					// Remove namespaces
-					int32 LastNamespaceTokenIndex = INDEX_NONE;
-					if (ObjectName.FindLastChar(TEXT(':'), LastNamespaceTokenIndex))
-					{
-						//+1 to remove the ':' character we found
-						ObjectName.RightChopInline(LastNamespaceTokenIndex + 1, EAllowShrinking::Yes);
-						Object->SetName(TCHAR_TO_UTF8(*ObjectName));
-					}
-				}
-			}
-
-			void FFbxParser::EnsureNodeNameAreValid(const FString& BaseFilename)
-			{
-				TSet<FString> AllNodeName;
-				int32 CurrentNameIndex = 1;
-				for (int32 NodeIndex = 0; NodeIndex < SDKScene->GetNodeCount(); ++NodeIndex)
-				{
-					FbxNode* Node = SDKScene->GetNode(NodeIndex);
-					FString NodeName = UTF8_TO_TCHAR(Node->GetName());
-					if (NodeName.IsEmpty())
-					{
-						do
-						{
-							NodeName = TEXT("ncl1_") + FString::FromInt(CurrentNameIndex++);
-						} while (AllNodeName.Contains(NodeName));
-
-						Node->SetName(TCHAR_TO_UTF8(*NodeName));
 						if (!GIsAutomationTesting)
 						{
 							UInterchangeResultDisplay_Generic* Message = AddMessage<UInterchangeResultDisplay_Generic>();
-							Message->Text = FText::Format(LOCTEXT("EnsureNodeNameAreValid_NoNodeName", "Interchange FBX file Loading: Found node with no name, new node name is '{0}'"), FText::FromString(NodeName));
+							Message->Text = FText::Format(LOCTEXT("CleanupFbxData_NoObjectName", "Interchange FBX file Loading: Found object with no name, new object name is '{0}'"), FText::FromString(ObjectName));
 						}
-					}
-					ManageNamespace(bKeepFbxNamespace, NodeName, Node);
-					
-					// Do not allow node to be named same as filename as this creates problems later on (reimport)
-					if (AllNodeName.Contains(NodeName))
-					{
-						FString UniqueNodeName;
-						do
-						{
-							UniqueNodeName = NodeName + FString::FromInt(CurrentNameIndex++);
-						} while (AllNodeName.Contains(UniqueNodeName));
+					};
 
-						FbxString UniqueName(TCHAR_TO_UTF8(*UniqueNodeName));
-						Node->SetName(UniqueName);
+				//bUseNodeRules: 
+				// - we don't manage the namespace as each namespace is expected to be in separate assets
+				// - name sanitization will be part of the uniqueness check, but we don't set the sanitized name to the object, 
+				//		as node names will have their namespace managed and sanitized when we create the names for Interchange nodes specifically.
+				auto MakeFbxObjectNameUnique = [this](FbxObject* Object, TMap<FString, int32>& Names, bool bUseNCL = true, bool bUseNodeRules = false)
+					{
+						FString NodeName = UTF8_TO_TCHAR(Object->GetName());
+						FString UniqueNodeName = NodeName;
 
-						if (!GIsAutomationTesting)
+						if (!bUseNodeRules)
 						{
-							UInterchangeResultDisplay_Generic* Message = AddMessage<UInterchangeResultDisplay_Generic>();
-							Message->Text = FText::Format(LOCTEXT("EnsureNodeNameAreValid_NodeNameClash", "FBX File Loading: Found name clash, node '{0}' was renamed to '{1}'"), FText::FromString(NodeName), FText::FromString(UniqueNodeName));
+							//Manage namespaces should be done before sanitizing since it relies on the ":" character
+							GetFbxHelper()->ManageNamespaceAndRenameObject(UniqueNodeName, Object);
 						}
-					}
-					AllNodeName.Add(NodeName);
-				}
-			}
-
-			void FFbxParser::CleanupFbxData()
-			{
-				auto MakeFbxObjectNameUnique = [bKeepFbxNamespaceClosure = bKeepFbxNamespace](FbxObject* Object, TMap<FString, int32>& Names)
-					{
-						FString ObjectName = UTF8_TO_TCHAR(Object->GetName());
-						//Manage namespaces should be done before sanitizing since it relies on the ":" character
-						ManageNamespace(bKeepFbxNamespaceClosure, ObjectName, Object);
 
 						//Make sure to use the sanitized name to avoid name collisions later on
-						UE::Interchange::SanitizeName(ObjectName);
+						UE::Interchange::SanitizeName(UniqueNodeName);
 
-						if (int32* Count = Names.Find(ObjectName))
+						if (int32* Count = Names.Find(UniqueNodeName))
 						{
 							(*Count)++;
-							ObjectName += TEXT("_ncl_") + FString::FromInt(*Count);
-							Object->SetName(TCHAR_TO_UTF8(*ObjectName));
+							if (bUseNodeRules)
+							{
+								//if we are using node rules then its expected to fix the name facing towards Interchange, instead of right in place (in FBX struct).
+								UniqueNodeName = NodeName + (bUseNCL ? TEXT("_ncl_") : TEXT("")) + FString::FromInt(*Count);
+								Object->SetName(TCHAR_TO_UTF8(*UniqueNodeName));
+							}
+							else
+							{
+								UniqueNodeName += (bUseNCL ? TEXT("_ncl_") : TEXT("")) + FString::FromInt(*Count);
+								Object->SetName(TCHAR_TO_UTF8(*UniqueNodeName));
+							}
+
+							if (!GIsAutomationTesting)
+							{
+								UInterchangeResultDisplay_Generic* Message = AddMessage<UInterchangeResultDisplay_Generic>();
+								Message->Text = FText::Format(LOCTEXT("CleanupFbxData_NodeNameClash", "FBX File Loading: Found name clash, object '{0}' was renamed to '{1}'"), FText::FromString(NodeName), FText::FromString(UniqueNodeName));
+							}
 						}
 						else
 						{
-							Names.Add(ObjectName, 0);
+							Names.Add(UniqueNodeName, 0);
 						}
 					};
 
@@ -410,11 +381,13 @@ namespace UE
 				{
 					FbxNode* Node = SDKScene->GetNode(NodeIndex);
 					FString NodeName = UTF8_TO_TCHAR(Node->GetName());
+					bool bUseNCL = false;
 					if (NodeName.IsEmpty())
 					{
-						Node->SetName(TCHAR_TO_UTF8(TEXT("Node")));
+						bUseNCL = true;
+						ReportAndSetNoName(Node, TEXT("Node"));
 					}
-					MakeFbxObjectNameUnique(Node, NodeNames);
+					MakeFbxObjectNameUnique(Node, NodeNames, bUseNCL, true);
 				}
 
 				//////////////////////////////////////////////////////////////////////////
@@ -436,7 +409,7 @@ namespace UE
 					FString MeshName = UTF8_TO_TCHAR(Mesh->GetName());
 					if (MeshName.IsEmpty())
 					{
-						Mesh->SetName(TCHAR_TO_UTF8(TEXT("Mesh")));
+						ReportAndSetNoName(Mesh, TEXT("Mesh"));
 					}
 					MakeFbxObjectNameUnique(Mesh, MeshNames);
 				}
@@ -451,7 +424,7 @@ namespace UE
 					FString MaterialName = UTF8_TO_TCHAR(Material->GetName());
 					if (MaterialName.IsEmpty())
 					{
-						Material->SetName(TCHAR_TO_UTF8(TEXT("Material")));
+						ReportAndSetNoName(Material, TEXT("Material"));
 					}
 					MakeFbxObjectNameUnique(Material, MaterialNames);
 				}
diff --git a/Engine/Plugins/Interchange/Runtime/Source/Parsers/Fbx/Private/FbxHelper.cpp b/Engine/Plugins/Interchange/Runtime/Source/Parsers/Fbx/Private/FbxHelper.cpp
--- a/Engine/Plugins/Interchange/Runtime/Source/Parsers/Fbx/Private/FbxHelper.cpp
+++ b/Engine/Plugins/Interchange/Runtime/Source/Parsers/Fbx/Private/FbxHelper.cpp
@@ -119,6 +119,9 @@ namespace UE
 				}
 
 				FString ObjName = FFbxConvert::MakeString(Object->GetName());
+
+				ManageNamespace(ObjName);
+
 				UE::Interchange::SanitizeName(ObjName, bIsJoint);
 
 				return ObjName;
@@ -134,7 +137,13 @@ namespace UE
 				const FbxNode* ParentNode = Node;
 				while (ParentNode)
 				{
-					UniqueIDTokens.Add(GetFbxObjectName(ParentNode));
+					//As we don't want the name to have the namespace removed,
+					//we manually generate the name here:
+
+					FString ParentNodeName = FFbxConvert::MakeString(ParentNode->GetName());
+					UE::Interchange::SanitizeName(ParentNodeName);
+
+					UniqueIDTokens.Add(ParentNodeName);
 					ParentNode = ParentNode->GetParent();
 				}
 				FString UniqueID;
@@ -149,6 +158,50 @@ namespace UE
 				return UniqueID;
 			}
 
+			void FFbxHelper::ManageNamespace(FString& ObjectName) const
+			{
+				if (bKeepFbxNamespace)
+				{
+					if (ObjectName.Contains(TEXT(":")))
+					{
+						ObjectName.ReplaceInline(TEXT(":"), TEXT("_"));
+					}
+				}
+				else
+				{
+					// Remove namespaces
+					int32 LastNamespaceTokenIndex = INDEX_NONE;
+					if (ObjectName.FindLastChar(TEXT(':'), LastNamespaceTokenIndex))
+					{
+						//+1 to remove the ':' character we found
+						ObjectName.RightChopInline(LastNamespaceTokenIndex + 1, EAllowShrinking::Yes);
+					}
+				}
+			}
+
+			void FFbxHelper::ManageNamespaceAndRenameObject(FString& ObjectName, FbxObject* Object) const
+			{
+				if (bKeepFbxNamespace)
+				{
+					if (ObjectName.Contains(TEXT(":")))
+					{
+						ObjectName.ReplaceInline(TEXT(":"), TEXT("_"));
+						Object->SetName(TCHAR_TO_UTF8(*ObjectName));
+					}
+				}
+				else
+				{
+					// Remove namespaces
+					int32 LastNamespaceTokenIndex = INDEX_NONE;
+					if (ObjectName.FindLastChar(TEXT(':'), LastNamespaceTokenIndex))
+					{
+						//+1 to remove the ':' character we found
+						ObjectName.RightChopInline(LastNamespaceTokenIndex + 1, EAllowShrinking::Yes);
+						Object->SetName(TCHAR_TO_UTF8(*ObjectName));
+					}
+				}
+			}
+
 			FString FFbxHelper::GetUniqueIDString(const uint64 UniqueID) const
 			{
 				FStringFormatNamedArguments FormatArguments;
diff --git a/Engine/Plugins/Interchange/Runtime/Source/Parsers/Fbx/Private/FbxHelper.h b/Engine/Plugins/Interchange/Runtime/Source/Parsers/Fbx/Private/FbxHelper.h
--- a/Engine/Plugins/Interchange/Runtime/Source/Parsers/Fbx/Private/FbxHelper.h
+++ b/Engine/Plugins/Interchange/Runtime/Source/Parsers/Fbx/Private/FbxHelper.h
@@ -36,6 +36,11 @@ namespace UE::Interchange
 		struct FFbxHelper
 		{
 		public:
+			explicit FFbxHelper(bool bInKeepFbxNamespace)
+				: bKeepFbxNamespace(bInKeepFbxNamespace)
+			{
+			}
+
 			void Reset() { MaterialNameClashMap.Reset(); }
 			FString GetMeshName(FbxGeometryBase* Mesh) const;
 			FString GetMeshUniqueID(FbxGeometryBase* Mesh) const;
@@ -60,11 +65,14 @@ namespace UE::Interchange
 				*/
 			FString GetFbxNodeHierarchyName(const FbxNode* Node) const;
 
+			void ManageNamespace(FString& ObjectName) const;
+			void ManageNamespaceAndRenameObject(FString& ObjectName, FbxObject* Object) const;
 		protected:
 			FString GetUniqueIDString(const uint64 UniqueID) const;
 
 		private:
 			mutable TMap<FString, const FbxObject*> MaterialNameClashMap;
+			bool bKeepFbxNamespace;
 		};
 
 		void ProcessCustomAttribute(FFbxParser& Parser, UInterchangeBaseNode* UnrealNode, FbxProperty Property, const TOptional<FString>& PayloadKey = TOptional<FString>());
