Fixed issue with level reimport's strategy applied to mesh factory nodes' reimport strategy
- Root cause: Level pipeline was applying level's reimport strategy to all factory nodes on which a new actor factory node was depending on.
- Solution: Apply level's reimport strategy only on new actor factory nodes it creates.

#rnx

diff --git a/Engine/Plugins/Interchange/Runtime/Source/Pipelines/Private/InterchangeGenericScenesPipeline.cpp b/Engine/Plugins/Interchange/Runtime/Source/Pipelines/Private/InterchangeGenericScenesPipeline.cpp
--- a/Engine/Plugins/Interchange/Runtime/Source/Pipelines/Private/InterchangeGenericScenesPipeline.cpp
+++ b/Engine/Plugins/Interchange/Runtime/Source/Pipelines/Private/InterchangeGenericScenesPipeline.cpp
@@ -87,24 +87,31 @@ namespace UE::Interchange::Private
 			}
 			else
 			{
+
 				bool bParentIsJoint = ParentSceneNode->IsSpecializedTypeContains(UE::Interchange::FSceneNodeStaticData::GetJointSpecializeTypeString());
 				if (bParentIsJoint)
 				{
-					//Check if its a rootjoint:
-					//	aka check parent's parent if its !joint:
-					FString ParentsParentUid = ParentSceneNode->GetParentUid();
-					if (const UInterchangeSceneNode* ParentsParentSceneNode = Cast<UInterchangeSceneNode>(BaseNodeContainer->GetNode(ParentsParentUid)))
+					FString GreatParentUid = ParentSceneNode->GetParentUid();
+					if (const UInterchangeSceneNode* GreatParentSceneNode = Cast<UInterchangeSceneNode>(BaseNodeContainer->GetNode(GreatParentUid)))
 					{
-						bool bParentsParentIsJoint = ParentSceneNode->IsSpecializedTypeContains(UE::Interchange::FSceneNodeStaticData::GetJointSpecializeTypeString());
-						if (bParentsParentIsJoint)
-						{
-							return FindFactoryParentSceneNodeUid(BaseNodeContainer, ActiveSkeletonUids, ParentSceneNode);
-						}
-						else
-						{
-							return ParentUid;
-						}
+						return FindFactoryParentSceneNodeUid(BaseNodeContainer, ActiveSkeletonUids, ParentSceneNode);
 					}
+
+					// #interchange_skeletal_revisit: Code above reflects actual behavior of code below 
+					//FString GreatParentUid = ParentSceneNode->GetParentUid();
+					//if (const UInterchangeSceneNode* GreatParentSceneNode = Cast<UInterchangeSceneNode>(BaseNodeContainer->GetNode(GreatParentUid)))
+					//{
+					//  Should we call ParentSceneNode or GreatParentSceneNode below?
+					//	bool bGreatParentIsJoint = (ParentSceneNode or GreatParentSceneNode?)->IsSpecializedTypeContains(UE::Interchange::FSceneNodeStaticData::GetJointSpecializeTypeString());
+					//	if (bGreatParentIsJoint)
+					//	{
+					//		return FindFactoryParentSceneNodeUid(BaseNodeContainer, ActiveSkeletonUids, ParentSceneNode);
+					//	}
+					//	else
+					//	{
+					//		return ParentUid;
+					//	}
+					//}
 				}
 				else
 				{
@@ -116,21 +123,6 @@ namespace UE::Interchange::Private
 		return UInterchangeBaseNode::InvalidNodeUid();
 	}
 
-	void UpdateReimportStrategyFlags(UInterchangeBaseNodeContainer& NodeContainer, UInterchangeFactoryBaseNode& FactoryNode, EReimportStrategyFlags ReimportPropertyStrategy)
-	{
-		FactoryNode.SetReimportStrategyFlags(ReimportPropertyStrategy);
-
-		TArray<FString> ActorDependencies;
-		FactoryNode.GetFactoryDependencies(ActorDependencies);
-		for (const FString& FactoryNodeID : ActorDependencies)
-		{
-			if (UInterchangeFactoryBaseNode* DependencyFactoryNode = NodeContainer.GetFactoryNode(FactoryNodeID))
-			{
-				UpdateReimportStrategyFlags(NodeContainer, *DependencyFactoryNode, ReimportPropertyStrategy);
-			}
-		}
-	}
-
 #if WITH_EDITOR
 	void DeleteAssets(const TArray<UObject*>& AssetsToDelete)
 	{
@@ -740,7 +732,7 @@ void UInterchangeGenericLevelPipeline::ExecuteSceneNodePreImport(const FTransfor
 	}
 	BaseNodeContainer->SetupNode(ActorFactoryNode, ActorFactoryNodeUid, SceneNode->GetDisplayLabel(), EInterchangeNodeContainerType::FactoryData, ParentFactoryNodeUid);
 #if WITH_EDITORONLY_DATA
-	//The level must be create before any actor asset since all actors will be create in the specified level
+	//The level must be create before any actor asset since all actors will be created in the specified level
 	if (LevelFactoryNode)
 	{
 		ActorFactoryNode->AddFactoryDependencyUid(LevelFactoryNode->GetUniqueID());
@@ -820,10 +812,6 @@ void UInterchangeGenericLevelPipeline::ExecuteSceneNodePreImport(const FTransfor
 	SceneNode->GetCustomComponentVisibility(bCustomVisibility);
 	ActorFactoryNode->SetCustomComponentVisibility(bCustomVisibility);
 	
-	bCustomVisibility = true;
-	SceneNode->GetCustomActorVisibility(bCustomVisibility);
-	ActorFactoryNode->SetCustomActorVisibility(bCustomVisibility);
-	
 	ActorFactoryNode->SetCustomMobility(EComponentMobility::Static);
 
 	if (TranslatedAssetNode)
@@ -833,12 +821,13 @@ void UInterchangeGenericLevelPipeline::ExecuteSceneNodePreImport(const FTransfor
 
 	SetUpComponents(ActorFactoryNode, SceneNode);
 	
-
-	//Make sure all actor factory nodes and dependencies have the specified strategy
-	Private::UpdateReimportStrategyFlags(*BaseNodeContainer, *ActorFactoryNode, ReimportPropertyStrategy);
-	if (bForceReimportDeletedActors)
+	if (CacheContextParam.ContextType == EInterchangePipelineContext::SceneReimport)
 	{
-		ActorFactoryNode->SetForceNodeReimport();
+		ActorFactoryNode->SetReimportStrategyFlags(ReimportPropertyStrategy);
+		if (bForceReimportDeletedActors)
+		{
+			ActorFactoryNode->SetForceNodeReimport();
+		}
 	}
 
 #if WITH_EDITORONLY_DATA
