[USD Interchange] Nanite Assembly reimport fixes

Factory node association fallback:
During reimport, FInterchangeNaniteAssemblyBuilder now extracts existing parts from FNaniteAssemblyData and associates them back to their originating factory
nodes using NodeUniqueID metadata. This improves the success rate of locating part meshes when reference objects are missing.

Material slot name consistency:
FInterchangeNaniteAssemblyBuilder now copies material slot names from part meshes instead of allowing the underlying FNaniteAssemblyDataBuilder to
generate a slot name. This ensures Interchange is able to correctly resolve material and slot references during reimport.

Improved PointInstancer support:
Added support in the USD translator for applying the NaniteAssemblyRootAPI schema directly to PointInstancer prims (static mesh assembly cases only, which
do not require a base mesh).

diff --git a/Engine/Plugins/Interchange/Extensions/OpenUSD/Source/Import/Private/InterchangeUsdTranslator.cpp b/Engine/Plugins/Interchange/Extensions/OpenUSD/Source/Import/Private/InterchangeUsdTranslator.cpp
--- a/Engine/Plugins/Interchange/Extensions/OpenUSD/Source/Import/Private/InterchangeUsdTranslator.cpp
+++ b/Engine/Plugins/Interchange/Extensions/OpenUSD/Source/Import/Private/InterchangeUsdTranslator.cpp
@@ -2237,10 +2237,18 @@ namespace UE::InterchangeUsdTranslator::Private
 
 			TArray<FTransform> TransformStack = { ParentGlobalTransform };
 			TArray<FPrimPrototypeEntry> PrototypeStack;
-			constexpr bool bTraverseInstanceProxies = true;
-			for (const UE::FUsdPrim& ChildPrim : AssemblyRootPrim.GetFilteredChildren(bTraverseInstanceProxies))
+
+			if (MeshTypeIs(ENaniteAssemblyMeshType::StaticMesh) && AssemblyRootPrim.IsA(TEXT("PointInstancer")))
+			{
+				TraversePointInstancerHierarchy(AssemblyRootPrim, Info, TransformStack, PrototypeStack);
+			}
+			else
 			{
-				TraverseImpl(ChildPrim, Info, TransformStack, PrototypeStack);
+				constexpr bool bTraverseInstanceProxies = true;
+				for (const UE::FUsdPrim& ChildPrim : AssemblyRootPrim.GetFilteredChildren(bTraverseInstanceProxies))
+				{
+					TraverseImpl(ChildPrim, Info, TransformStack, PrototypeStack);
+				}
 			}
 
 			if (!ValidateTraversalResult())
@@ -2280,7 +2288,7 @@ namespace UE::InterchangeUsdTranslator::Private
 				TransformStack.Add(FTransform::Identity);
 				PrototypeStack.Emplace(Prim.GetPrimPath(), PrototypePrim.GetPrimPath(), PrototypeIndex, PrototypePaths.Num());
 
-				TraverseImpl(PrototypePrim, Info, TransformStack, PrototypeStack);
+				TraverseImpl(PrototypePrim, PrototypePrimInfo, TransformStack, PrototypeStack);
 
 				// Pop before entering the next
 				PrototypeStack.Pop();
@@ -2375,7 +2383,7 @@ namespace UE::InterchangeUsdTranslator::Private
 
 			while (!CurrentPrim.IsPseudoRoot())
 			{
-				if (CurrentPrim.IsA(TEXT("PointInstancer")))
+				if (CurrentPrim != AssemblyRootPrim && CurrentPrim.IsA(TEXT("PointInstancer")))
 				{
 					UE_LOG(LogUsd, Warning,
 						TEXT("NaniteAssemblyRoot '%s' cannot currently be a descendant of a pointinstancer - '%s'")
diff --git a/Engine/Plugins/Interchange/Runtime/Source/Import/Private/Mesh/InterchangeMeshHelper.cpp b/Engine/Plugins/Interchange/Runtime/Source/Import/Private/Mesh/InterchangeMeshHelper.cpp
--- a/Engine/Plugins/Interchange/Runtime/Source/Import/Private/Mesh/InterchangeMeshHelper.cpp
+++ b/Engine/Plugins/Interchange/Runtime/Source/Import/Private/Mesh/InterchangeMeshHelper.cpp
@@ -32,6 +32,7 @@
 #include "StaticMeshAttributes.h"
 #include "StaticMeshOperations.h"
 
+#include "InterchangeAssetImportData.h"
 #include "InterchangeCommonPipelineDataFactoryNode.h"
 #include "InterchangeHelper.h"
 #include "InterchangeImportLog.h"
@@ -44,6 +45,12 @@
 
 #include <type_traits>
 
+#if WITH_EDITORONLY_DATA
+
+#include "EditorFramework/AssetImportData.h"
+
+#endif //WITH_EDITORONLY_DATA
+
 namespace UE::Interchange::Private::MeshHelper
 {
 	bool AddConvexGeomFromVertices(
@@ -1681,6 +1688,10 @@ namespace UE::Interchange::Private::MeshHelper
 		// The base node container
 		const UInterchangeBaseNodeContainer* NodeContainer = nullptr;
 
+		// Maps to help reassociate existing part meshes and slot bindings (intended for use during Interchange reimport)
+		TMap<FString, TWeakObjectPtr<MeshType>> CurrentNodeUidToPartMesh;
+		TMap<TWeakObjectPtr<MeshType>, TArray<int32>> PartMeshToMaterialRemaps;
+
 		bool bCtorIsValid = false;
 
 		FInternalNaniteAssemblyBuilder(MeshType* InMesh
@@ -1689,19 +1700,23 @@ namespace UE::Interchange::Private::MeshHelper
 		  : Mesh(InMesh)
 		  , MeshFactoryNode(InMeshFactoryNode)
 		  , NodeContainer(InNodeContainer)
-		  , bCtorIsValid(Mesh && MeshFactoryNode && NodeContainer)
+		  , bCtorIsValid(IsInGameThread() && Mesh && MeshFactoryNode && NodeContainer)
 		{
 			if (!ensure(bCtorIsValid))
 			{
 				return;
 			}
 
-			// Add the initial base mesh materials to the builder
+			Mesh->ClearCachedNaniteAssemblyReferences();
+
+			ExtractPreexistingAssemblyPartMeshes();
+
+			// Add the initial base mesh materials to the builder.
 			ForEachMeshMaterial(
 				Mesh,
-				[this](int32 LocalMaterialIndex, const TObjectPtr<UMaterialInterface>& Material)
+				[this](int32 LocalMaterialIndex, const FName& SlotName, const TObjectPtr<UMaterialInterface>& Material)
 				{
-					Builder.AddMaterialSlot(Material);
+					Builder.AddMaterialSlot({ Material, SlotName });
 				}
 			);
 		}
@@ -1868,7 +1883,7 @@ namespace UE::Interchange::Private::MeshHelper
 		}
 
 		/** Add a node to the FNaniteAssemblyDataBuilder. Also add a new part if this is the first time the part
-		  * the first time the part mesh has been encountered.
+		  * mesh has been encountered.
 		  * 
 		  * Note that the bool return value is not success/failure, but whether or not a new part was created.
 		  */
@@ -1901,10 +1916,24 @@ namespace UE::Interchange::Private::MeshHelper
 			{
 				ForEachMeshMaterial(
 					PartMesh,
-					[this, PartIndex](int32 LocalMaterialIndex, const TObjectPtr<UMaterialInterface>& PartMaterial)
+					[this, PartIndex, PartMesh](int32 LocalMaterialIndex, const FName& SlotName, const TObjectPtr<UMaterialInterface>& PartMaterial)
 					{
 						int32 MaterialIndex = INDEX_NONE;
 						
+						// First attempt to find an existing material slot that was already present on the base mesh.
+						if (const TArray<int32>* RemappedIndices = PartMeshToMaterialRemaps.Find(PartMesh))
+						{
+							if (RemappedIndices->IsValidIndex(LocalMaterialIndex))
+							{
+								const int32 RemappedIndex = (*RemappedIndices)[LocalMaterialIndex];
+								if (Builder.GetMaterialSlots().IsValidIndex(RemappedIndex))
+								{
+									Builder.RemapPartMaterial(PartIndex, LocalMaterialIndex, RemappedIndex);
+									return;
+								}
+							}
+						}
+
 						// For now we always merge identical materials. There are however more advanced workflows 
 						// that require exact control over material and slot management apparently, and so more
 						// flexibility may be required in the future.
@@ -1921,7 +1950,7 @@ namespace UE::Interchange::Private::MeshHelper
 
 						if (MaterialIndex == INDEX_NONE)
 						{
-							MaterialIndex = Builder.AddMaterialSlot(PartMaterial);
+							MaterialIndex = Builder.AddMaterialSlot({ PartMaterial, SlotName });
 						}
 
 						Builder.RemapPartMaterial(PartIndex, LocalMaterialIndex, MaterialIndex);
@@ -1969,6 +1998,7 @@ namespace UE::Interchange::Private::MeshHelper
 						if (const UInterchangeMeshFactoryNode* PartMeshFactoryNode = Cast<const UInterchangeMeshFactoryNode>(NodeContainer->GetFactoryNode(*FactoryNodeDependencyUid)))
 						{
 							FSoftObjectPath PartReferenceObject;
+							// Attempt to get the dependent object from the factory node. NOTE this will return false during a reimport
 							if (PartMeshFactoryNode->GetCustomReferenceObject(PartReferenceObject))
 							{
 								if (MeshType* MeshObject = Cast<MeshType>(PartReferenceObject.TryLoad()))
@@ -1977,6 +2007,15 @@ namespace UE::Interchange::Private::MeshHelper
 									bFoundPartObject = true;
 								}
 							}
+							// Otherwise try and find a matching part that was already attached to the base mesh we received
+							else if (TWeakObjectPtr<MeshType>* MeshObject = CurrentNodeUidToPartMesh.Find(*FactoryNodeDependencyUid))
+							{
+								if (MeshObject->IsValid())
+								{
+									OutPartObjects[Index].Add(MeshObject->Get());
+									bFoundPartObject = true;
+								}
+							}
 						}
 					}
 					else if (FPackageName::IsValidLongPackageName(PartUid) || FPackageName::IsValidObjectPath(PartUid))
@@ -2012,7 +2051,7 @@ namespace UE::Interchange::Private::MeshHelper
 				);
 				for (const FString& PartUid : MissingPartUids)
 				{
-					Message.Append(FString::Printf(TEXT("... %s"), *PartUid));
+					Message.Append(FString::Printf(TEXT("... %s\n"), *PartUid));
 				}
 				UE_LOG(LogInterchangeImport, Warning, TEXT("%s"), *Message);
 			}
@@ -2020,6 +2059,47 @@ namespace UE::Interchange::Private::MeshHelper
 			return bFoundAtLeastOnePartObject;
 		}
 
+		/** Extract the current nanite assembly parts (if any) from the base mesh. This enables us to map the
+		  * incoming factory node UIDs to usable parts during reimport, which is the best we can do until
+		  * Interchange provides valid dependencies on reimports.
+		  */
+		void ExtractPreexistingAssemblyPartMeshes()
+		{
+#if WITH_EDITORONLY_DATA
+			FMeshNaniteSettings Settings = Mesh->GetNaniteSettings();
+			if (!Settings.bEnabled)
+			{
+				return;
+			}
+
+			const FNaniteAssemblyData& AssemblyData = Settings.NaniteAssemblyData;
+			for (const FNaniteAssemblyPart& Part : AssemblyData.Parts)
+			{
+				const FSoftObjectPath& PartObjectPath = Part.MeshObjectPath;
+				if (!PartObjectPath.IsValid() || !PartObjectPath.IsAsset())
+				{
+					continue;
+				}
+
+				if (MeshType* PartObject = Cast<MeshType>(PartObjectPath.TryLoad()))
+				{
+					if (UAssetImportData* ImportDataPtr = PartObject->GetAssetImportData())
+					{
+						if (UInterchangeAssetImportData* AssetImportData = Cast<UInterchangeAssetImportData>(ImportDataPtr))
+						{
+							const FString& NodeUid = AssetImportData->NodeUniqueID;
+							if (!NodeUid.IsEmpty())
+							{
+								CurrentNodeUidToPartMesh.Add(NodeUid, PartObject);
+								PartMeshToMaterialRemaps.Add(PartObject, Part.MaterialRemap);
+							}
+						}
+					}
+				}
+			}
+#endif //WITH_EDITORONLY_DATA
+		}
+
 		/** Set the nanite enabled flag to the given value. Note the bool return value here is the previous enabled 
 		  * on/off state (since we may need to revert it), not whether the new setting succeeded/failed. 
 		  */
@@ -2050,14 +2130,14 @@ namespace UE::Interchange::Private::MeshHelper
 			{
 				for (const FSkeletalMaterial& SkeletalMaterial : MeshTypeObject->GetMaterials())
 				{
-					Func(LocalMaterialIndex++, SkeletalMaterial.MaterialInterface);
+					Func(LocalMaterialIndex++, SkeletalMaterial.MaterialSlotName, SkeletalMaterial.MaterialInterface);
 				}
 			}
 			else if constexpr (IsStaticMesh())
 			{
 				for (const FStaticMaterial& StaticMaterial : MeshTypeObject->GetStaticMaterials())
 				{
-					Func(LocalMaterialIndex++, StaticMaterial.MaterialInterface);
+					Func(LocalMaterialIndex++, StaticMaterial.MaterialSlotName, StaticMaterial.MaterialInterface);
 				}
 			}
 		}
