Fixed issues with glTF level and asset reimport of material instances
- Bug's issue: Updated the glTF asset and scene asset pipelines to do not search for materials.
- Hide all the options related to material import since only material instances are created
- Fixed logic issues on reimport and reimport strategies
- Reimport strategy flag was not carried on from material factory node to the material instance factory node replacing it.
- Reimport strategy of material instances was not applied in the material factory
Improved log when no data has been imported
- Reports success when no data and no error has been imported
- Added a accompagnying message indicating that no data has been imported
- Logging a warning message that no data has been imported

#rnx

diff --git a/Engine/Plugins/Interchange/Runtime/Source/Import/Private/Material/InterchangeMaterialFactory.cpp b/Engine/Plugins/Interchange/Runtime/Source/Import/Private/Material/InterchangeMaterialFactory.cpp
--- a/Engine/Plugins/Interchange/Runtime/Source/Import/Private/Material/InterchangeMaterialFactory.cpp
+++ b/Engine/Plugins/Interchange/Runtime/Source/Import/Private/Material/InterchangeMaterialFactory.cpp
@@ -456,13 +456,15 @@ namespace UE::Interchange::MaterialFactory::Internal
 		{
 			bool bInputValue = false;
 			const FString InputKey = bIsAParameter ? UInterchangeShaderPortsAPI::MakeInputParameterKey(InputName) : UInterchangeShaderPortsAPI::MakeInputValueKey(InputName);
-			FactoryNode.GetBooleanAttribute(InputKey, bInputValue);
 
-			if (bInputValue != bInstanceValue)
+			if (FactoryNode.GetBooleanAttribute(InputKey, bInputValue))
 			{
-				if (UMaterialInstanceConstant* MaterialInstanceConstant = Cast<UMaterialInstanceConstant>(&MaterialInstance))
+				if (bInputValue != bInstanceValue)
 				{
-					MaterialInstanceConstant->SetStaticSwitchParameterValueEditorOnly(ParameterName, bInputValue);
+					if (UMaterialInstanceConstant* MaterialInstanceConstant = Cast<UMaterialInstanceConstant>(&MaterialInstance))
+					{
+						MaterialInstanceConstant->SetStaticSwitchParameterValueEditorOnly(ParameterName, bInputValue);
+					}
 				}
 			}
 		}
@@ -478,20 +480,21 @@ namespace UE::Interchange::MaterialFactory::Internal
 		{
 			float InputValue = 0.f;
 			const FString InputKey = bIsAParameter ? UInterchangeShaderPortsAPI::MakeInputParameterKey(InputName) : UInterchangeShaderPortsAPI::MakeInputValueKey(InputName);
-			FactoryNode.GetFloatAttribute(InputKey, InputValue);
-
-			if (!FMath::IsNearlyEqual(InputValue, InstanceValue))
+			if (FactoryNode.GetFloatAttribute(InputKey, InputValue))
 			{
-#if WITH_EDITOR
-				if (UMaterialInstanceConstant* MaterialInstanceConstant = Cast<UMaterialInstanceConstant>(&MaterialInstance))
+				if (!FMath::IsNearlyEqual(InputValue, InstanceValue))
 				{
-					MaterialInstanceConstant->SetScalarParameterValueEditorOnly(ParameterName, InputValue);
-				}
-				else
+#if WITH_EDITOR
+					if (UMaterialInstanceConstant* MaterialInstanceConstant = Cast<UMaterialInstanceConstant>(&MaterialInstance))
+					{
+						MaterialInstanceConstant->SetScalarParameterValueEditorOnly(ParameterName, InputValue);
+					}
+					else
 #endif // #if WITH_EDITOR
-				if (UMaterialInstanceDynamic* MaterialInstanceDynamic = Cast<UMaterialInstanceDynamic>(&MaterialInstance))
-				{
-					MaterialInstanceDynamic->SetScalarParameterValue(ParameterName, InputValue);
+						if (UMaterialInstanceDynamic* MaterialInstanceDynamic = Cast<UMaterialInstanceDynamic>(&MaterialInstance))
+						{
+							MaterialInstanceDynamic->SetScalarParameterValue(ParameterName, InputValue);
+						}
 				}
 			}
 		}
@@ -807,6 +810,11 @@ namespace UE::Interchange::MaterialFactory::Internal
 			MaterialInstance.BasePropertyOverrides.BlendMode = BlendMode;
 			MaterialInstance.BlendMode = BlendMode;
 		}
+		// If the factory node does not have a blend mode, disable blend mode override
+		else
+		{
+			MaterialInstance.BasePropertyOverrides.bOverride_BlendMode = false;
+		}
 	}
 }
 
@@ -956,7 +964,7 @@ UInterchangeFactoryBase::FImportAssetResult UInterchangeMaterialFactory::BeginIm
 			Material = NewObject<UObject>(Arguments.Parent, MaterialClass, *Arguments.AssetName, RF_Public | RF_Standalone);
 		}
 	}
-	else if (ExistingAsset->GetClass()->IsChildOf(MaterialClass))
+	else if (ExistingAsset && ExistingAsset->GetClass()->IsChildOf(MaterialClass))
 	{
 		//This is a reimport, we are just re-updating the source data
 		Material = ExistingAsset;
@@ -2012,11 +2020,11 @@ void UInterchangeMaterialFactory::SetupReimportedMaterialInstance(UMaterialInsta
 {
 	using namespace UE::Interchange::MaterialFactory::Internal;
 
-	auto ValidateBool = [&](const FString& InputName) -> bool
+	const EReimportStrategyFlags ReimportStrategy = FactoryNode.GetReimportStrategyFlags();
+	if (ReimportStrategy == EReimportStrategyFlags::ApplyNoProperties)
 	{
-
-		return true;
-	};
+		return;
+	}
 
 	TArray<FString> Inputs;
 	UInterchangeShaderPortsAPI::GatherInputs(&FactoryNode, Inputs);
@@ -2041,7 +2049,7 @@ void UInterchangeMaterialFactory::SetupReimportedMaterialInstance(UMaterialInsta
 						bPreviousInputValue = bInstanceValue;
 					}
 
-					if (bInstanceValue == bPreviousInputValue)
+					if (ReimportStrategy == EReimportStrategyFlags::ApplyPipelineProperties || bInstanceValue == bPreviousInputValue)
 					{
 						UpdateParameterBool(MaterialInstance, InputName, FactoryNode, bIsAParameter);
 					}
@@ -2051,14 +2059,14 @@ void UInterchangeMaterialFactory::SetupReimportedMaterialInstance(UMaterialInsta
 #endif
 		case UE::Interchange::EAttributeTypes::Float:
 			{
-				bool bUpdateParameter = true;
+				bool bUpdateParameter = ReimportStrategy == EReimportStrategyFlags::ApplyPipelineProperties;
 				float PreviousValue;
 				if (PreviousFactoryNode.GetFloatAttribute(AttributKey, PreviousValue))
 				{
 					float CurrentValue;
 					if (MaterialInstance.GetScalarParameterValue(*InputName, CurrentValue))
 					{
-						bUpdateParameter = FMath::IsNearlyEqual(PreviousValue, CurrentValue);
+						bUpdateParameter |= FMath::IsNearlyEqual(PreviousValue, CurrentValue);
 					}
 				}
 
@@ -2070,14 +2078,14 @@ void UInterchangeMaterialFactory::SetupReimportedMaterialInstance(UMaterialInsta
 			break;
 		case UE::Interchange::EAttributeTypes::LinearColor:
 			{
-				bool bUpdateParameter = true;
+				bool bUpdateParameter = ReimportStrategy == EReimportStrategyFlags::ApplyPipelineProperties;
 				FLinearColor PreviousValue;
 				if (PreviousFactoryNode.GetLinearColorAttribute(AttributKey, PreviousValue))
 				{
 					FLinearColor CurrentValue;
 					if (MaterialInstance.GetVectorParameterValue(*InputName, CurrentValue))
 					{
-						bUpdateParameter = PreviousValue.Equals(CurrentValue);
+						bUpdateParameter |= PreviousValue.Equals(CurrentValue);
 					}
 				}
 
@@ -2089,7 +2097,7 @@ void UInterchangeMaterialFactory::SetupReimportedMaterialInstance(UMaterialInsta
 			break;
 		case UE::Interchange::EAttributeTypes::String:
 			{
-				bool bUpdateParameter = true;
+				bool bUpdateParameter = ReimportStrategy == EReimportStrategyFlags::ApplyPipelineProperties;
 				FString PreviousValue;
 				if (PreviousFactoryNode.GetStringAttribute(AttributKey, PreviousValue))
 				{
@@ -2097,7 +2105,7 @@ void UInterchangeMaterialFactory::SetupReimportedMaterialInstance(UMaterialInsta
 					if (MaterialInstance.GetTextureParameterValue(*InputName, CurrentValue))
 					{
 						FSoftObjectPath PreviousObjectPath(PreviousValue);
-						bUpdateParameter = PreviousObjectPath == FSoftObjectPath(CurrentValue);
+						bUpdateParameter |= PreviousObjectPath == FSoftObjectPath(CurrentValue);
 					}
 				}
 
@@ -2110,7 +2118,25 @@ void UInterchangeMaterialFactory::SetupReimportedMaterialInstance(UMaterialInsta
 		}
 	}
 
-	OverrideBlendMode(MaterialInstance, FactoryNode);
+	bool bUpdateParameter = ReimportStrategy == EReimportStrategyFlags::ApplyPipelineProperties;
+	// If previous factory node has a blend mode, check whether it is the same as the material instance's one
+	if (TEnumAsByte<EBlendMode> BlendMode; PreviousFactoryNode.GetCustomBlendMode(BlendMode))
+	{
+		if (MaterialInstance.BasePropertyOverrides.bOverride_BlendMode)
+		{
+			bUpdateParameter |= MaterialInstance.BasePropertyOverrides.BlendMode == BlendMode;
+		}
+	}
+	// If there was no BlendMode override in previous import and things have not changed, update is allowed
+	else if(!MaterialInstance.BasePropertyOverrides.bOverride_BlendMode)
+	{
+		bUpdateParameter |= true;
+	}
+
+	if(bUpdateParameter)
+	{
+		OverrideBlendMode(MaterialInstance, FactoryNode);
+	}
 }
 
 UClass* UInterchangeMaterialFunctionFactory::GetFactoryClass() const
diff --git a/Engine/Plugins/Interchange/Runtime/Source/Pipelines/Private/InterchangeglTFPipeline.cpp b/Engine/Plugins/Interchange/Runtime/Source/Pipelines/Private/InterchangeglTFPipeline.cpp
--- a/Engine/Plugins/Interchange/Runtime/Source/Pipelines/Private/InterchangeglTFPipeline.cpp
+++ b/Engine/Plugins/Interchange/Runtime/Source/Pipelines/Private/InterchangeglTFPipeline.cpp
@@ -186,6 +186,11 @@ UInterchangeMaterialInstanceFactoryNode* UGLTFPipelineSettings::BuildMaterialIns
 {
 	using namespace UE::Interchange::GLTFMaterials;
 
+	if (!ShaderGraphNode)
+	{
+		return nullptr;
+	}
+
 	FString MaterialFactoryNodeUid = UInterchangeFactoryBaseNode::BuildFactoryNodeUid(ShaderGraphNode->GetUniqueID());
 
 	// Check the search location of the material, if under root, then we shouldn't replace the material factory node and just take the existing one
@@ -236,69 +241,81 @@ UInterchangeMaterialInstanceFactoryNode* UGLTFPipelineSettings::BuildMaterialIns
 
 	FString MaterialFactoryNodeName = ShaderGraphNode->GetDisplayLabel();
 
-	UInterchangeMaterialInstanceFactoryNode* MaterialInstanceFactoryNode = NewObject<UInterchangeMaterialInstanceFactoryNode>(NodeContainer);
-	NodeContainer->SetupAndReplaceFactoryNode(MaterialInstanceFactoryNode, MaterialFactoryNodeUid, MaterialFactoryNodeName, EInterchangeNodeContainerType::FactoryData, OldFactoryNodeUId);
-
-	MaterialInstanceFactoryNode->SetCustomParent(Parent);
-
-	UInterchangeEditorUtilitiesBase* EditorUtilities = UInterchangeManager::GetInterchangeManager().GetEditorUtilities();
-	const UClass* MaterialClass = (EditorUtilities && EditorUtilities->IsRuntimeOrPIE()) ? UMaterialInstanceDynamic::StaticClass() : UMaterialInstanceConstant::StaticClass();
-	MaterialInstanceFactoryNode->SetCustomInstanceClassName(MaterialClass->GetPathName());
+	// Cache reimport strategy flag from material factory node before replacing it
+	EReimportStrategyFlags ReimportStrategy = EReimportStrategyFlags::ApplyNoProperties;
+	if (UInterchangeBaseMaterialFactoryNode* MaterialFactoryNode = Cast<UInterchangeBaseMaterialFactoryNode>(NodeContainer->GetFactoryNode(OldFactoryNodeUId)))
+	{
+		ReimportStrategy = MaterialFactoryNode->GetReimportStrategyFlags();
+	}
 
-	for (const TPair<FString, UE::Interchange::FAttributeKey>& GltfAttributeKey : GltfAttributeKeys)
+	if (UInterchangeMaterialInstanceFactoryNode* MaterialInstanceFactoryNode = NewObject<UInterchangeMaterialInstanceFactoryNode>(NodeContainer))
 	{
-		UE::Interchange::EAttributeTypes AttributeType = ShaderGraphNode->GetAttributeType(GltfAttributeKey.Value);
+		NodeContainer->SetupAndReplaceFactoryNode(MaterialInstanceFactoryNode, MaterialFactoryNodeUid, MaterialFactoryNodeName, EInterchangeNodeContainerType::FactoryData, OldFactoryNodeUId);
 
-		FString InputValueKey = UInterchangeShaderPortsAPI::MakeInputValueKey(GltfAttributeKey.Key);
+		MaterialInstanceFactoryNode->SetReimportStrategyFlags(ReimportStrategy);
+		MaterialInstanceFactoryNode->SetCustomParent(Parent);
 
-		//we are only using 4 attribute types for now:
-		switch (AttributeType)
-		{
-		case UE::Interchange::EAttributeTypes::Bool:
+		UInterchangeEditorUtilitiesBase* EditorUtilities = UInterchangeManager::GetInterchangeManager().GetEditorUtilities();
+		const UClass* MaterialClass = (EditorUtilities && EditorUtilities->IsRuntimeOrPIE()) ? UMaterialInstanceDynamic::StaticClass() : UMaterialInstanceConstant::StaticClass();
+		MaterialInstanceFactoryNode->SetCustomInstanceClassName(MaterialClass->GetPathName());
+
+		for (const TPair<FString, UE::Interchange::FAttributeKey>& GltfAttributeKey : GltfAttributeKeys)
 		{
-			bool Value;
-			if (ShaderGraphNode->GetBooleanAttribute(GltfAttributeKey.Value.Key, Value))
+			UE::Interchange::EAttributeTypes AttributeType = ShaderGraphNode->GetAttributeType(GltfAttributeKey.Value);
+
+			FString InputValueKey = UInterchangeShaderPortsAPI::MakeInputValueKey(GltfAttributeKey.Key);
+
+			//we are only using 4 attribute types for now:
+			switch (AttributeType)
 			{
-				MaterialInstanceFactoryNode->AddBooleanAttribute(InputValueKey, Value);
+			case UE::Interchange::EAttributeTypes::Bool:
+			{
+				bool Value;
+				if (ShaderGraphNode->GetBooleanAttribute(GltfAttributeKey.Value.Key, Value))
+				{
+					MaterialInstanceFactoryNode->AddBooleanAttribute(InputValueKey, Value);
+				}
 			}
-		}
-		break;
-		case UE::Interchange::EAttributeTypes::Float:
-		{
-			float Value;
-			if (ShaderGraphNode->GetFloatAttribute(GltfAttributeKey.Value.Key, Value))
+			break;
+			case UE::Interchange::EAttributeTypes::Float:
 			{
-				MaterialInstanceFactoryNode->AddFloatAttribute(InputValueKey, Value);
+				float Value;
+				if (ShaderGraphNode->GetFloatAttribute(GltfAttributeKey.Value.Key, Value))
+				{
+					MaterialInstanceFactoryNode->AddFloatAttribute(InputValueKey, Value);
+				}
 			}
-		}
-		break;
-		case UE::Interchange::EAttributeTypes::LinearColor:
-		{
-			FLinearColor Value;
-			if (ShaderGraphNode->GetLinearColorAttribute(GltfAttributeKey.Value.Key, Value))
+			break;
+			case UE::Interchange::EAttributeTypes::LinearColor:
 			{
-				MaterialInstanceFactoryNode->AddLinearColorAttribute(InputValueKey, Value);
+				FLinearColor Value;
+				if (ShaderGraphNode->GetLinearColorAttribute(GltfAttributeKey.Value.Key, Value))
+				{
+					MaterialInstanceFactoryNode->AddLinearColorAttribute(InputValueKey, Value);
+				}
 			}
-		}
-		break;
-		case UE::Interchange::EAttributeTypes::String:
-		{
-			FString TextureUid;
-			if (ShaderGraphNode->GetStringAttribute(GltfAttributeKey.Value.Key, TextureUid))
+			break;
+			case UE::Interchange::EAttributeTypes::String:
 			{
-				FString FactoryTextureUid = UInterchangeFactoryBaseNode::BuildFactoryNodeUid(TextureUid);
+				FString TextureUid;
+				if (ShaderGraphNode->GetStringAttribute(GltfAttributeKey.Value.Key, TextureUid))
+				{
+					FString FactoryTextureUid = UInterchangeFactoryBaseNode::BuildFactoryNodeUid(TextureUid);
 
-				MaterialInstanceFactoryNode->AddStringAttribute(InputValueKey, FactoryTextureUid);
-				MaterialInstanceFactoryNode->AddFactoryDependencyUid(FactoryTextureUid);
+					MaterialInstanceFactoryNode->AddStringAttribute(InputValueKey, FactoryTextureUid);
+					MaterialInstanceFactoryNode->AddFactoryDependencyUid(FactoryTextureUid);
+				}
 			}
-		}
-		break;
-		default:
 			break;
+			default:
+				break;
+			}
 		}
+
+		return MaterialInstanceFactoryNode;
 	}
 
-	return MaterialInstanceFactoryNode;
+	return nullptr;
 }
 
 UInterchangeGLTFPipeline::UInterchangeGLTFPipeline()
@@ -357,8 +374,11 @@ void UInterchangeGLTFPipeline::ExecutePipeline(UInterchangeBaseNodeContainer* No
 		{
 			UInterchangeMaterialInstanceFactoryNode* MaterialInstanceFactoryNode = GLTFPipelineSettings->BuildMaterialInstance(NodeContainer, ShaderGraphNode.Value, ShaderGraphNode.Key);
 
-			UInterchangeSourceNode* SourceNode = UInterchangeSourceNode::FindOrCreateUniqueInstance(NodeContainer);
-			UE::Interchange::PipelineHelper::FillSubPathFromSourceNode(MaterialInstanceFactoryNode, SourceNode);
+			if (MaterialInstanceFactoryNode)
+			{
+				UInterchangeSourceNode* SourceNode = UInterchangeSourceNode::FindOrCreateUniqueInstance(NodeContainer);
+				UE::Interchange::PipelineHelper::FillSubPathFromSourceNode(MaterialInstanceFactoryNode, SourceNode);
+			}
 		}
 	}
 }
diff --git a/Engine/Source/Runtime/Interchange/Engine/Private/InterchangeManager.cpp b/Engine/Source/Runtime/Interchange/Engine/Private/InterchangeManager.cpp
--- a/Engine/Source/Runtime/Interchange/Engine/Private/InterchangeManager.cpp
+++ b/Engine/Source/Runtime/Interchange/Engine/Private/InterchangeManager.cpp
@@ -2756,6 +2756,7 @@ void UInterchangeManager::ReleaseAsyncHelper(TWeakPtr<UE::Interchange::FImportAs
 	const bool bLogWarningsAndErrors = !GIsAutomationTesting;
 	
 	double ImportTimeDelta;
+	int32 ResultCount = 0;
 	bool bSucceeded = false;
 	{
 		TSharedPtr<FImportAsyncHelper> AsyncHelperPtr = AsyncHelper.Pin();
@@ -2771,10 +2772,11 @@ void UInterchangeManager::ReleaseAsyncHelper(TWeakPtr<UE::Interchange::FImportAs
 			}
 		}
 		
-		auto ForEachResult = [&bSucceeded, bLogWarningsAndErrors](TArray<UInterchangeResult*>&& Results)
+		auto ForEachResult = [&bSucceeded, &ResultCount, bLogWarningsAndErrors](TArray<UInterchangeResult*>&& Results)
 		{
 			if (!bSucceeded || bLogWarningsAndErrors)
 			{
+				ResultCount += Results.Num();
 				for (UInterchangeResult* Result : Results)
 				{
 					if (!Result)
@@ -2825,6 +2827,7 @@ void UInterchangeManager::ReleaseAsyncHelper(TWeakPtr<UE::Interchange::FImportAs
 		if (Notification.IsValid())
 		{
 			FText TitleText;
+			FText ProgressText = FText::GetEmpty();
 			if (bImportCanceled)
 			{
 				TitleText = NSLOCTEXT("Interchange", "Asynchronous_import_canceled", "Import Canceled");
@@ -2833,7 +2836,14 @@ void UInterchangeManager::ReleaseAsyncHelper(TWeakPtr<UE::Interchange::FImportAs
 			}
 			else
 			{
-				if(bSucceeded)
+				if (ResultCount == 0)
+				{
+					UE_LOG(LogInterchangeEngine, Warning, TEXT("Import Done but no data imported"));
+					ProgressText = NSLOCTEXT("Interchange", "Asynchronous_import_progress", "No Data Imported");
+					TitleText = NSLOCTEXT("Interchange", "Asynchronous_import_empty_results", "Import Done");
+					bSucceeded = true; // Mark the fact as nothing was imported as a success
+				}
+				else if (bSucceeded)
 				{
 					FString ImportDoneText = TEXT("Import Done in ") + FString::SanitizeFloat(ImportTimeDelta) + TEXT(" seconds.");
 					UE_LOG(LogInterchangeEngine, Display, TEXT("%s"), *ImportDoneText);
@@ -2848,7 +2858,7 @@ void UInterchangeManager::ReleaseAsyncHelper(TWeakPtr<UE::Interchange::FImportAs
 				Notification->SetProgressText(NSLOCTEXT("Interchange", "Asynchronous_operation_completed", "Operation Completed."));
 			}
 									
-			Notification->SetComplete(TitleText, FText::GetEmpty(), bSucceeded);
+			Notification->SetComplete(TitleText, ProgressText, bSucceeded);
 			Notification = nullptr; //This should delete the notification
 		}
 
