Fixed issue with level reimport for level instance with a new file
Fixed issue with a level reimport done while the targetd map asset is opened in the editor.
While at it fixed the use case where a user does a level import to create a level instance or a packed actor in a folder where the target map asset already exists. In that case, the user is asked to either replace the map asset or skip it. If the user chooses to replace the map asset, any actor from a previous import will be deleted. If they choose to skip, no actor will be added to the map asset and the existing map will be linked to the level instance actor.

#rnx

diff --git a/Engine/Plugins/Interchange/Runtime/Source/Import/Private/Scene/InterchangeActorFactory.cpp b/Engine/Plugins/Interchange/Runtime/Source/Import/Private/Scene/InterchangeActorFactory.cpp
--- a/Engine/Plugins/Interchange/Runtime/Source/Import/Private/Scene/InterchangeActorFactory.cpp
+++ b/Engine/Plugins/Interchange/Runtime/Source/Import/Private/Scene/InterchangeActorFactory.cpp
@@ -200,7 +200,8 @@ UObject* UInterchangeActorFactory::ImportSceneObject_GameThread(const UInterchan
 		return nullptr;
 	}
 
-	AActor* SpawnedActor = ActorHelper::SpawnFactoryActor(CreateSceneObjectsParams);
+	FText OutReason;
+	AActor* SpawnedActor = ActorHelper::SpawnFactoryActor(CreateSceneObjectsParams, OutReason);
 
 	if (SpawnedActor)
 	{
@@ -232,6 +233,14 @@ UObject* UInterchangeActorFactory::ImportSceneObject_GameThread(const UInterchan
 
 		ProcessLayerNames(FactoryNode, SpawnedActor);
 	}
+	else if (!OutReason.IsEmpty())
+	{
+		UInterchangeResultWarning_Generic* Message = AddMessage<UInterchangeResultWarning_Generic>();
+		Message->SourceAssetName = CreateSceneObjectsParams.SourceData->GetFilename();
+		Message->DestinationAssetName = CreateSceneObjectsParams.ObjectName;
+		Message->AssetType = FactoryNode->GetObjectClass();
+		Message->Text = OutReason;
+	}
 
 	return SpawnedActor;
 }
diff --git a/Engine/Plugins/Interchange/Runtime/Source/Import/Private/Scene/InterchangeActorHelper.cpp b/Engine/Plugins/Interchange/Runtime/Source/Import/Private/Scene/InterchangeActorHelper.cpp
--- a/Engine/Plugins/Interchange/Runtime/Source/Import/Private/Scene/InterchangeActorHelper.cpp
+++ b/Engine/Plugins/Interchange/Runtime/Source/Import/Private/Scene/InterchangeActorHelper.cpp
@@ -42,8 +42,9 @@ namespace UE::Interchange::ActorHelper
 		return  ParentActor;
 	}
 
-	AActor* SpawnFactoryActor(const UInterchangeFactoryBase::FImportSceneObjectsParams& CreateSceneObjectsParams)
+	AActor* SpawnFactoryActor(const UInterchangeFactoryBase::FImportSceneObjectsParams& CreateSceneObjectsParams, FText& OutReason)
 	{
+		OutReason = FText::GetEmpty();
 		UInterchangeActorFactoryNode* FactoryNode = Cast<UInterchangeActorFactoryNode>(CreateSceneObjectsParams.FactoryNode);
 		const UInterchangeBaseNodeContainer* NodeContainer = CreateSceneObjectsParams.NodeContainer;
 
@@ -61,6 +62,11 @@ namespace UE::Interchange::ActorHelper
 		AActor* ParentActor = GetSpawnedParentActor(NodeContainer, FactoryNode);
 		UWorld* const World = [&SpawnParameters, &ParentActor]()
 		{
+			if (SpawnParameters.OverrideLevel)
+			{
+				return SpawnParameters.OverrideLevel->GetWorld();
+			}
+
 			UWorld* DefaultWorld = nullptr;
 
 			if (ParentActor)
@@ -82,7 +88,7 @@ namespace UE::Interchange::ActorHelper
 				DefaultWorld = GEngine->GetWorld();
 			}
 
-			return SpawnParameters.OverrideLevel ? SpawnParameters.OverrideLevel->GetWorld() : DefaultWorld;
+			return DefaultWorld;
 		}();
 
 		if (!World)
@@ -110,6 +116,23 @@ namespace UE::Interchange::ActorHelper
 			return nullptr;
 		}
 
+		// #interchange_levelinstance_rework:
+		// Check whether or not to create the actor if it will be added to a level related to a level instance
+		// Use case: User is importing into level with level instance or packed actor enabled and the targeted
+		// level already exists and they have indicated to skip its update/replacement
+		FString LevelFactoryNodeUid;
+		if (FactoryNode->GetCustomLevelUid(LevelFactoryNodeUid))
+		{
+			if (const UInterchangeFactoryBaseNode* LevelFactoryNode = NodeContainer->GetFactoryNode(LevelFactoryNodeUid))
+			{
+				if (LevelFactoryNode->ShouldSkipNodeImport())
+				{
+					OutReason = NSLOCTEXT("Interchange", "SpawnFactoryActor_Wrong_Level", "Cannot create or update an actor when targeted level exists and should be skipped.");
+					return nullptr;
+				}
+			}
+		}
+
 		if (!SpawnedActor)
 		{
 			SpawnedActor = World->SpawnActor<AActor>(ActorClass, SpawnParameters);
@@ -155,6 +178,12 @@ namespace UE::Interchange::ActorHelper
 
 		return SpawnedActor;
 	}
+	
+	AActor* SpawnFactoryActor(const UInterchangeFactoryBase::FImportSceneObjectsParams& CreateSceneObjectsParams)
+	{
+		FText OutReason;
+		return SpawnFactoryActor(CreateSceneObjectsParams, OutReason);
+	}
 
 	const UInterchangeFactoryBaseNode* FindAssetInstanceFactoryNode(const UInterchangeBaseNodeContainer* NodeContainer, const UInterchangeFactoryBaseNode* ActorFactoryNode)
 	{
diff --git a/Engine/Plugins/Interchange/Runtime/Source/Import/Private/Scene/InterchangeLevelFactory.cpp b/Engine/Plugins/Interchange/Runtime/Source/Import/Private/Scene/InterchangeLevelFactory.cpp
--- a/Engine/Plugins/Interchange/Runtime/Source/Import/Private/Scene/InterchangeLevelFactory.cpp
+++ b/Engine/Plugins/Interchange/Runtime/Source/Import/Private/Scene/InterchangeLevelFactory.cpp
@@ -2,8 +2,11 @@
 
 #include "Scene/InterchangeLevelFactory.h"
 
+#include "Algo/TopologicalSort.h"
+#include "Engine/Level.h"
 #include "Engine/World.h"
-#include "InterchangeAssetImportData.h"
+#include "GameFramework/Actor.h"
+#include "InterchangeAssetUserData.h"
 #include "InterchangeEditorUtilitiesBase.h"
 #include "InterchangeImportCommon.h"
 #include "InterchangeImportLog.h"
@@ -11,6 +14,7 @@
 #include "InterchangeLevelInstanceActorFactoryNode.h"
 #include "InterchangeManager.h"
 #include "InterchangeResult.h"
+#include "InterchangeSceneImportAsset.h"
 #include "InterchangeSourceData.h"
 #include "InterchangeTranslatorBase.h"
 #include "InterchangeVariantSetNode.h"
@@ -39,64 +43,158 @@ namespace UE::Interchange::Private::InterchangeLevelFactory
 		return Cast<UInterchangeLevelFactoryNode>(Arguments.AssetNode);
 	}
 
-	UObject* FindOrCreateAsset(const UInterchangeFactoryBase::FImportAssetObjectParams& Arguments, UClass* TargetClass)
+	UWorld* FindOrCreateWorldAsset(const UInterchangeFactoryBase::FImportAssetObjectParams& Arguments, UClass* TargetClass)
 	{
+		if (Arguments.ReimportObject && Arguments.ReimportObject->GetClass()->IsChildOf(TargetClass))
+		{
+			return Cast<UWorld>(Arguments.ReimportObject);
+		}
 
-		UObject* TargetAsset = nullptr;
-		if (Arguments.ReimportObject)
+		// Check whether map already exists. It should
+		UPackage* ParentPackage = Cast<UPackage>(Arguments.Parent);
+		if (!ensure(ParentPackage))
 		{
-			TargetAsset = Arguments.ReimportObject;
+			return nullptr;
 		}
-		
-		bool bCreateWorldPartition = false;
-		//Verify and ensure if we should create a level or not
-		if (UInterchangeLevelFactoryNode* LevelFactoryNode = Cast<UInterchangeLevelFactoryNode>(Arguments.AssetNode))
+
+		// #interchange_levelinstance_rework
+		// If this is not a reimport and the level already exists, it means the user has accepted to import over
+		// Destroy all imported actors from existing world, equivalent to overwriting source data in static mesh for instance
+		// Otherwise new actors will be added to the world instead of being updated
+		if (!Arguments.ReimportObject)
 		{
-			bool bShouldCreateLevel = false;
-			if (LevelFactoryNode->GetCustomShouldCreateLevel(bShouldCreateLevel))
+			const FTopLevelAssetPath TargetPath(*ParentPackage->GetPathName(), *Arguments.AssetName);
+			if (UWorld* ExistingWorld = Cast<UWorld>(FSoftObjectPath(TargetPath).TryLoad()))
 			{
-				ensure(bShouldCreateLevel ? TargetAsset == nullptr : TargetAsset != nullptr);
+				TSet<AActor*> ActorsToDelete;
+				for (ULevel* Level : ExistingWorld->GetLevels())
+				{
+					for (AActor* Actor : Level->Actors)
+					{
+						if (!IsValid(Actor) || ActorsToDelete.Contains(Actor))
+						{
+							continue;
+						}
+
+						if (USceneComponent* RootComponent = Actor->GetRootComponent())
+						{
+							// Only delete actors created from a previous import
+							if (RootComponent->GetAssetUserData<UInterchangeAssetUserData>())
+							{
+								ActorsToDelete.Add(Actor);
+
+								// All parent actors were part of the import. Add them
+								AActor* ParentActor = Actor->GetAttachParentActor();
+								while (IsValid(ParentActor))
+								{
+									ActorsToDelete.Add(ParentActor);
+									ParentActor = ParentActor->GetAttachParentActor();
+								}
+							}
+						}
+					}
+				}
+
+				// Sort actors from leaf to root to avoid reassignment of child actors
+				// which generate multiple unnecessary editor's notifications when there is a re-parenting
+				TArray<AActor*> SortedActors = ActorsToDelete.Array();
+				ActorsToDelete.Empty();
+
+				TMap<AActor*, TArray<AActor*>> ParentToChildren;
+				for (AActor* Actor : SortedActors)
+				{
+					if (!IsValid(Actor))
+					{
+						continue;
+					}
+
+					AActor* ParentActor = Actor->GetAttachParentActor();
+					if (ParentActor)
+					{
+						TArray<AActor*>& Children = ParentToChildren.FindOrAdd(ParentActor);
+						Children.Add(Actor);
+					}
+				}
+
+				auto FindDependencies = [&ParentToChildren](AActor* Actor) -> TArray<AActor*>
+					{
+						if (TArray<AActor*>* Children = ParentToChildren.Find(Actor))
+						{
+							return *Children;
+						}
+
+						return {};
+					};
+
+				Algo::TopologicalSort(SortedActors, FindDependencies);
+
+				// Delete actors
+				for (AActor* Actor : SortedActors)
+				{
+					if (!IsValid(Actor))
+					{
+						continue;
+					}
+
+					ExistingWorld->EditorDestroyActor(Actor, true);
+
+					// Actor actual deletion can be delayed. Rename to avoid future name collision
+					// Call UObject::Rename directly on actor to avoid AActor::Rename which unnecessarily unregister and re-register components
+					Actor->UObject::Rename(nullptr, GetTransientPackage(), REN_DontCreateRedirectors);
+				}
+
+				return ExistingWorld;
 			}
-			LevelFactoryNode->GetCustomCreateWorldPartitionLevel(bCreateWorldPartition);
 		}
-		// Create a new asset or return existing asset, if possible
-		if (!TargetAsset)
+
+		UInterchangeLevelFactoryNode* LevelFactoryNode = Cast<UInterchangeLevelFactoryNode>(Arguments.AssetNode);
+		if (!LevelFactoryNode)
 		{
-			if(!ensure(IsInGameThread()))
-			{
-				return nullptr;
-			}
-			
-			// Create a new world.
-			const bool bAddToRoot = false;
-
-			constexpr bool bEnableWorldPartitionStreaming = false;
-			constexpr bool bInformEngineOfWorld = false;
-			bool bIsRuntimeOrPIE = false;
-			if (UInterchangeEditorUtilitiesBase* EditorUtilities = UInterchangeManager::GetInterchangeManager().GetEditorUtilities())
-			{
-				bIsRuntimeOrPIE = EditorUtilities->IsRuntimeOrPIE();
-			}
+			return nullptr;
+		}
+
+		// Map creation can only occur in the game thread
+		if (!ensure(IsInGameThread()))
+		{
+			return nullptr;
+		}
 
-			// Those are the init values taken from the default in UWorld::CreateWorld + CreateWorldPartition.
-			UWorld::InitializationValues InitValues = UWorld::InitializationValues()
-				.ShouldSimulatePhysics(false)
-				.EnableTraceCollision(true)
-				.CreateNavigation(!bIsRuntimeOrPIE)
-				.CreateAISystem(!bIsRuntimeOrPIE)
-				.CreateWorldPartition(bCreateWorldPartition)
-				.EnableWorldPartitionStreaming(bEnableWorldPartitionStreaming);
+		bool bShouldCreateLevel = false;
+		if (LevelFactoryNode->GetCustomShouldCreateLevel(bShouldCreateLevel))
+		{
+			ensure(bShouldCreateLevel);
+		}
 
-			UWorld* NewWorld = UWorld::CreateWorld(EWorldType::Inactive, bInformEngineOfWorld, *Arguments.AssetName, Cast<UPackage>(Arguments.Parent), bAddToRoot, ERHIFeatureLevel::Num, &InitValues);
-			NewWorld->SetFlags(RF_Public | RF_Standalone);
-			TargetAsset = NewWorld;
+		bool bCreateWorldPartition = false;
+		LevelFactoryNode->GetCustomCreateWorldPartitionLevel(bCreateWorldPartition);
+
+		// Create a new world.
+		constexpr bool bAddToRoot = false;
+		constexpr bool bEnableWorldPartitionStreaming = false;
+		constexpr bool bInformEngineOfWorld = false;
+		bool bIsRuntimeOrPIE = false;
+		if (UInterchangeEditorUtilitiesBase* EditorUtilities = UInterchangeManager::GetInterchangeManager().GetEditorUtilities())
+		{
+			bIsRuntimeOrPIE = EditorUtilities->IsRuntimeOrPIE();
 		}
-		else if (!TargetAsset->GetClass()->IsChildOf(TargetClass))
+
+		// Those are the init values taken from the default in UWorld::CreateWorld + CreateWorldPartition.
+		UWorld::InitializationValues InitValues = UWorld::InitializationValues()
+			.ShouldSimulatePhysics(false)
+			.EnableTraceCollision(true)
+			.CreateNavigation(!bIsRuntimeOrPIE)
+			.CreateAISystem(!bIsRuntimeOrPIE)
+			.CreateWorldPartition(bCreateWorldPartition)
+			.EnableWorldPartitionStreaming(bEnableWorldPartitionStreaming);
+
+		UWorld* NewWorld = UWorld::CreateWorld(EWorldType::Inactive, bInformEngineOfWorld, *Arguments.AssetName, ParentPackage, bAddToRoot, ERHIFeatureLevel::Num, &InitValues);
+		if (ensure(NewWorld))
 		{
-			TargetAsset = nullptr;
+			NewWorld->SetFlags(RF_Public | RF_Standalone);
+			return NewWorld;
 		}
 
-		return TargetAsset;
+		return nullptr;
 	}
 }
 
@@ -106,6 +204,48 @@ UClass* UInterchangeLevelFactory::GetFactoryClass() const
 	return UWorld::StaticClass();
 }
 
+UObject* UInterchangeLevelFactory::GetObjectToReimport(UObject* ReimportObject, const UInterchangeFactoryBaseNode& FactoryNode, const FString& PackageName, const FString& AssetName, const FString& SubPathString)
+{
+	// #interchange_levelinstance_rework
+	// The code below only works when associated with the UInterchangeGenericLevelPipeline and
+	// the current implementation with only one level instance actor
+
+	// If this is a level reimport, try to find the level of the previous import
+	if (ReimportObject && FactoryNode.IsA<UInterchangeLevelFactoryNode>())
+	{
+		if (UInterchangeSceneImportAsset* SceneImportAsset = Cast<UInterchangeSceneImportAsset>(ReimportObject))
+		{
+			// If reimporting with the same file, the InterchangeSceneImportAsset should be able to return the level of the previous import
+			if (UObject* ObjectToReimport = SceneImportAsset->GetSceneObject(PackageName, AssetName, SubPathString))
+			{
+				return ObjectToReimport;
+			}
+
+			// If reimporting with a new file, return the level of the previous import if any
+			TArray<const UInterchangeFactoryBaseNode*> FacoryNodes = SceneImportAsset->GetFactoryNodesOfClass(UInterchangeLevelFactoryNode::StaticClass());
+			ensure(FacoryNodes.Num() < 2);
+
+			if (FacoryNodes.Num() == 1)
+			{
+				const UInterchangeLevelFactoryNode* LevelFactoryNode = Cast<UInterchangeLevelFactoryNode>(FacoryNodes[0]);
+				
+				if (ensure(LevelFactoryNode))
+				{
+					FSoftObjectPath ObjectPath;
+					ensure(LevelFactoryNode->GetCustomReferenceObject(ObjectPath));
+
+					UObject* ObjectToReimport = ObjectPath.TryLoad();
+					ensure(ObjectToReimport);
+
+					return ObjectToReimport;
+				}
+
+			}
+		}
+	}
+
+	return Super::GetObjectToReimport(ReimportObject, FactoryNode, PackageName, AssetName, SubPathString);
+}
 
 UInterchangeFactoryBase::FImportAssetResult UInterchangeLevelFactory::BeginImportAsset_GameThread(const FImportAssetObjectParams& Arguments)
 {
@@ -120,7 +260,7 @@ UInterchangeFactoryBase::FImportAssetResult UInterchangeLevelFactory::BeginImpor
 		return {};
 	}
 
-	UWorld* WorldAsset = Cast<UWorld>(FindOrCreateAsset(Arguments, TargetClass));
+	UWorld* WorldAsset = FindOrCreateWorldAsset(Arguments, TargetClass);
 
 	if (!WorldAsset)
 	{
diff --git a/Engine/Plugins/Interchange/Runtime/Source/Import/Private/Scene/InterchangeLevelInstanceActorFactory.cpp b/Engine/Plugins/Interchange/Runtime/Source/Import/Private/Scene/InterchangeLevelInstanceActorFactory.cpp
--- a/Engine/Plugins/Interchange/Runtime/Source/Import/Private/Scene/InterchangeLevelInstanceActorFactory.cpp
+++ b/Engine/Plugins/Interchange/Runtime/Source/Import/Private/Scene/InterchangeLevelInstanceActorFactory.cpp
@@ -9,6 +9,7 @@
 #include "InterchangeLevelFactoryNode.h"
 #include "InterchangeLevelInstanceActorFactoryNode.h"
 #include "InterchangeManager.h"
+#include "InterchangeSceneImportAsset.h"
 #include "LevelInstance/LevelInstanceActor.h"
 #include "LevelInstance/LevelInstanceSubsystem.h"
 #include "Nodes/InterchangeBaseNodeContainer.h"
@@ -23,6 +24,49 @@ UClass* UInterchangeLevelInstanceActorFactory::GetFactoryClass() const
 	return ALevelInstance::StaticClass();
 }
 
+UObject* UInterchangeLevelInstanceActorFactory::GetObjectToReimport(UObject* ReimportObject, const UInterchangeFactoryBaseNode& FactoryNode, const FString& PackageName, const FString& AssetName, const FString& SubPathString)
+{
+	// #interchange_levelinstance_rework
+	// The code below only works when associated with the UInterchangeGenericLevelPipeline and
+	// the current implementation with only one LevelInstance actor
+
+	// If this is a level reimport, try to find the level instance of the previous import
+	if (ReimportObject && FactoryNode.IsA<UInterchangeLevelInstanceActorFactoryNode>())
+	{
+		if (UInterchangeSceneImportAsset* SceneImportAsset = Cast<UInterchangeSceneImportAsset>(ReimportObject))
+		{
+			// If reimporting with the same file, the InterchangeSceneImportAsset should be able to return the level of the previous import
+			if (UObject* ObjectToReimport = SceneImportAsset->GetSceneObject(PackageName, AssetName, SubPathString))
+			{
+				return ObjectToReimport;
+			}
+
+			TArray<const UInterchangeFactoryBaseNode*> FacoryNodes = SceneImportAsset->GetFactoryNodesOfClass(UInterchangeLevelInstanceActorFactoryNode::StaticClass());
+			ensure(FacoryNodes.Num() < 2);
+
+			// If reimporting with a new file, return the LevelInstance actor of the previous import if any
+			if (FacoryNodes.Num() == 1)
+			{
+				const UInterchangeLevelInstanceActorFactoryNode* ActorFactoryNode = Cast<UInterchangeLevelInstanceActorFactoryNode>(FacoryNodes[0]);
+
+				if (ensure(ActorFactoryNode))
+				{
+					FSoftObjectPath ObjectPath;
+					ensure(ActorFactoryNode->GetCustomReferenceObject(ObjectPath));
+
+					UObject* ObjectToReimport = ObjectPath.TryLoad();
+					ensure(ObjectToReimport);
+
+					return ObjectToReimport;
+				}
+
+			}
+		}
+	}
+
+	return Super::GetObjectToReimport(ReimportObject, FactoryNode, PackageName, AssetName, SubPathString);
+}
+
 UObject* UInterchangeLevelInstanceActorFactory::ImportSceneObject_GameThread(const UInterchangeFactoryBase::FImportSceneObjectsParams& CreateSceneObjectsParams)
 {
 #if !WITH_EDITOR
@@ -67,6 +111,20 @@ UObject* UInterchangeLevelInstanceActorFactory::ImportSceneObject_GameThread(con
 		return nullptr;
 	}
 
+	// #interchange_levelinstance_rework
+	// If the reference world is loaded in the editor, do not create a level instance
+	// This scenario needs to be properly implemented
+	if (bIsReimport && ReferenceWorld == (UWorld*)GWorld)
+	{
+		UInterchangeResultWarning_Generic* Message = AddMessage<UInterchangeResultWarning_Generic>();
+		Message->SourceAssetName = CreateSceneObjectsParams.SourceData ? CreateSceneObjectsParams.SourceData->GetFilename() : TEXT("");
+		Message->DestinationAssetName = CreateSceneObjectsParams.ObjectName;
+		Message->AssetType = ReferenceLevelFactoryNode->GetObjectClass();
+		Message->Text = NSLOCTEXT("ImportSceneObject_GameThread", "SkipLevelInstance", "Reimporting a level instance in the level it references. No LevelInstance will be created or updated.");
+		
+		return nullptr;
+	}
+
 	//Reference world must be saved if we want to update the level instance actor
 	if (UInterchangeEditorUtilitiesBase* EditorUtilities = UInterchangeManager::GetInterchangeManager().GetEditorUtilities())
 	{
diff --git a/Engine/Plugins/Interchange/Runtime/Source/Import/Public/Scene/InterchangeActorHelper.h b/Engine/Plugins/Interchange/Runtime/Source/Import/Public/Scene/InterchangeActorHelper.h
--- a/Engine/Plugins/Interchange/Runtime/Source/Import/Public/Scene/InterchangeActorHelper.h
+++ b/Engine/Plugins/Interchange/Runtime/Source/Import/Public/Scene/InterchangeActorHelper.h
@@ -4,6 +4,8 @@
 
 #include "InterchangeFactoryBase.h"
 
+#define UE_API INTERCHANGEIMPORT_API
+
 class AActor;
 
 class UInterchangeActorFactoryNode;
@@ -21,37 +23,40 @@ namespace UE::Interchange::ActorHelper
 	/**
 	 * Returns the parent actor of the actor described by the FactoryNode using the node hierarchy.
 	 */
-	INTERCHANGEIMPORT_API AActor* GetSpawnedParentActor(const UInterchangeBaseNodeContainer* NodeContainer, const UInterchangeActorFactoryNode* FactoryNode);
+	UE_API AActor* GetSpawnedParentActor(const UInterchangeBaseNodeContainer* NodeContainer, const UInterchangeActorFactoryNode* FactoryNode);
 
 	/**
 	 * Spawns and return a AAActor using the given SceneObjectsParam.
 	 * The actor's parent will be determine by the node hierarchy and its type by the factory node's ObjectClass.
 	 */
-	INTERCHANGEIMPORT_API AActor* SpawnFactoryActor(const UInterchangeFactoryBase::FImportSceneObjectsParams& CreateSceneObjectsParams);
-	
+	UE_API AActor* SpawnFactoryActor(const UInterchangeFactoryBase::FImportSceneObjectsParams& CreateSceneObjectsParams, FText& OutReason);
+	UE_API AActor* SpawnFactoryActor(const UInterchangeFactoryBase::FImportSceneObjectsParams& CreateSceneObjectsParams);
+
 	/**
 	 * Returns the factory node of the asset instanced by ActorFactoryNode.
 	 */
-	INTERCHANGEIMPORT_API const UInterchangeFactoryBaseNode* FindAssetInstanceFactoryNode(const UInterchangeBaseNodeContainer* NodeContainer, const UInterchangeFactoryBaseNode* ActorFactoryNode);
+	UE_API const UInterchangeFactoryBaseNode* FindAssetInstanceFactoryNode(const UInterchangeBaseNodeContainer* NodeContainer, const UInterchangeFactoryBaseNode* ActorFactoryNode);
 
 	/**
 	 * Applies material slot dependencies stored in ActorFactoryNode to MeshComponent.
 	 */
-	INTERCHANGEIMPORT_API void ApplySlotMaterialDependencies(const UInterchangeBaseNodeContainer& NodeContainer, const UInterchangeMeshActorFactoryNode& ActorFactoryNode, UMeshComponent& MeshComponent);
+	UE_API void ApplySlotMaterialDependencies(const UInterchangeBaseNodeContainer& NodeContainer, const UInterchangeMeshActorFactoryNode& ActorFactoryNode, UMeshComponent& MeshComponent);
 
 	/**
 	 * Applies custom attributes from to ObjectToUpdate based on content of CreateSceneObjectsParams
 	 * Reimport policies will be applied if applicable
 	 */
-	INTERCHANGEIMPORT_API void ApplyAllCustomAttributes(const UInterchangeFactoryBase::FImportSceneObjectsParams& CreateSceneObjectsParams, UObject& ObjectToUpdate);
+	UE_API void ApplyAllCustomAttributes(const UInterchangeFactoryBase::FImportSceneObjectsParams& CreateSceneObjectsParams, UObject& ObjectToUpdate);
 
 	/**
 	 * Adds metadata to Actor via UInterchangeAssetUserData
 	 */
-	INTERCHANGEIMPORT_API void AddInterchangeAssetUserDataToActor(AActor* Actor, const UInterchangeSceneImportAsset* SceneImportAsset, const UInterchangeFactoryBaseNode* FactoryNode);
+	UE_API void AddInterchangeAssetUserDataToActor(AActor* Actor, const UInterchangeSceneImportAsset* SceneImportAsset, const UInterchangeFactoryBaseNode* FactoryNode);
 
 	/**
 	 * Adds metadata to Actor via UInterchangeAssetUserData
 	 */
-	INTERCHANGEIMPORT_API void AddInterchangeLevelAssetUserDataToWorld(UWorld* World, const UInterchangeSceneImportAsset* SceneImportAsset);
+	UE_API void AddInterchangeLevelAssetUserDataToWorld(UWorld* World, const UInterchangeSceneImportAsset* SceneImportAsset);
 }
+
+#undef UE_API
\ No newline at end of file
diff --git a/Engine/Plugins/Interchange/Runtime/Source/Import/Public/Scene/InterchangeLevelFactory.h b/Engine/Plugins/Interchange/Runtime/Source/Import/Public/Scene/InterchangeLevelFactory.h
--- a/Engine/Plugins/Interchange/Runtime/Source/Import/Public/Scene/InterchangeLevelFactory.h
+++ b/Engine/Plugins/Interchange/Runtime/Source/Import/Public/Scene/InterchangeLevelFactory.h
@@ -23,6 +23,8 @@ public:
 	UE_API virtual UClass* GetFactoryClass() const override;
 	UE_API virtual FImportAssetResult BeginImportAsset_GameThread(const FImportAssetObjectParams& Arguments) override;
 	UE_API virtual void SetupObject_GameThread(const FSetupObjectParams& Arguments) override;
+	// For a level reimport, returns the level created on the previous import if found
+	UE_API virtual UObject* GetObjectToReimport(UObject* ReimportObject, const UInterchangeFactoryBaseNode& FactoryNode, const FString& PackageName, const FString& AssetName, const FString& SubPathString) override;
 
 	// Interchange factory base interface end
 	//////////////////////////////////////////////////////////////////////////
diff --git a/Engine/Plugins/Interchange/Runtime/Source/Import/Public/Scene/InterchangeLevelInstanceActorFactory.h b/Engine/Plugins/Interchange/Runtime/Source/Import/Public/Scene/InterchangeLevelInstanceActorFactory.h
--- a/Engine/Plugins/Interchange/Runtime/Source/Import/Public/Scene/InterchangeLevelInstanceActorFactory.h
+++ b/Engine/Plugins/Interchange/Runtime/Source/Import/Public/Scene/InterchangeLevelInstanceActorFactory.h
@@ -25,8 +25,9 @@ public:
 	// Interchange factory base interface begin
 
 	UE_API virtual UClass* GetFactoryClass() const override;
-
 	UE_API virtual UObject* ImportSceneObject_GameThread(const UInterchangeFactoryBase::FImportSceneObjectsParams& CreateSceneObjectsParams) override;
+	// For a level reimport, returns the LevelInstance actor created on the previous import if found
+	UE_API virtual UObject* GetObjectToReimport(UObject* ReimportObject, const UInterchangeFactoryBaseNode& FactoryNode, const FString& PackageName, const FString& AssetName, const FString& SubPathString) override;
 
 	// Interchange factory base interface end
 	//////////////////////////////////////////////////////////////////////////
diff --git a/Engine/Source/Runtime/Interchange/Engine/Private/InterchangeSceneImportAsset.cpp b/Engine/Source/Runtime/Interchange/Engine/Private/InterchangeSceneImportAsset.cpp
--- a/Engine/Source/Runtime/Interchange/Engine/Private/InterchangeSceneImportAsset.cpp
+++ b/Engine/Source/Runtime/Interchange/Engine/Private/InterchangeSceneImportAsset.cpp
@@ -306,6 +306,32 @@ const UInterchangeFactoryBaseNode* UInterchangeSceneImportAsset::GetFactoryNode(
 
 	return nullptr;
 }
+TArray<const UInterchangeFactoryBaseNode*> UInterchangeSceneImportAsset::GetFactoryNodesOfClass(const UClass* Class) const
+{
+#if WITH_EDITORONLY_DATA
+	if (!ensure(AssetImportData && Class && Class->IsChildOf(UInterchangeFactoryBaseNode::StaticClass())))
+	{
+		return {};
+	}
+
+	if (UInterchangeBaseNodeContainer* Container = AssetImportData->GetNodeContainer())
+	{
+		TArray<const UInterchangeFactoryBaseNode*> FactoryNodes;
+
+		Container->IterateNodes([&FactoryNodes, Class](const FString& NodeUid, UInterchangeBaseNode* Node)
+			{
+				if (Node && Node->IsA(Class))
+				{
+					FactoryNodes.Add(Cast<UInterchangeFactoryBaseNode>(Node));
+				}
+			});
+
+		return FactoryNodes;
+	}
+#endif
+
+	return {};
+}
 
 void UInterchangeSceneImportAsset::GetSceneSoftObjectPaths(TArray<FSoftObjectPath>& SoftObjectPaths) const
 {
diff --git a/Engine/Source/Runtime/Interchange/Engine/Private/Tasks/InterchangeTaskImportObject.cpp b/Engine/Source/Runtime/Interchange/Engine/Private/Tasks/InterchangeTaskImportObject.cpp
--- a/Engine/Source/Runtime/Interchange/Engine/Private/Tasks/InterchangeTaskImportObject.cpp
+++ b/Engine/Source/Runtime/Interchange/Engine/Private/Tasks/InterchangeTaskImportObject.cpp
@@ -366,18 +366,7 @@ void UE::Interchange::FTaskImportObject_GameThread::Execute()
 		};
 
 	//Check if the node reference a map
-	const bool bRefObjectIsMap = [this]()
-		{
-			FSoftObjectPath Reference;
-			if (FactoryNode->GetCustomReferenceObject(Reference))
-			{
-				if (UObject* RefObj = Reference.TryLoad())
-				{
-					return RefObj->IsA<UWorld>();
-				}
-			}
-			return false;
-		}();
+	const bool bRefObjectIsMap = FactoryNode->GetObjectClass()->IsChildOf<UWorld>();
 
 	//If the factory node is disable see if there is an existing UObject for it
 	if (!FactoryNode->IsEnabled())
@@ -576,6 +565,8 @@ void UE::Interchange::FTaskImportObject_GameThread::Execute()
 		FactoryNode->SetCustomReferenceObject(FSoftObjectPath(ExistingAsset));
 		if (bSkipObjectNoReplace)
 		{
+			// #interchange_levelinstance_rework: Need better way to indicate that actors should not be added to level
+			FactoryNode->SetSkipNodeImport();
 			return;
 		}
 
diff --git a/Engine/Source/Runtime/Interchange/Engine/Public/InterchangeSceneImportAsset.h b/Engine/Source/Runtime/Interchange/Engine/Public/InterchangeSceneImportAsset.h
--- a/Engine/Source/Runtime/Interchange/Engine/Public/InterchangeSceneImportAsset.h
+++ b/Engine/Source/Runtime/Interchange/Engine/Public/InterchangeSceneImportAsset.h
@@ -8,6 +8,8 @@
 
 #include "InterchangeSceneImportAsset.generated.h"
 
+#define UE_API INTERCHANGEENGINE_API
+
 class UInterchangeAssetImportData;
 class UInterchangeFactoryBaseNode;
 class ULevel;
@@ -20,7 +22,7 @@ class UInterchangeSceneImportAsset : public UObject, public IInterface_AssetUser
 {
 	GENERATED_BODY()
 
-	INTERCHANGEENGINE_API virtual ~UInterchangeSceneImportAsset();
+	UE_API virtual ~UInterchangeSceneImportAsset();
 
 public:
 #if WITH_EDITORONLY_DATA
@@ -34,25 +36,25 @@ public:
 #endif // #if WITH_EDITORONLY_DATA
 
 	//~ Begin UObject Interface
-	INTERCHANGEENGINE_API virtual void GetAssetRegistryTags(FAssetRegistryTagsContext Context) const override;
+	UE_API virtual void GetAssetRegistryTags(FAssetRegistryTagsContext Context) const override;
 	UE_DEPRECATED(5.4, "Implement the version that takes FAssetRegistryTagsContext instead.")
-	INTERCHANGEENGINE_API virtual void GetAssetRegistryTags(TArray<FAssetRegistryTag>& OutTags) const override;
-	INTERCHANGEENGINE_API virtual void PostLoad() override;
+	UE_API virtual void GetAssetRegistryTags(TArray<FAssetRegistryTag>& OutTags) const override;
+	UE_API virtual void PostLoad() override;
 	//~ End UObject Interface
 
 	//~ Begin IInterface_AssetUserData Interface
-	INTERCHANGEENGINE_API virtual void AddAssetUserData( UAssetUserData* InUserData ) override;
-	INTERCHANGEENGINE_API virtual void RemoveUserDataOfClass( TSubclassOf<UAssetUserData> InUserDataClass ) override;
-	INTERCHANGEENGINE_API virtual UAssetUserData* GetAssetUserDataOfClass( TSubclassOf<UAssetUserData> InUserDataClass ) override;
-	INTERCHANGEENGINE_API virtual const TArray<UAssetUserData*>* GetAssetUserDataArray() const override;
+	UE_API virtual void AddAssetUserData( UAssetUserData* InUserData ) override;
+	UE_API virtual void RemoveUserDataOfClass( TSubclassOf<UAssetUserData> InUserDataClass ) override;
+	UE_API virtual UAssetUserData* GetAssetUserDataOfClass( TSubclassOf<UAssetUserData> InUserDataClass ) override;
+	UE_API virtual const TArray<UAssetUserData*>* GetAssetUserDataArray() const override;
 	//~ End IInterface_AssetUserData Interface
 
 #if WITH_EDITOR
-	INTERCHANGEENGINE_API void RegisterWorldRenameCallbacks();
+	UE_API void RegisterWorldRenameCallbacks();
 #endif
 
 	/** Updates the SceneObjects cache based on the node container stored in AssetImportData */
-	INTERCHANGEENGINE_API void UpdateSceneObjects();
+	UE_API void UpdateSceneObjects();
 
 	/**
 	 * Returns the UObject which asset path is '//PackageName.AssetName[:SubPathString]'.
@@ -62,7 +64,7 @@ public:
 	 * @param AssetName: Asset name of the actual object to reimport
 	 * @param SubPathString: Optional subobject name
 	 */
-	INTERCHANGEENGINE_API UObject* GetSceneObject(const FString& PackageName, const FString& AssetName, const FString& SubPathString = FString()) const;
+	UE_API UObject* GetSceneObject(const FString& PackageName, const FString& AssetName, const FString& SubPathString = FString()) const;
 
 	/**
 	 * Returns the factory node associated with the asset which path is '//PackageName.AssetName[:SubPathString]'.
@@ -72,13 +74,19 @@ public:
 	 * @param AssetName: Asset name of the actual object to reimport
 	 * @param SubPathString: Optional subobject name
 	 */
-	INTERCHANGEENGINE_API const UInterchangeFactoryBaseNode* GetFactoryNode(const FString& PackageName, const FString& AssetName, const FString& SubPathString = FString()) const;
-	
-	INTERCHANGEENGINE_API const UInterchangeFactoryBaseNode* GetFactoryNode(const FSoftObjectPath& ObjectPath) const;
+	UE_API const UInterchangeFactoryBaseNode* GetFactoryNode(const FString& PackageName, const FString& AssetName, const FString& SubPathString = FString()) const;
 
-	INTERCHANGEENGINE_API void GetSceneSoftObjectPaths(TArray<FSoftObjectPath>& SoftObjectPaths) const;
+	/**
+	 * Returns the array of factory nodes corresponding to the given class.
+	 * @param Class: Class to be looked for
+	 */
+	UE_API TArray<const UInterchangeFactoryBaseNode*> GetFactoryNodesOfClass(const UClass* Class) const;
 
-	INTERCHANGEENGINE_API bool ContainsObject(const FSoftObjectPath& SoftObjectPath) const;
+	UE_API const UInterchangeFactoryBaseNode* GetFactoryNode(const FSoftObjectPath& ObjectPath) const;
+
+	UE_API void GetSceneSoftObjectPaths(TArray<FSoftObjectPath>& SoftObjectPaths) const;
+
+	UE_API bool ContainsObject(const FSoftObjectPath& SoftObjectPath) const;
 
 private:
 
@@ -105,3 +113,5 @@ private:
 	TMap< FSoftObjectPath, FString > SceneObjects;
 #endif // #if WITH_EDITORONLY_DATA
 };
+
+#undef UE_API
