Added warning messages to help users understand why custom collisions following naming convention are not imported properly

Two warnings:
- A mesh starting with a custom collision prefix was detected but not converted to a collision properly. (either because of a typo, or because the convention was not followed properly)
- A mesh starting with a custom collision prefix was detected but the 'Import Collisions According To Mesh Name' setting is disabled.

diff --git a/Engine/Plugins/Interchange/Runtime/Source/Pipelines/Private/InterchangeGenericStaticMeshPipeline.cpp b/Engine/Plugins/Interchange/Runtime/Source/Pipelines/Private/InterchangeGenericStaticMeshPipeline.cpp
--- a/Engine/Plugins/Interchange/Runtime/Source/Pipelines/Private/InterchangeGenericStaticMeshPipeline.cpp
+++ b/Engine/Plugins/Interchange/Runtime/Source/Pipelines/Private/InterchangeGenericStaticMeshPipeline.cpp
@@ -916,6 +916,19 @@ void UInterchangeGenericMeshPipeline::AddLodDataToStaticMesh(UInterchangeStaticM
 
 			UE::Interchange::MeshesUtilities::ApplySlotMaterialDependencies(*StaticMeshFactoryNode, SlotMaterialDependencies, *BaseNodeContainer, &ExistingLodSlotMaterialDependencies);
 
+			const FString MeshName = UE::Interchange::Private::GetNodeName(PipelineMeshesUtilities, *BaseNodeContainer, NodeUid);
+			// Ignore the "_" so we can detect potential typos done by the user.
+			const bool bHasCollisionPrefix = MeshName.StartsWith(TEXT("UBX")) || MeshName.StartsWith(TEXT("UCX")) || MeshName.StartsWith(TEXT("USP")) || MeshName.StartsWith(TEXT("UCP")) || MeshName.StartsWith(TEXT("MCDCX"));
+
+			auto LogWarning = [this, &StaticMeshFactoryNode](const FText& Text)
+				{
+					UInterchangeResultWarning_Generic* Message = AddMessage<UInterchangeResultWarning_Generic>();
+					Message->SourceAssetName = SourceDatas[0]->GetFilename();
+					Message->AssetFriendlyName = StaticMeshFactoryNode->GetAssetName();
+					Message->AssetType = StaticMeshFactoryNode->GetObjectClass();
+					Message->Text = Text;
+				};
+
 			if (bImportCollisionAccordingToMeshName)
 			{
 				TTuple<EInterchangeMeshCollision, FString> MeshType = GetCollisionMeshType(PipelineMeshesUtilities, *BaseNodeContainer, NodeUid, NodeUids, bImportCollisionAccordingToMeshName);
@@ -927,6 +940,15 @@ void UInterchangeGenericMeshPipeline::AddLodDataToStaticMesh(UInterchangeStaticM
 				{
 				case EInterchangeMeshCollision::None:
 					LodDataNode->AddMeshUid(NodeUid);
+
+					// If the user is attempting to import a collision mesh using a naming convention but the mesh doesn't match, warn them that it won't be imported as collision.
+					if (bHasCollisionPrefix)
+					{
+						LogWarning(FText::Format(NSLOCTEXT("UInterchangeGenericStaticMeshPipeline", "NoMatchingRenderMeshForCustomCollision",
+							"No render mesh found for custom collision '{0}'. The name must follow the pattern 'UCX_[ExactRenderMeshName]', and a matching render mesh with that name must exist in the source file."),
+							FText::FromString(MeshName)));
+					}
+
 					break;
 
 				case EInterchangeMeshCollision::Box:
@@ -960,6 +982,14 @@ void UInterchangeGenericMeshPipeline::AddLodDataToStaticMesh(UInterchangeStaticM
 			else
 			{
 				LodDataNode->AddMeshUid(NodeUid);
+
+				// Warn the user that this mesh appears to be a custom collision mesh but that the pipeline is not set to import it as such
+				if (bHasCollisionPrefix)
+				{
+					LogWarning(FText::Format(NSLOCTEXT("UInterchangeGenericStaticMeshPipeline", "ImportCollisionAccordingToMeshNameDisabled",
+						"The mesh '{0}' appears to use custom collision naming pattern, but the setting 'Import Collisions According to Mesh Name' is disabled. It will be treated as a render mesh."),
+						FText::FromString(MeshName)));
+				}
 			}
 		}
 		UE::Interchange::MeshesUtilities::ReorderSlotMaterialDependencies(*StaticMeshFactoryNode, *BaseNodeContainer);
