Fixed regression introduce by CL 46787379
Root cause: Removed code which should not have been
Solution: Make the initial code more explicit on the intent plus comments.

#rnx

diff --git a/Engine/Plugins/Interchange/Runtime/Source/Pipelines/Private/InterchangeGenericScenesPipeline.cpp b/Engine/Plugins/Interchange/Runtime/Source/Pipelines/Private/InterchangeGenericScenesPipeline.cpp
--- a/Engine/Plugins/Interchange/Runtime/Source/Pipelines/Private/InterchangeGenericScenesPipeline.cpp
+++ b/Engine/Plugins/Interchange/Runtime/Source/Pipelines/Private/InterchangeGenericScenesPipeline.cpp
@@ -808,10 +808,27 @@ void UInterchangeGenericLevelPipeline::ExecuteSceneNodePreImport(const FTransfor
 		}
 	}
 
-	bool bCustomVisibility = true;
-	SceneNode->GetCustomComponentVisibility(bCustomVisibility);
-	ActorFactoryNode->SetCustomComponentVisibility(bCustomVisibility);
+	bool bCustomVisibility;
+	// Copy component's visibility from SceneNode or set component visible by default
+	if (SceneNode->GetCustomComponentVisibility(bCustomVisibility))
+	{
+		ActorFactoryNode->SetCustomComponentVisibility(bCustomVisibility);
+	}
+	else
+	{
+		ActorFactoryNode->SetCustomComponentVisibility(true);
+	}
 	
+	// Copy actor's visibility from SceneNode or set actor visible by default
+	if (SceneNode->GetCustomActorVisibility(bCustomVisibility))
+	{
+		ActorFactoryNode->SetCustomActorVisibility(bCustomVisibility);
+	}
+	else
+	{
+		ActorFactoryNode->SetCustomActorVisibility(true);
+	}
+
 	ActorFactoryNode->SetCustomMobility(EComponentMobility::Static);
 
 	if (TranslatedAssetNode)
