Replaced use of /InterchangeAssets/Materials/PhongSurfaceMaterial.PhongSurfaceMaterial with /InterchangeAssets/Materials/FBXLegacyPhongSurfaceMaterial.FBXLegacyPhongSurfaceMaterial

#rnx

diff --git a/Engine/Plugins/Interchange/Runtime/Source/Parsers/Fbx/Private/Ufbx/UfbxMaterial.cpp b/Engine/Plugins/Interchange/Runtime/Source/Parsers/Fbx/Private/Ufbx/UfbxMaterial.cpp
--- a/Engine/Plugins/Interchange/Runtime/Source/Parsers/Fbx/Private/Ufbx/UfbxMaterial.cpp
+++ b/Engine/Plugins/Interchange/Runtime/Source/Parsers/Fbx/Private/Ufbx/UfbxMaterial.cpp
@@ -18,6 +18,10 @@
 #include "Texture/InterchangeImageWrapperTranslator.h"
 #endif
 
+#define GET_VALLUE_FROM_FACTOR_MAP( FactorMap ) FactorMap.has_value ? FactorMap.value_real : 1.f;
+
+#define LOCTEXT_NAMESPACE "InterchangeUfbxMaterial"
+
 namespace UE::Interchange::Private
 {
 
@@ -141,6 +145,7 @@ namespace UE::Interchange::Private
 	UInterchangeMaterialInstanceNode* CreateMaterialInstanceNode(UInterchangeBaseNodeContainer& NodeContainer, const FString& NodeUID, const FString& NodeLabel)
 	{
 		UInterchangeMaterialInstanceNode* MaterialInstanceNode = NewObject<UInterchangeMaterialInstanceNode>(&NodeContainer);
+		MaterialInstanceNode->SetCustomParent(TEXT("/InterchangeAssets/Materials/FBXLegacyPhongSurfaceMaterial.FBXLegacyPhongSurfaceMaterial"));
 		NodeContainer.SetupNode(MaterialInstanceNode, NodeUID, NodeLabel, EInterchangeNodeContainerType::TranslatedAsset);
 
 		return MaterialInstanceNode;
@@ -152,19 +157,32 @@ namespace UE::Interchange::Private
 
 		for (ufbx_material* Material : Parser.Scene->materials)
 		{
+			if (!Material)
+			{
+				continue;
+			}
+
 			FString MaterialLabel = Parser.GetMaterialLabel(*Material).ToString();
 			
 			FString MaterialUid = Parser.GetMaterialUid(*Material);
 			UInterchangeMaterialInstanceNode* MaterialInstanceNode = CreateMaterialInstanceNode(NodeContainer, MaterialUid, MaterialLabel);
+			if (MaterialInstanceNode == nullptr)
+			{
+				FFormatNamedArguments Args
+				{
+					{ TEXT("MaterialName"), FText::FromString(MaterialLabel) }
+				};
+				UInterchangeResultError_Generic* Message = Parser.AddMessage<UInterchangeResultError_Generic>();
+				Message->Text = FText::Format(LOCTEXT("CannotCreateUfbxMaterial", "Cannot create uFBX material '{MaterialName}'."), Args);
+				continue;
+			}
 
 			// WithUVTransform
 			bool bHasUVTransform = false;
-			// "/InterchangeAssets/Materials/PhongSurfaceMaterial.PhongSurfaceMaterial
-			FString ParentMaterial{ TEXT("PhongSurfaceMaterial") };
 
 			FMaterialMapConverter MaterialMapConverter(Parser, NodeContainer, *MaterialInstanceNode);
 			
-			auto SetColorFromMap = [&MaterialInstanceNode](FName InputName, const ufbx_material_map& Map, float Factor, TFunctionRef<bool(const FLinearColor&)> Filter = [](const FLinearColor&){return true;})
+			auto SetColorFromMap = [&MaterialInstanceNode](FName InputName, const ufbx_material_map& Map, float Factor, TFunctionRef<bool(const FLinearColor&)> Filter = [](const FLinearColor&){return true;}) -> bool
 			{
 				if (Map.value_components >=3)
 				{
@@ -172,16 +190,18 @@ namespace UE::Interchange::Private
 					FLinearColor PropertyValue(Convert::ConvertVec4(Map.value_vec4) * Factor);
 					if (Filter(PropertyValue))
 					{
-						MaterialInstanceNode->AddVectorParameterValue(InputAttributeKey, PropertyValue);
+						return MaterialInstanceNode->AddVectorParameterValue(InputAttributeKey, PropertyValue);
 					}
 				}
+
+				return false;
 			};
 
-			auto SetScalarFromMap = [&MaterialInstanceNode](FName InputName, const ufbx_material_map& Map, float Factor)
+			auto SetScalarFromMap = [&MaterialInstanceNode](FName InputName, const ufbx_material_map& Map, float Factor) -> bool
 			{
 				const FString InputAttributeKey = InputName.ToString();
 				float PropertyValue(Map.value_real);
-				MaterialInstanceNode->AddScalarParameterValue(InputAttributeKey, PropertyValue);
+				return MaterialInstanceNode->AddScalarParameterValue(InputAttributeKey, PropertyValue);
 			};
 
 			auto SetTextureFromMap = [&MaterialInstanceNode, &MaterialMapConverter, &bHasUVTransform](FName InputName, const ufbx_material_map& Map, float Factor) -> bool
@@ -201,22 +221,26 @@ namespace UE::Interchange::Private
 
 				MaterialInstanceNode->AddTextureParameterValue(InputName.ToString() + TEXT("Map"), TextureNode->GetUniqueID());
 
-				// #ufbx_todo: FBX parser used MapWeight 1 and multiplied although there was Factor provided. Why?
 				MaterialInstanceNode->AddScalarParameterValue(InputName.ToString() + TEXT("MapWeight"), Factor);
 
-				if (Texture->has_uv_transform)
-				{
-					bHasUVTransform = true;
-					const ufbx_transform& UVTransform = Texture->uv_transform;
-					FVector4f OffsetScale(
-						UVTransform.translation.x, UVTransform.translation.y, 
-						UVTransform.scale.x, UVTransform.scale.y);
-					MaterialInstanceNode->AddVectorParameterValue(InputName.ToString() + TEXT("Map") + Inputs::PostFix::OffsetScale, OffsetScale);
-
-					FString RotationParameterName = InputName.ToString() + TEXT("Map") + Inputs::PostFix::Rotation;
-					Math::TVector<double> AttributeValue = Convert::ConvertQuat(UVTransform.rotation).Euler();
-					MaterialInstanceNode->AddScalarParameterValue(RotationParameterName, FMath::DegreesToRadians(AttributeValue.Z));
-				}
+				// #ufbx_todo: Handle texture transform
+				// Current material asset used as parent of the material instance does not handle UV transform yet.
+				// Note that handling of UV transform should also be added to FBX SDK based parser
+				// Begin of future handling of UV transform
+				//if (Texture->has_uv_transform)
+				//{
+				//	bHasUVTransform = true;
+				//	const ufbx_transform& UVTransform = Texture->uv_transform;
+				//	FVector4f OffsetScale(
+				//		UVTransform.translation.x, UVTransform.translation.y, 
+				//		UVTransform.scale.x, UVTransform.scale.y);
+				//	MaterialInstanceNode->AddVectorParameterValue(InputName.ToString() + TEXT("Map") + Inputs::PostFix::OffsetScale, OffsetScale);
+
+				//	FString RotationParameterName = InputName.ToString() + TEXT("Map") + Inputs::PostFix::Rotation;
+				//	Math::TVector<double> AttributeValue = Convert::ConvertQuat(UVTransform.rotation).Euler();
+				//	MaterialInstanceNode->AddScalarParameterValue(RotationParameterName, FMath::DegreesToRadians(AttributeValue.Z));
+				//}
+				// End of future handling of UV transform
 
 				return true;
 			};
@@ -227,144 +251,119 @@ namespace UE::Interchange::Private
 			if (Material->fbx.diffuse_color.has_value)
 			{
 				FName InputName = Phong::Parameters::DiffuseColor;
-				const float Factor = Material->fbx.diffuse_factor.value_real;
+				const float Factor = GET_VALLUE_FROM_FACTOR_MAP(Material->fbx.diffuse_factor);
 				const ufbx_material_map& Map = Material->fbx.diffuse_color;
 
-				if (!SetTextureFromMap(InputName, Map, Factor))
-				{
-					SetColorFromMap(InputName, Map, Factor);
-				}
-
-				bHasInput = true;
+				bHasInput |= SetTextureFromMap(InputName, Map, Factor) ? true : SetColorFromMap(InputName, Map, Factor);
 			}
 
 			// Ambient
 			if (Material->fbx.ambient_color.has_value)
 			{
 				FName InputName = Phong::Parameters::AmbientColor;
-				const float Factor = Material->fbx.ambient_factor.value_real;
+				const float Factor = GET_VALLUE_FROM_FACTOR_MAP(Material->fbx.ambient_factor);
 				const ufbx_material_map& Map = Material->fbx.ambient_color;
 
-				if (!SetTextureFromMap(InputName, Map, Factor))
-				{
-					SetColorFromMap(InputName, Map, Factor);
-				}
-
-				bHasInput = true;
+				bHasInput |= SetTextureFromMap(InputName, Map, Factor) ? true : SetColorFromMap(InputName, Map, Factor);
 			}
 			
 			// Emissive
 			if (Material->fbx.emission_color.has_value)
 			{
 				const FName InputName = Phong::Parameters::EmissiveColor;
-				const float Factor = Material->fbx.emission_factor.value_real;
+				const float Factor = GET_VALLUE_FROM_FACTOR_MAP(Material->fbx.emission_factor);
 				const ufbx_material_map& Map = Material->fbx.emission_color;
 
-				if (!SetTextureFromMap(InputName, Map, Factor))
-				{
-					SetColorFromMap(InputName, Map, Factor, [](const FLinearColor& Color) { return !Color.IsAlmostBlack(); });
-				}
-
-				bHasInput = true;
+				bHasInput |= SetTextureFromMap(InputName, Map, Factor) ? true : SetColorFromMap(InputName, Map, Factor);
 			}
 
-			// if it has no specular no shininess then just take a Lambert material
 			bool bIsPhong = false;
 
 			// Specular
 			if (Material->fbx.specular_color.has_value)
 			{
 				const FName InputName = Phong::Parameters::SpecularColor;
-				const float Factor = Material->fbx.specular_factor.value_real;
+				const float Factor = GET_VALLUE_FROM_FACTOR_MAP(Material->fbx.specular_factor);
 				const ufbx_material_map& Map = Material->fbx.specular_color;
 
-				if (!SetTextureFromMap(InputName, Map, Factor))
+				if (bool bHasSpecular = SetTextureFromMap(InputName, Map, Factor) ? true : SetColorFromMap(InputName, Map, Factor))
 				{
-					SetColorFromMap(InputName, Map, Factor);
+					bHasInput = true;
+					bIsPhong = true;
 				}
-
-				bHasInput = true;
-				bIsPhong = true;
 			}
 
 			// Shininess
 			if (Material->fbx.specular_exponent.has_value)
 			{
 				const FName InputName = Phong::Parameters::Shininess;
-				const float Factor = Material->fbx.specular_exponent.value_real;
+				const float Shininess = Material->fbx.specular_exponent.value_real;
 				const ufbx_material_map& Map = Material->fbx.specular_exponent;
 
-				if (!SetTextureFromMap(InputName, Map, Factor))
+				if (SetTextureFromMap(InputName, Map, 1.f))
 				{
-					SetScalarFromMap(InputName, Map, Factor);
+					bIsPhong |= MaterialInstanceNode->AddStaticSwitchParameterValue(TEXT("bHasShininessMap"), true);
+				}
+				else
+				{
+					bIsPhong |= SetScalarFromMap(InputName, Map, Shininess);
 				}
 
-				bHasInput = true;
-				bIsPhong = true;
+				bHasInput |= bIsPhong;
 			}
 
+			// Normal
 			if (Material->fbx.normal_map.has_value)
 			{
-				if (SetTextureFromMap(Phong::Parameters::Normal, Material->fbx.normal_map, Material->fbx.normal_map.value_real))
+				if (SetTextureFromMap(Phong::Parameters::Normal, Material->fbx.normal_map, 1.f))
 				{
 					bHasInput = true;
 				}
-				else
+				else if (SetTextureFromMap(Phong::Parameters::Normal, Material->fbx.bump, 1.f))
 				{
-					if (SetTextureFromMap(Phong::Parameters::Normal, Material->fbx.bump, Material->fbx.bump.value_real))
-					{
-						MaterialInstanceNode->AddStaticSwitchParameterValue(TEXT("bHasBump"), true);
-						bHasInput = true;
-					}
+					MaterialInstanceNode->AddStaticSwitchParameterValue(TEXT("bHasBump"), true);
+					bHasInput = true;
 				}
 			}
 
-			if (!bIsPhong)
-			{
-				ParentMaterial = TEXT("LambertSurfaceMaterial");
-			}
+			MaterialInstanceNode->AddStaticSwitchParameterValue(TEXT("bIsPhong"), bIsPhong);
 
 			// Opacity
 			// Connect only if transparency is either a texture or different from 0.f
 			// Check fbx.transparency first...
-			if (Material->fbx.transparency_color.has_value || Material->fbx.transparency_factor.has_value)
+			if (Material->fbx.transparency_color.has_value)
 			{
 				FName InputName = Phong::Parameters::Opacity;
-				const ufbx_material_map& Map = Material->fbx.transparency_color;
-				const ufbx_material_map& FactorMap = Material->fbx.transparency_factor;
-				const float Factor = FactorMap.value_real;
-
-				// #ufbx_todo: note - 'transparency' texture in not inverted when setting as opacity, this is same behavior as was done for FBX SDK import
-
-				// The texture is hooked to the OpacityMask when transparency is with a texture
-				bool bIsMask = Map.value_int == 0;
+				const ufbx_material_map& ColorMap = Material->fbx.transparency_color;
+				const float Factor = GET_VALLUE_FROM_FACTOR_MAP(Material->fbx.transparency_factor);
 
-				if (bIsMask)
+				if (ColorMap.texture_enabled && ColorMap.texture)
 				{
-					InputName = Phong::Parameters::OpacityMask;								
-				}
+					FString HasOpacity{ TEXT("bHasOpacity") };
+					EBlendMode BlendMode = EBlendMode::BLEND_Translucent;
 
-				if (SetTextureFromMap(InputName, Map, Factor))
-				{
-					if (bIsMask)
+					// The texture must be hooked to the OpacityMask when transparency_color.value_int is equal to 0
+					if (ColorMap.value_int == 0)
 					{
-						ParentMaterial = TEXT("PhongSurfaceMaskMaterial");
-					}
-					else
-					{
-						ParentMaterial = TEXT("PhongSurfaceTranslucentMaterial");
+						InputName = Phong::Parameters::OpacityMask;
+						BlendMode = EBlendMode::BLEND_Masked;
+						MaterialInstanceNode->AddStaticSwitchParameterValue(TEXT("bHasOpacityMaskMap"), true);
+						HasOpacity += TEXT("Mask");
 					}
 
-					bHasInput = true;
+					HasOpacity += TEXT("Map");
+					MaterialInstanceNode->AddStaticSwitchParameterValue(HasOpacity, true);
+					MaterialInstanceNode->SetCustomBlendMode(BlendMode);
+					
+					bHasInput |= SetTextureFromMap(InputName, ColorMap, Factor);
 				}
 				else
 				{
-					const float OpacityScalar = 1.f - Factor*Material->fbx.transparency_color.value_real;
-					if (OpacityScalar+UE_KINDA_SMALL_NUMBER < 1.f)
+					const float TransparencyScalar = Factor * ColorMap.value_real;
+					const float OpacityScalar = 1.f - FMath::Clamp(TransparencyScalar, 0.f, 1.f);
+					if (OpacityScalar < 1.f)
 					{
-						ParentMaterial = TEXT("PhongSurfaceTranslucentMaterial");
-						MaterialInstanceNode->AddScalarParameterValue(Phong::Parameters::Opacity.ToString(), OpacityScalar);
-						bHasInput = true;
+						bHasInput |= MaterialInstanceNode->AddScalarParameterValue(InputName.ToString(), OpacityScalar);
 					}
 				}
 			}
@@ -391,13 +390,8 @@ namespace UE::Interchange::Private
 					Parser.ConvertProperty(Prop, MaterialInstanceNode);
 				}
 			}
-
-			if (bHasUVTransform)
-			{
-				ParentMaterial += "WithUVTransform";
-			}
-
-			MaterialInstanceNode->SetCustomParent(FString(TEXT("/InterchangeAssets/Materials"))/ParentMaterial + TEXT(".") + ParentMaterial);
 		}
 	}
 }
+
+#undef LOCTEXT_NAMESPACE
