Fixed compilation error with NonUnity InterchangeWorker

#rnx

diff --git a/Engine/Plugins/Interchange/Runtime/Source/Parsers/Fbx/Private/FbxMaterial.cpp b/Engine/Plugins/Interchange/Runtime/Source/Parsers/Fbx/Private/FbxMaterial.cpp
--- a/Engine/Plugins/Interchange/Runtime/Source/Parsers/Fbx/Private/FbxMaterial.cpp
+++ b/Engine/Plugins/Interchange/Runtime/Source/Parsers/Fbx/Private/FbxMaterial.cpp
@@ -23,23 +23,6 @@ namespace UE
 	{
 		namespace Private
 		{
-			// same enum as the one in EngineTypes, but to avoid compiling against Engine (it has some compile failures on Mac)
-			enum EBlendMode : int
-			{
-				BLEND_Opaque,
-				BLEND_Masked,
-				BLEND_Translucent,
-				BLEND_Additive,
-				BLEND_Modulate,
-				BLEND_AlphaComposite,
-				BLEND_AlphaHoldout,
-				BLEND_TranslucentColoredTransmittance, /*Substrate only */
-				BLEND_MAX,
-				// Renamed blend modes. These blend modes are remapped onto legacy ones and kept hidden for not confusing users in legacy mode, while allowing to use the new blend mode names into code.
-				BLEND_TranslucentGreyTransmittance = BLEND_Translucent, /*Substrate only */
-				BLEND_ColoredTransmittanceOnly = BLEND_Modulate, /*Substrate only */
-			};
-
 			UInterchangeMaterialInstanceNode* FFbxMaterial::CreateMaterialInstanceNode(UInterchangeBaseNodeContainer& NodeContainer, const FString& NodeUID, const FString& NodeName)
 			{
 				UInterchangeMaterialInstanceNode* MaterialInstanceNode = NewObject<UInterchangeMaterialInstanceNode>(&NodeContainer);
diff --git a/Engine/Plugins/Interchange/Runtime/Source/Parsers/Fbx/Private/FbxMaterial.h b/Engine/Plugins/Interchange/Runtime/Source/Parsers/Fbx/Private/FbxMaterial.h
--- a/Engine/Plugins/Interchange/Runtime/Source/Parsers/Fbx/Private/FbxMaterial.h
+++ b/Engine/Plugins/Interchange/Runtime/Source/Parsers/Fbx/Private/FbxMaterial.h
@@ -21,6 +21,23 @@ namespace UE
 		{
 			class FFbxParser;
 
+			// same enum as the one in EngineTypes, required because InterchangeWorker does not compile with Engine
+			enum EBlendMode : int
+			{
+				BLEND_Opaque,
+				BLEND_Masked,
+				BLEND_Translucent,
+				BLEND_Additive,
+				BLEND_Modulate,
+				BLEND_AlphaComposite,
+				BLEND_AlphaHoldout,
+				BLEND_TranslucentColoredTransmittance, /*Substrate only */
+				BLEND_MAX,
+				// Renamed blend modes. These blend modes are remapped onto legacy ones and kept hidden for not confusing users in legacy mode, while allowing to use the new blend mode names into code.
+				BLEND_TranslucentGreyTransmittance = BLEND_Translucent, /*Substrate only */
+				BLEND_ColoredTransmittanceOnly = BLEND_Modulate, /*Substrate only */
+			};
+
 			class FFbxMaterial
 			{
 			public:
diff --git a/Engine/Plugins/Interchange/Runtime/Source/Parsers/Fbx/Private/Ufbx/UfbxMaterial.cpp b/Engine/Plugins/Interchange/Runtime/Source/Parsers/Fbx/Private/Ufbx/UfbxMaterial.cpp
--- a/Engine/Plugins/Interchange/Runtime/Source/Parsers/Fbx/Private/Ufbx/UfbxMaterial.cpp
+++ b/Engine/Plugins/Interchange/Runtime/Source/Parsers/Fbx/Private/Ufbx/UfbxMaterial.cpp
@@ -5,6 +5,8 @@
 
 #include "UfbxConvert.h"
 
+#include "FbxMaterial.h"
+
 #include "InterchangeMaterialDefinitions.h"
 #include "InterchangeMaterialInstanceNode.h"
 #include "InterchangeTexture2DNode.h"
