[Interchange] Fixing missing animation sequence import for LevelSequence on StaticMesh imported with asset import pipeline.

#rnx

diff --git a/Engine/Plugins/Interchange/Runtime/Source/Pipelines/Private/InterchangeGenericAnimationPipeline.cpp b/Engine/Plugins/Interchange/Runtime/Source/Pipelines/Private/InterchangeGenericAnimationPipeline.cpp
--- a/Engine/Plugins/Interchange/Runtime/Source/Pipelines/Private/InterchangeGenericAnimationPipeline.cpp
+++ b/Engine/Plugins/Interchange/Runtime/Source/Pipelines/Private/InterchangeGenericAnimationPipeline.cpp
@@ -414,6 +414,13 @@ void UInterchangeGenericAnimationPipeline::ExecutePipeline(UInterchangeBaseNodeC
 				}
 				return bSomeTrackNodeUsedBySkeletalMesh;
 			};
+
+		bool bAllowSceneRootAsJoint = true;
+		if (const UInterchangeSourceNode* SourceNode = UInterchangeSourceNode::GetUniqueInstance(BaseNodeContainer))
+		{
+			SourceNode->GetCustomAllowSceneRootAsJoint(bAllowSceneRootAsJoint);
+		}
+
 		//Support rigid mesh animation animation data  (UAnimSequence for rigid mesh)
 		for (UInterchangeAnimationTrackSetNode* TrackSetNode : TrackSetNodes)
 		{
@@ -471,6 +478,14 @@ void UInterchangeGenericAnimationPipeline::ExecutePipeline(UInterchangeBaseNodeC
 										{
 											if (const UInterchangeSceneNode* ParentNode = Cast<UInterchangeSceneNode>(BaseNodeContainer->GetNode(ParentUid)))
 											{
+												bool bIsSceneRoot = false;
+												ParentNode->GetCustomIsSceneRoot(bIsSceneRoot);
+												if (!bAllowSceneRootAsJoint && bIsSceneRoot)
+												{
+													//Do not include scene root node in Skeleton.
+													break;
+												}
+
 												if(ParentNode->IsSpecializedTypeContains(UE::Interchange::FSceneNodeStaticData::GetJointSpecializeTypeString()))
 												{
 													SkeletonRootUid = ParentUid;
@@ -478,6 +493,10 @@ void UInterchangeGenericAnimationPipeline::ExecutePipeline(UInterchangeBaseNodeC
 												LastSceneNode = ParentUid;
 												ParentUid = ParentNode->GetParentUid();
 											}
+											else
+											{
+												break;
+											}
 										}
 										if (SkeletonRootUid.IsEmpty())
 										{
