[Interchange] Removing unintended LOD specialization from named LOD Groups.

#rnx

diff --git a/Engine/Plugins/Interchange/Runtime/Source/Pipelines/Private/InterchangeGenericStaticMeshPipeline.cpp b/Engine/Plugins/Interchange/Runtime/Source/Pipelines/Private/InterchangeGenericStaticMeshPipeline.cpp
--- a/Engine/Plugins/Interchange/Runtime/Source/Pipelines/Private/InterchangeGenericStaticMeshPipeline.cpp
+++ b/Engine/Plugins/Interchange/Runtime/Source/Pipelines/Private/InterchangeGenericStaticMeshPipeline.cpp
@@ -719,6 +719,17 @@ UInterchangeStaticMeshFactoryNode* UInterchangeGenericMeshPipeline::CreateStatic
 				{
 					StaticMeshFactoryNode->RegisterAttribute<FString>(UE::Interchange::FAttributeKey(FString(TEXT("LODGroupSceneNodeUid"))), BaseMeshUid);
 				}
+				else
+				{
+					bool bLODGroup_WithoutLODSpecialization = false;
+					if (BaseMeshSceneNode->GetAttribute<bool>(TEXT("LODGroup_WithoutLODSpecialization"), bLODGroup_WithoutLODSpecialization))
+					{
+						if (bLODGroup_WithoutLODSpecialization)
+						{
+							StaticMeshFactoryNode->RegisterAttribute<FString>(UE::Interchange::FAttributeKey(FString(TEXT("LODGroupSceneNodeUid"))), BaseMeshUid);
+						}
+					}
+				}
 			}
 		};
 
diff --git a/Engine/Plugins/Interchange/Runtime/Source/Pipelines/Private/InterchangePipelineMeshesUtilities.cpp b/Engine/Plugins/Interchange/Runtime/Source/Pipelines/Private/InterchangePipelineMeshesUtilities.cpp
--- a/Engine/Plugins/Interchange/Runtime/Source/Pipelines/Private/InterchangePipelineMeshesUtilities.cpp
+++ b/Engine/Plugins/Interchange/Runtime/Source/Pipelines/Private/InterchangePipelineMeshesUtilities.cpp
@@ -113,7 +113,7 @@ namespace UE::Private::InterchangeMeshPipeline
 				const UInterchangeSceneNode* SceneNodeConst = Cast<UInterchangeSceneNode>(BaseNodeContainer->GetNode(LodPrefixNodePair.Key));
 				if (UInterchangeSceneNode* SceneNode = const_cast<UInterchangeSceneNode*>(SceneNodeConst))
 				{
-					SceneNode->AddSpecializedType(UE::Interchange::FSceneNodeStaticData::GetLodGroupSpecializeTypeString());
+					SceneNode->RegisterAttribute<bool>(UE::Interchange::FAttributeKey(FString(TEXT("LODGroup_WithoutLODSpecialization"))), true);
 				}
 			}
 		}
