Fixed editor hanging when level reimporting a file containing skeletal meshes
- RootCause: Race condition between a skeletal mesh async build and saving the level instance actor's reference world which trigger another async build
- Solution : Clean-up the code by performing the save of the level only where and when necessary

#rnx

diff --git a/Engine/Plugins/Interchange/Runtime/Source/Import/Private/Scene/InterchangeLevelInstanceActorFactory.cpp b/Engine/Plugins/Interchange/Runtime/Source/Import/Private/Scene/InterchangeLevelInstanceActorFactory.cpp
--- a/Engine/Plugins/Interchange/Runtime/Source/Import/Private/Scene/InterchangeLevelInstanceActorFactory.cpp
+++ b/Engine/Plugins/Interchange/Runtime/Source/Import/Private/Scene/InterchangeLevelInstanceActorFactory.cpp
@@ -125,15 +125,6 @@ UObject* UInterchangeLevelInstanceActorFactory::ImportSceneObject_GameThread(con
 		return nullptr;
 	}
 
-	//Reference world must be saved if we want to update the level instance actor
-	if (UInterchangeEditorUtilitiesBase* EditorUtilities = UInterchangeManager::GetInterchangeManager().GetEditorUtilities())
-	{
-		if (!EditorUtilities->SaveAsset(ReferenceWorld))
-		{
-			UE_LOG(LogInterchangeImport, Warning, TEXT("UInterchangeLevelInstanceActorFactory::ImportSceneObject_GameThread: Cannot save the level instance actor (%s) referenced world (%s)"), *FactoryNode->GetDisplayLabel(), *ReferenceWorld->GetName());
-		}
-	}
-
 	UWorld* TargetWorld = CreateSceneObjectsParams.Level ? CreateSceneObjectsParams.Level->GetWorld() : GWorld->GetCurrentLevel()->GetWorld();
 	ensure (TargetWorld);
 
@@ -149,15 +140,6 @@ UObject* UInterchangeLevelInstanceActorFactory::ImportSceneObject_GameThread(con
 
 	UClass* LevelInstanceClass = bIsPackedActor ? APackedLevelActor::StaticClass() : ALevelInstance::StaticClass();
 
-	//Reference world must be saved if we want to update the level instance actor
-	if (UInterchangeEditorUtilitiesBase* EditorUtilities = UInterchangeManager::GetInterchangeManager().GetEditorUtilities())
-	{
-		if (!EditorUtilities->SaveAsset(ReferenceWorld))
-		{
-			UE_LOG(LogInterchangeImport, Warning, TEXT("UInterchangeLevelInstanceActorFactory::ImportSceneObject_GameThread: Cannot save the level instance actor (%s) referenced world (%s)"), *FactoryNode->GetDisplayLabel(), *ReferenceWorld->GetName());
-		}
-	}
-
 	ALevelInstance* LevelInstanceActor = nullptr;
 	//Packed level actor need a blueprint between the reference level and the actor
 	if (bIsPackedActor)
@@ -180,23 +162,33 @@ UObject* UInterchangeLevelInstanceActorFactory::ImportSceneObject_GameThread(con
 		const FString BPAssetName = FPackedLevelActorBuilder::GetPackedBPPrefix() + ReferenceWorldPath.GetAssetName();
 		const FString BPPackageName = FString::Format(TEXT("{0}/{1}"), { ReferenceWorldPackageDir , BPAssetName });
 	
-		UBlueprint* NewBP = Cast<UBlueprint>(FindExisitingAsset(BPPackageName, BPAssetName));
+		UBlueprint* PackedActorBP = Cast<UBlueprint>(FindExisitingAsset(BPPackageName, BPAssetName));
 		TSharedPtr<FPackedLevelActorBuilder> PackedLevelActorBuilder = FPackedLevelActorBuilder::CreateDefaultBuilder();
 
-		if (!NewBP)
+		if (!PackedActorBP)
 		{
+			// #interchange_level_instance_revisit
+			//Reference world must be saved if we want to create the packed actor
+			if (UInterchangeEditorUtilitiesBase* EditorUtilities = UInterchangeManager::GetInterchangeManager().GetEditorUtilities())
+			{
+				if (!EditorUtilities->SaveAsset(ReferenceWorld))
+				{
+					UE_LOG(LogInterchangeImport, Warning, TEXT("UInterchangeLevelInstanceActorFactory::ImportSceneObject_GameThread: Cannot save the packed actor's (%s) referenced world (%s)"), *FactoryNode->GetDisplayLabel(), *ReferenceWorld->GetName());
+				}
+			}
+
 			constexpr bool bCompileTrue = true;
 			const FString BPAssetPath = FString::Format(TEXT("{0}.{1}"), { BPPackageName , BPAssetName });
-			NewBP = PackedLevelActorBuilder->CreatePackedLevelActorBlueprint(TSoftObjectPtr<UBlueprint>(FSoftObjectPath(BPAssetPath)), ReferenceWorld, bCompileTrue);
+			PackedActorBP = PackedLevelActorBuilder->CreatePackedLevelActorBlueprint(TSoftObjectPtr<UBlueprint>(FSoftObjectPath(BPAssetPath)), ReferenceWorld, bCompileTrue);
 		}
 
-		if (ensure(NewBP))
+		if (ensure(PackedActorBP))
 		{
 			TSharedPtr<UE::Interchange::FImportAsyncHelper, ESPMode::ThreadSafe> AsyncHelper = UInterchangeManager::GetInterchangeManager().GetAsyncHelper(CreateSceneObjectsParams.AsyncHelperUid);
 			if (AsyncHelper.IsValid())
 			{
 				constexpr int32 SourceIndexZero = 0;
-				NewBP->SetInternalFlags(EInternalObjectFlags::Async);
+				PackedActorBP->SetInternalFlags(EInternalObjectFlags::Async);
 				FImportAsyncHelper::FImportedObjectInfo& AssetInfo = AsyncHelper->AddDefaultImportedAssetGetRef(SourceIndexZero);
 
 				//We need a valid factory node to have a valid DefaultImportedAsset entry. Its mandatory to register all created asset
@@ -204,10 +196,10 @@ UObject* UInterchangeLevelInstanceActorFactory::ImportSceneObject_GameThread(con
 				UInterchangeBaseNodeContainer* MutableNodeContainer = const_cast<UInterchangeBaseNodeContainer*>(CreateSceneObjectsParams.NodeContainer);
 				UInterchangeFactoryBaseNode* BpFactoryNode = NewObject<UInterchangeFactoryBaseNode>(MutableNodeContainer);
 				FString UniqueId = FGuid::NewGuid().ToString(EGuidFormats::Base36Encoded);
-				MutableNodeContainer->SetupNode(BpFactoryNode, UniqueId, NewBP->GetName(), EInterchangeNodeContainerType::FactoryData);
-				BpFactoryNode->SetCustomReferenceObject(NewBP);
+				MutableNodeContainer->SetupNode(BpFactoryNode, UniqueId, PackedActorBP->GetName(), EInterchangeNodeContainerType::FactoryData);
+				BpFactoryNode->SetCustomReferenceObject(PackedActorBP);
 
-				AssetInfo.ImportedObject = NewBP;
+				AssetInfo.ImportedObject = PackedActorBP;
 				AssetInfo.Factory = this;
 				AssetInfo.FactoryNode = BpFactoryNode;
 				AssetInfo.bIsReimport = false;
@@ -216,7 +208,7 @@ UObject* UInterchangeLevelInstanceActorFactory::ImportSceneObject_GameThread(con
 			LevelInstanceActor = Cast<ALevelInstance>(CreateSceneObjectsParams.ReimportObject);
 			if(!LevelInstanceActor)
 			{
-				LevelInstanceActor = TargetWorld->SpawnActor<APackedLevelActor>(NewBP->GeneratedClass, SpawnParameters);
+				LevelInstanceActor = TargetWorld->SpawnActor<APackedLevelActor>(PackedActorBP->GeneratedClass, SpawnParameters);
 			}
 		}
 
