Interchange: Fix crash when failing to find a material referenced by an UInterchangeMaterialReferenceFactoryNode

diff --git a/Engine/Plugins/Interchange/Runtime/Source/Import/Private/Material/InterchangeMaterialFactory.cpp b/Engine/Plugins/Interchange/Runtime/Source/Import/Private/Material/InterchangeMaterialFactory.cpp
--- a/Engine/Plugins/Interchange/Runtime/Source/Import/Private/Material/InterchangeMaterialFactory.cpp
+++ b/Engine/Plugins/Interchange/Runtime/Source/Import/Private/Material/InterchangeMaterialFactory.cpp
@@ -895,7 +895,7 @@ UInterchangeFactoryBase::FImportAssetResult UInterchangeMaterialFactory::BeginIm
 		Message->SourceAssetName = Arguments.SourceData->GetFilename();
 		Message->DestinationAssetName = Arguments.AssetName;
 		Message->AssetType = GetFactoryClass();
-		Message->Text = FText::Format(LOCTEXT("MatFactory_CouldNotCreateMat", "Could not create Material asset {0}. Reason: {1}"), FText::FromString(Arguments.AssetName), Info);
+		Message->Text = FText::Format(LOCTEXT("MatFactory_CouldNotCreateMat", "Failed to handle Material {0}. Reason: {1}"), FText::FromString(Arguments.AssetName), Info);
 		bSkipImport = true;
 		ImportAssetResult.bIsFactorySkipAsset = true;
 	};
@@ -932,6 +932,15 @@ UInterchangeFactoryBase::FImportAssetResult UInterchangeMaterialFactory::BeginIm
 		{
 			ExistingAsset = ReferenceObject.TryLoad();
 		}
+
+		if (!ExistingAsset && MaterialFactoryNode->IsA<UInterchangeMaterialReferenceFactoryNode>())
+		{
+			CouldNotCreateMaterialLog(FText::Format(
+				LOCTEXT("MatFactory_CouldNotCreateMat_CannotFindReferencedMaterial", "Could not find referenced material '{0}'"),
+				FText::FromString(ReferenceObject.ToString())
+			));
+			return ImportAssetResult;
+		}
 	}
 
 	bool bCanImportMaterial = true;
