Interchange: OBJ Translator fix Parsing files that have only objects without a group.  Add Default Group Label when there is no Group

diff --git a/Engine/Plugins/Interchange/Runtime/Source/Import/Private/Mesh/InterchangeOBJTranslator.cpp b/Engine/Plugins/Interchange/Runtime/Source/Import/Private/Mesh/InterchangeOBJTranslator.cpp
--- a/Engine/Plugins/Interchange/Runtime/Source/Import/Private/Mesh/InterchangeOBJTranslator.cpp
+++ b/Engine/Plugins/Interchange/Runtime/Source/Import/Private/Mesh/InterchangeOBJTranslator.cpp
@@ -154,7 +154,7 @@ struct FObjData
 	TMap<FString, FMaterialData> Materials;
 
 	/** This is the current group name being parsed */
-	FString CurrentGroup;
+	FString CurrentGroup = TEXT("Default");
 
 	/** This is the current material name being parsed */
 	FString MaterialBeingDefined;
@@ -761,6 +761,10 @@ namespace ObjParserUtils
 			INTERCHANGE_OBJ_TRANSLATOR_LOG_ONCE(LogTemp, Warning, TEXT("Interchange Obj translator: Unexpected extra arguments on o keyword"));
 		}
 
+		// Use Object Name as Current Group and if there is no Group underneath 
+		// all Data will be in that default group
+		ObjData.CurrentGroup = ObjectName;
+
 		//We do not do anything yet with the object name.
 		return true;
 	}
