MaterialX:
- Fix UDIMs by searching for all potential texture candidates in the folder instead of replacing the tag by the first index
- Fix wrong shadergraph with the <heighttonormal> node, use the connected image node for the NormalFromHeightMap expression instead of the current node

diff --git a/Engine/Plugins/Interchange/Runtime/Source/Import/Private/MaterialX/MaterialXUtils/MaterialXSurfaceShaderAbstract.cpp b/Engine/Plugins/Interchange/Runtime/Source/Import/Private/MaterialX/MaterialXUtils/MaterialXSurfaceShaderAbstract.cpp
--- a/Engine/Plugins/Interchange/Runtime/Source/Import/Private/MaterialX/MaterialXUtils/MaterialXSurfaceShaderAbstract.cpp
+++ b/Engine/Plugins/Interchange/Runtime/Source/Import/Private/MaterialX/MaterialXUtils/MaterialXSurfaceShaderAbstract.cpp
@@ -1070,31 +1070,25 @@ void FMaterialXSurfaceShaderAbstract::ConnectHeightToNormalInputToOutput(const F
 		if(mx::NodePtr ConnectedNode = Input->getConnectedNode();
 		   ConnectedNode && ConnectedNode->getCategory() == mx::Category::Image)
 		{
-			//we need to copy the content of the image node to this node
-			Connect.UpstreamNode->copyContentFrom(ConnectedNode);
-
-			// the copy overwrite every attribute of the node, so we need to get them back, essentially the type and the renaming
-			// the output is always a vec3
-			Connect.UpstreamNode->setType(mx::Type::Vector3);
-
-			SetMatchingInputsNames(Connect.UpstreamNode);
-
-			mx::GraphElementPtr Graph = Connect.UpstreamNode->getParent()->asA<mx::GraphElement>();
-			Graph->removeNode(ConnectedNode->getName());
-
 			using namespace UE::Interchange::Materials::Standard::Nodes;
 
-			if(UInterchangeTextureNode* TextureNode = CreateTextureNode<UInterchangeTexture2DNode>(Connect.UpstreamNode))
+			if(UInterchangeTextureNode* TextureNode = CreateTextureNode<UInterchangeTexture2DNode>(ConnectedNode))
 			{
 				UInterchangeShaderNode* HeightMapNode = CreateShaderNode(Connect.UpstreamNode, Connect.UpstreamNode->getName().c_str(), NormalFromHeightMap::Name.ToString());
 				UInterchangeShaderPortsAPI::ConnectDefaultOuputToInput(Connect.ParentShaderNode, Connect.InputChannelName, HeightMapNode->GetUniqueID());
 
 				const FString TextureNodeName = Connect.UpstreamNode->getName().c_str() + FString{ "_texture" };
 				UInterchangeShaderNode* TextureShaderNode = CreateShaderNode(Connect.UpstreamNode, TextureNodeName, TextureObject::Name.ToString());
-				TextureShaderNode->AddStringAttribute(UInterchangeShaderPortsAPI::MakeInputValueKey(TextureObject::Inputs::Texture.ToString()), TextureNode->GetUniqueID());
+				const FString TextureObjectInputKey = UInterchangeShaderPortsAPI::MakeInputParameterKey(TextureObject::Inputs::Texture.ToString());
+				TextureShaderNode->AddStringAttribute(TextureObjectInputKey, TextureNode->GetUniqueID());
 				UInterchangeShaderPortsAPI::ConnectDefaultOuputToInput(HeightMapNode, NormalFromHeightMap::Inputs::HeightMap.ToString(), TextureShaderNode->GetUniqueID());
 
-				AddAttributeFromValueOrInterface(Connect.UpstreamNode->getInput("scale"), NormalFromHeightMap::Inputs::Intensity.ToString(), HeightMapNode);
+				// The default value in the material function is 8, but the default value for the node is 1
+				if (!AddAttributeFromValueOrInterface(Connect.UpstreamNode->getInput("scale"), NormalFromHeightMap::Inputs::Intensity.ToString(), HeightMapNode))
+				{
+					const FString IntensityInputKey = UInterchangeShaderPortsAPI::MakeInputValueKey(NormalFromHeightMap::Inputs::Intensity.ToString());
+					HeightMapNode->AddFloatAttribute(IntensityInputKey, 1.f);
+				}
 			}
 			else
 			{
@@ -1558,21 +1552,7 @@ FString FMaterialXSurfaceShaderAbstract::GetInputName(MaterialX::InputPtr Input)
 
 FString FMaterialXSurfaceShaderAbstract::GetFilePrefix(MaterialX::ElementPtr Element) const
 {
-	FString FilePrefix;
-
-	if(Element)
-	{
-		if(Element->hasFilePrefix())
-		{
-			return FString(Element->getFilePrefix().c_str());
-		}
-		else
-		{
-			return GetFilePrefix(Element->getParent());
-		}
-	}
-
-	return FilePrefix;
+	return UTF8_TO_TCHAR(Element->getActiveFilePrefix().c_str());
 }
 
 FLinearColor FMaterialXSurfaceShaderAbstract::GetVector(MaterialX::InputPtr Input) const
diff --git a/Engine/Plugins/Interchange/Runtime/Source/Import/Public/MaterialX/MaterialXUtils/MaterialXSurfaceShaderAbstract.h b/Engine/Plugins/Interchange/Runtime/Source/Import/Public/MaterialX/MaterialXUtils/MaterialXSurfaceShaderAbstract.h
--- a/Engine/Plugins/Interchange/Runtime/Source/Import/Public/MaterialX/MaterialXUtils/MaterialXSurfaceShaderAbstract.h
+++ b/Engine/Plugins/Interchange/Runtime/Source/Import/Public/MaterialX/MaterialXUtils/MaterialXSurfaceShaderAbstract.h
@@ -406,9 +406,25 @@ protected:
 				else
 				{
 					Filepath = InputFile->hasValue() ? InputFile->getValueString().c_str() : InputFile->getInterfaceInput()->getValueString().c_str();
-					Filepath.ReplaceInline(TEXT(".<UDIM>."), TEXT(".1001."));
 					const FString FilePrefix = GetFilePrefix(InputFile);
 					Filepath = FPaths::Combine(FilePrefix, Filepath);
+					
+					// replace <UDIM> by a wildcard instead of just 1001, because the UDIM can only be for example one single texture like blabla.1003.png, which end up in a non VT
+					if (Filepath.ReplaceInline(TEXT("<UDIM>"), TEXT("*")))
+					{
+						TArray<FString> FileNames;
+						// we want to preserve the path given by the FilePrefix
+						FString SubPath = FPaths::GetPath(Filepath);
+						IFileManager::Get().FindFiles(FileNames, *FPaths::ConvertRelativePathToFull(FPaths::GetPath(Node->getActiveSourceUri().c_str()), Filepath), true, false);
+						if (!FileNames.IsEmpty())
+						{
+							Filepath = FPaths::Combine(SubPath, FileNames[0]);
+						}
+						else
+						{
+							MTLX_LOG("MaterialXUDIM", "Couldn't find any UDIM textures in the folder {0}", *SubPath);
+						}
+					}
 				}
 
 				const FString Filename = FPaths::GetCleanFilename(Filepath);
