Innterchange:
- Allow Material Instance to have ShaderGraphNode as parent during import

diff --git a/Engine/Plugins/Interchange/Runtime/Source/FactoryNodes/Private/InterchangeMaterialFactoryNode.cpp b/Engine/Plugins/Interchange/Runtime/Source/FactoryNodes/Private/InterchangeMaterialFactoryNode.cpp
--- a/Engine/Plugins/Interchange/Runtime/Source/FactoryNodes/Private/InterchangeMaterialFactoryNode.cpp
+++ b/Engine/Plugins/Interchange/Runtime/Source/FactoryNodes/Private/InterchangeMaterialFactoryNode.cpp
@@ -496,6 +496,16 @@ bool UInterchangeMaterialInstanceFactoryNode::SetCustomBlendMode(TEnumAsByte<EBl
 	IMPLEMENT_NODE_ATTRIBUTE_SETTER_NODELEGATE(BlendMode, TEnumAsByte<EBlendMode>);
 }
 
+bool UInterchangeMaterialInstanceFactoryNode::GetCustomTwoSided(bool& AttributeValue) const
+{
+	IMPLEMENT_NODE_ATTRIBUTE_GETTER(TwoSided, bool);
+}
+
+bool UInterchangeMaterialInstanceFactoryNode::SetCustomTwoSided(bool AttributeValue)
+{
+	IMPLEMENT_NODE_ATTRIBUTE_SETTER_NODELEGATE(TwoSided, bool);
+}
+
 FString UInterchangeMaterialReferenceFactoryNode::GetTypeName() const
 {
 	const FString TypeName = TEXT("MaterialReferenceFactoryNode");
diff --git a/Engine/Plugins/Interchange/Runtime/Source/FactoryNodes/Public/InterchangeMaterialFactoryNode.h b/Engine/Plugins/Interchange/Runtime/Source/FactoryNodes/Public/InterchangeMaterialFactoryNode.h
--- a/Engine/Plugins/Interchange/Runtime/Source/FactoryNodes/Public/InterchangeMaterialFactoryNode.h
+++ b/Engine/Plugins/Interchange/Runtime/Source/FactoryNodes/Public/InterchangeMaterialFactoryNode.h
@@ -334,11 +334,17 @@ public:
 	UFUNCTION(BlueprintCallable, Category = "Interchange | Node | MaterialInstanceFactory")
 	INTERCHANGEFACTORYNODES_API bool SetCustomBlendMode(TEnumAsByte<EBlendMode> AttributeValue, bool bAddApplyDelegate = true);
 
+	UFUNCTION(BlueprintCallable, Category = "Interchange | Node | MaterialInstanceFactory")
+	INTERCHANGEFACTORYNODES_API bool GetCustomTwoSided(bool& AttributeValue) const;
+
+	UFUNCTION(BlueprintCallable, Category = "Interchange | Node | MaterialInstanceFactory")
+	INTERCHANGEFACTORYNODES_API bool SetCustomTwoSided(bool AttributeValue);
+
 private:
 	IMPLEMENT_NODE_ATTRIBUTE_KEY(InstanceClassName)
 	IMPLEMENT_NODE_ATTRIBUTE_KEY(Parent)
 	IMPLEMENT_NODE_ATTRIBUTE_KEY(BlendMode)
-
+	IMPLEMENT_NODE_ATTRIBUTE_KEY(TwoSided)
 };
 
 /**
diff --git a/Engine/Plugins/Interchange/Runtime/Source/Import/Private/Material/InterchangeMaterialFactory.cpp b/Engine/Plugins/Interchange/Runtime/Source/Import/Private/Material/InterchangeMaterialFactory.cpp
--- a/Engine/Plugins/Interchange/Runtime/Source/Import/Private/Material/InterchangeMaterialFactory.cpp
+++ b/Engine/Plugins/Interchange/Runtime/Source/Import/Private/Material/InterchangeMaterialFactory.cpp
@@ -435,10 +435,10 @@ namespace UE::Interchange::MaterialFactory::Internal
 			FSoftObjectPath AssetPath(PathString);
 			OutAsset = Cast<AssetType>(AssetPath.TryLoad());
 		}
-		else if (const AssetFactoryType* TextureFactoryNode = Cast<AssetFactoryType>(NodeContainer.GetNode(PathString)))
+		else if (const AssetFactoryType* FactoryNode = Cast<AssetFactoryType>(NodeContainer.GetNode(PathString)))
 		{
 			FSoftObjectPath ReferenceObject;
-			TextureFactoryNode->GetCustomReferenceObject(ReferenceObject);
+			FactoryNode->GetCustomReferenceObject(ReferenceObject);
 			OutAsset = Cast<AssetType>(ReferenceObject.TryLoad());
 		}
 
@@ -816,6 +816,20 @@ namespace UE::Interchange::MaterialFactory::Internal
 			MaterialInstance.BasePropertyOverrides.bOverride_BlendMode = false;
 		}
 	}
+
+	void OverrideTwoSided(UMaterialInstance& MaterialInstance, const UInterchangeMaterialInstanceFactoryNode& FactoryNode)
+	{
+		if (bool bTwoSided; FactoryNode.GetCustomTwoSided(bTwoSided))
+		{
+			MaterialInstance.BasePropertyOverrides.bOverride_TwoSided = true;
+			MaterialInstance.BasePropertyOverrides.TwoSided = bTwoSided;
+			MaterialInstance.TwoSided = bTwoSided;
+		}
+		else
+		{
+			MaterialInstance.BasePropertyOverrides.bOverride_TwoSided = false;
+		}
+	}
 }
 
 namespace UE::Interchange::Materials::HashUtils
@@ -951,6 +965,20 @@ UInterchangeFactoryBase::FImportAssetResult UInterchangeMaterialFactory::BeginIm
 		}
 	}
 
+	auto GetParentMaterial = [&Arguments](const UInterchangeMaterialInstanceFactoryNode* MaterialInstanceFactoryNode) -> UMaterialInterface*
+	{
+		using namespace UE::Interchange::MaterialFactory::Internal;
+		UMaterialInterface* ParentMaterialAsset = nullptr;
+
+		FString PathOrUid;
+		if (!MaterialInstanceFactoryNode->GetCustomParent(PathOrUid))
+		{
+			return nullptr;
+		}
+
+		return TryLoadAsset<UMaterialInterface, UInterchangeFactoryBaseNode>(*Arguments.NodeContainer, PathOrUid);
+	};
+
 	bool bCanImportMaterial = true;
 	MaterialFactoryNode->GetCustomIsMaterialImportEnabled(bCanImportMaterial);
 	// create a new material or overwrite existing asset, if possible
@@ -960,11 +988,9 @@ UInterchangeFactoryBase::FImportAssetResult UInterchangeMaterialFactory::BeginIm
 		{
 			if (const UInterchangeMaterialInstanceFactoryNode* MaterialInstanceFactoryNode = Cast<UInterchangeMaterialInstanceFactoryNode>(MaterialFactoryNode))
 			{
-				FString ParentPath;
-				if (MaterialInstanceFactoryNode->GetCustomParent(ParentPath))
+				if(UMaterialInterface* ParentMaterial = GetParentMaterial(MaterialInstanceFactoryNode))				
 				{
-					FSoftObjectPath ParentMaterial(ParentPath);
-					Material = UMaterialInstanceDynamic::Create(Cast<UMaterialInterface>(ParentMaterial.TryLoad()), Arguments.Parent, *Arguments.AssetName);
+					Material = UMaterialInstanceDynamic::Create(ParentMaterial, Arguments.Parent, *Arguments.AssetName);
 				}
 			}
 		}
@@ -1049,21 +1075,14 @@ UInterchangeFactoryBase::FImportAssetResult UInterchangeMaterialFactory::BeginIm
 #if WITH_EDITOR
 				if (UMaterialInstanceConstant* MaterialInstanceConstant = Cast<UMaterialInstanceConstant>(MaterialInstance))
 				{
-					using namespace UE::Interchange::MaterialFactory::Internal;
-
-					FString ParentPath;
-					FSoftObjectPath ParentMaterialPath;
-					if (MaterialInstanceFactoryNode->GetCustomParent(ParentPath))
+					if (UMaterialInterface* ParentMaterialAsset = GetParentMaterial(MaterialInstanceFactoryNode))
 					{
-						if (UMaterialInterface* ParentMaterialAsset = TryLoadAsset<UMaterialInterface, UInterchangeFactoryBaseNode>(*Arguments.NodeContainer, ParentPath))
-						{
-							MaterialInstanceConstant->SetParentEditorOnly(ParentMaterialAsset);
-							MaterialInstanceConstant->PostEditChange();
-						}
-						else
-						{
-							UE_LOG(LogInterchangeImport, Error, TEXT("No parent material was found."))
-						}
+						MaterialInstanceConstant->SetParentEditorOnly(ParentMaterialAsset);
+						MaterialInstanceConstant->PostEditChange();
+					}
+					else
+					{
+						UE_LOG(LogInterchangeImport, Error, TEXT("No parent material was found."))
 					}
 				}
 #endif
@@ -2023,6 +2042,7 @@ void UInterchangeMaterialFactory::SetupMaterialInstance(UMaterialInstance& Mater
 	}
 
 	OverrideBlendMode(MaterialInstance, FactoryNode);
+	OverrideTwoSided(MaterialInstance, FactoryNode);
 }
 
 void UInterchangeMaterialFactory::SetupReimportedMaterialInstance(UMaterialInstance& MaterialInstance, const UInterchangeBaseNodeContainer& NodeContainer, const UInterchangeMaterialInstanceFactoryNode& FactoryNode, const UInterchangeMaterialInstanceFactoryNode& PreviousFactoryNode)
@@ -2127,25 +2147,34 @@ void UInterchangeMaterialFactory::SetupReimportedMaterialInstance(UMaterialInsta
 		}
 	}
 
-	bool bUpdateParameter = ReimportStrategy == EReimportStrategyFlags::ApplyPipelineProperties;
-	// If previous factory node has a blend mode, check whether it is the same as the material instance's one
-	if (TEnumAsByte<EBlendMode> BlendMode; PreviousFactoryNode.GetCustomBlendMode(BlendMode))
+	auto OverrideProperty = [&MaterialInstance, &FactoryNode, &PreviousFactoryNode, ReimportStrategy]
+	(auto GetCustomProperty, auto OverrideValue, auto OverrideFlag, auto Override)
 	{
-		if (MaterialInstance.BasePropertyOverrides.bOverride_BlendMode)
+		using T = std::remove_cv_t<std::remove_reference_t<decltype(OverrideValue)>>;
+
+		bool bUpdateParameter = ReimportStrategy == EReimportStrategyFlags::ApplyPipelineProperties;
+		// If previous factory node has a property (blend mode, twosided, etc.), check whether it is the same as the material instance's one
+		if (T Property{}; (PreviousFactoryNode.*GetCustomProperty)(Property))
 		{
-			bUpdateParameter |= MaterialInstance.BasePropertyOverrides.BlendMode == BlendMode;
+			if (OverrideFlag)
+			{
+				bUpdateParameter |= (T(OverrideValue) == Property);
+			}
+		}
+		// If there was no property override in previous import and things have not changed, update is allowed
+		else if (!OverrideFlag)
+		{
+			bUpdateParameter |= true;
 		}
-	}
-	// If there was no BlendMode override in previous import and things have not changed, update is allowed
-	else if(!MaterialInstance.BasePropertyOverrides.bOverride_BlendMode)
-	{
-		bUpdateParameter |= true;
-	}
 
-	if(bUpdateParameter)
-	{
-		OverrideBlendMode(MaterialInstance, FactoryNode);
-	}
+		if (bUpdateParameter)
+		{
+			Override(MaterialInstance, FactoryNode);
+		}
+	};
+
+	OverrideProperty(&UInterchangeMaterialInstanceFactoryNode::GetCustomBlendMode, MaterialInstance.BasePropertyOverrides.BlendMode, MaterialInstance.BasePropertyOverrides.bOverride_BlendMode, OverrideBlendMode);
+	OverrideProperty(&UInterchangeMaterialInstanceFactoryNode::GetCustomTwoSided, bool(MaterialInstance.BasePropertyOverrides.TwoSided), MaterialInstance.BasePropertyOverrides.bOverride_TwoSided, OverrideTwoSided);
 }
 
 UClass* UInterchangeMaterialFunctionFactory::GetFactoryClass() const
diff --git a/Engine/Plugins/Interchange/Runtime/Source/Nodes/Private/InterchangeMaterialInstanceNode.cpp b/Engine/Plugins/Interchange/Runtime/Source/Nodes/Private/InterchangeMaterialInstanceNode.cpp
--- a/Engine/Plugins/Interchange/Runtime/Source/Nodes/Private/InterchangeMaterialInstanceNode.cpp
+++ b/Engine/Plugins/Interchange/Runtime/Source/Nodes/Private/InterchangeMaterialInstanceNode.cpp
@@ -89,12 +89,22 @@ bool UInterchangeMaterialInstanceNode::GetStaticSwitchParameterValue(const FStri
 	return GetBooleanAttribute(UInterchangeShaderPortsAPI::MakeInputValueKey(ParameterName), AttributeValue);
 }
 
-bool UInterchangeMaterialInstanceNode::GetCustomBlendMode(int& AttributeValue) const
+bool UInterchangeMaterialInstanceNode::GetCustomBlendMode(int32& AttributeValue) const
 {
-	IMPLEMENT_NODE_ATTRIBUTE_GETTER(BlendMode, int);
+	IMPLEMENT_NODE_ATTRIBUTE_GETTER(BlendMode, int32);
 }
 
-bool UInterchangeMaterialInstanceNode::SetCustomBlendMode(int AttributeValue)
+bool UInterchangeMaterialInstanceNode::SetCustomBlendMode(int32 AttributeValue)
 {
-	IMPLEMENT_NODE_ATTRIBUTE_SETTER_NODELEGATE(BlendMode, int);
+	IMPLEMENT_NODE_ATTRIBUTE_SETTER_NODELEGATE(BlendMode, int32);
 }
+
+bool UInterchangeMaterialInstanceNode::GetCustomTwoSided(bool& AttributeValue) const
+{
+	IMPLEMENT_NODE_ATTRIBUTE_GETTER(TwoSided, bool);
+}
+
+bool UInterchangeMaterialInstanceNode::SetCustomTwoSided(bool AttributeValue)
+{
+	IMPLEMENT_NODE_ATTRIBUTE_SETTER_NODELEGATE(TwoSided, bool);
+}
\ No newline at end of file
diff --git a/Engine/Plugins/Interchange/Runtime/Source/Nodes/Public/InterchangeMaterialInstanceNode.h b/Engine/Plugins/Interchange/Runtime/Source/Nodes/Public/InterchangeMaterialInstanceNode.h
--- a/Engine/Plugins/Interchange/Runtime/Source/Nodes/Public/InterchangeMaterialInstanceNode.h
+++ b/Engine/Plugins/Interchange/Runtime/Source/Nodes/Public/InterchangeMaterialInstanceNode.h
@@ -25,15 +25,16 @@ public:
 
 
 	/**
-	 * Build and return a UID name for a shader node.
+	 * Build and return a UID name for a material instance node.
 	 */
 	static UE_API FString MakeNodeUid(const FStringView NodeName, const FStringView ParentNodeUid);
 
 	/**
-     * Creates a new UInterchangeShaderGraphNode and adds it to NodeContainer as a translated node.
+     * Creates a new UInterchangeMaterialInstanceNode and adds it to NodeContainer as a translated node.
      */
 	static UE_API UInterchangeMaterialInstanceNode* Create(UInterchangeBaseNodeContainer* NodeContainer, const FStringView NodeName, const FStringView ParentNodeUid);
 
+	/** Path to an existing Material or the UID of a ShaderGraphNode*/
 	UFUNCTION(BlueprintCallable, Category = "Interchange | Node | Material Instance")
 	UE_API bool SetCustomParent(const FString& AttributeValue) const;
 
@@ -65,16 +66,22 @@ public:
 	UE_API bool GetStaticSwitchParameterValue(const FString& ParameterName, bool& AttributeValue) const;
 
 	UFUNCTION(BlueprintCallable, Category = "Interchange | Node | Material")
-	UE_API bool GetCustomBlendMode(int& AttributeValue) const;
+	UE_API bool GetCustomBlendMode(int32& AttributeValue) const;
 
 	UFUNCTION(BlueprintCallable, Category = "Interchange | Node | Material")
-	UE_API bool SetCustomBlendMode(int AttributeValue);
+	UE_API bool SetCustomBlendMode(int32 AttributeValue);
+
+	UFUNCTION(BlueprintCallable, Category = "Interchange | Node | Material")
+	UE_API bool GetCustomTwoSided(bool& AttributeValue) const;
+
+	UFUNCTION(BlueprintCallable, Category = "Interchange | Node | Material")
+	UE_API bool SetCustomTwoSided(bool AttributeValue);
 
 private:
 
 	IMPLEMENT_NODE_ATTRIBUTE_KEY(Parent);
 	IMPLEMENT_NODE_ATTRIBUTE_KEY(BlendMode);
-
+	IMPLEMENT_NODE_ATTRIBUTE_KEY(TwoSided);
 };
 
 #undef UE_API
diff --git a/Engine/Plugins/Interchange/Runtime/Source/Pipelines/Private/InterchangeGenericMaterialPipeline.cpp b/Engine/Plugins/Interchange/Runtime/Source/Pipelines/Private/InterchangeGenericMaterialPipeline.cpp
--- a/Engine/Plugins/Interchange/Runtime/Source/Pipelines/Private/InterchangeGenericMaterialPipeline.cpp
+++ b/Engine/Plugins/Interchange/Runtime/Source/Pipelines/Private/InterchangeGenericMaterialPipeline.cpp
@@ -48,6 +48,7 @@
 #include "Materials/MaterialExpressionTangent.h"
 #include "Materials/MaterialExpressionTextureCoordinate.h"
 #include "Materials/MaterialExpressionTextureObject.h"
+#include "Materials/MaterialExpressionTextureObjectParameter.h"
 #include "Materials/MaterialExpressionTextureSample.h"
 #include "Materials/MaterialExpressionTextureSampleParameter2D.h"
 #include "Materials/MaterialExpressionTextureSampleParameter2DArray.h"
@@ -970,11 +971,19 @@ void UInterchangeGenericMaterialPipeline::ExecutePipeline(UInterchangeBaseNodeCo
 
 	for (UInterchangeMaterialInstanceNode* MaterialNode : MaterialInstanceNodes)
 	{
-		FString ParentPath;
-		if (!MaterialNode->GetCustomParent(ParentPath) || ParentPath.IsEmpty())
-		{
-			continue;
-		}
+		FString ParentPathOrMaterialUid;
+        if (!MaterialNode->GetCustomParent(ParentPathOrMaterialUid) || ParentPathOrMaterialUid.IsEmpty())
+        {
+            continue;
+        }
+
+        // the parent path can either be the path to an existing Material or the UID of a ShaderGraphNode
+        const FString MaterialFactoryNodeUid = UInterchangeFactoryBaseNode::BuildFactoryNodeUid(ParentPathOrMaterialUid);
+		const bool bIsMaterialFactoryUid = Cast<UInterchangeMaterialFactoryNode>(BaseNodeContainer->GetFactoryNode(MaterialFactoryNodeUid)) != nullptr;
+        if (bIsMaterialFactoryUid)
+        {
+            ParentPathOrMaterialUid = MaterialFactoryNodeUid;
+        }
 
 		UInterchangeMaterialInstanceFactoryNode* MaterialFactoryNode = nullptr;
 		FString DisplayLabel = MaterialNode->GetDisplayLabel();
@@ -1007,14 +1016,23 @@ void UInterchangeGenericMaterialPipeline::ExecutePipeline(UInterchangeBaseNodeCo
 		// Set MaterialFactoryNode's display label to MaterialNode's uniqueID
 		// to reconcile mesh's slot names and material assets
 		MaterialFactoryNode->SetDisplayLabel(MaterialNode->GetAssetName());
-		MaterialFactoryNode->SetCustomParent(ParentPath);
+		MaterialFactoryNode->SetCustomParent(ParentPathOrMaterialUid);
+		if (bIsMaterialFactoryUid)
+		{
+			MaterialFactoryNode->AddFactoryDependencyUid(ParentPathOrMaterialUid);
+		}
 
 		const UClass* MaterialClass = (EditorUtilities && EditorUtilities->IsRuntimeOrPIE()) ? UMaterialInstanceDynamic::StaticClass() : UMaterialInstanceConstant::StaticClass();
 		MaterialFactoryNode->SetCustomInstanceClassName(MaterialClass->GetPathName());
 
-		if (EBlendMode BlendMode; MaterialNode->GetCustomBlendMode(reinterpret_cast<int&>(BlendMode)))
+		if (int BlendMode; MaterialNode->GetCustomBlendMode(BlendMode))
 		{
-			MaterialFactoryNode->SetCustomBlendMode(BlendMode);
+			MaterialFactoryNode->SetCustomBlendMode(EBlendMode{ BlendMode });
+		}
+
+		if (bool bTwoSided; MaterialNode->GetCustomTwoSided(bTwoSided))
+		{
+			MaterialFactoryNode->SetCustomTwoSided(bTwoSided);
 		}
 
 		UE::Interchange::InterchangeGenericMaterialPipeline::Private::AddSpecularProfileToFactoryNode(MaterialNode, MaterialFactoryNode, BaseNodeContainer);
@@ -2237,14 +2255,15 @@ void UInterchangeGenericMaterialPipeline::HandleTextureObjectNode(const UInterch
 	FString TextureUid = GetTextureUidAttributeFromShaderNode(ShaderNode, Inputs::Texture, bIsAParameter);
 	FString ExpressionClassName;
 	FString TextureFactoryUid;
+	FString TextureObjectClassName = bIsAParameter ? UMaterialExpressionTextureObjectParameter::StaticClass()->GetName() : UMaterialExpressionTextureObject::StaticClass()->GetName();
 
 	if(const UInterchangeTextureNode* TextureNode = Cast<const UInterchangeTextureNode>(BaseNodeContainer->GetNode(TextureUid)))
 	{
-		HandleTextureNode(TextureNode, MaterialFactoryNode, TextureObjectFactoryNode, UMaterialExpressionTextureObject::StaticClass()->GetName(), bIsAParameter);
+		HandleTextureNode(TextureNode, MaterialFactoryNode, TextureObjectFactoryNode, TextureObjectClassName, bIsAParameter);
 	}
 	else
 	{
-		TextureObjectFactoryNode->SetCustomExpressionClassName(UMaterialExpressionTextureObject::StaticClass()->GetName());
+		TextureObjectFactoryNode->SetCustomExpressionClassName(TextureObjectClassName);
 		TextureObjectFactoryNode->AddStringAttribute(CreateInputKey(Inputs::Texture.ToString(), bIsAParameter), TextureFactoryUid);
 	}
 }
