Interchange : Remove invalid ensure in UInterchangeFbxTranslator::GetTexturePayloadData.

diff --git a/Engine/Plugins/Interchange/Runtime/Source/Import/Private/Fbx/InterchangeFbxTranslator.cpp b/Engine/Plugins/Interchange/Runtime/Source/Import/Private/Fbx/InterchangeFbxTranslator.cpp
--- a/Engine/Plugins/Interchange/Runtime/Source/Import/Private/Fbx/InterchangeFbxTranslator.cpp
+++ b/Engine/Plugins/Interchange/Runtime/Source/Import/Private/Fbx/InterchangeFbxTranslator.cpp
@@ -543,8 +543,10 @@ TOptional<UE::Interchange::FImportImage> UInterchangeFbxTranslator::GetTexturePa
 
 	UE::Interchange::Private::FScopedTranslator ScopedTranslator(PayLoadKey, Results, AnalyticsHelper);
 	const IInterchangeTexturePayloadInterface* TextureTranslator = ScopedTranslator.GetPayLoadInterface<IInterchangeTexturePayloadInterface>();
-	if (!ensure(TextureTranslator))
+	if (!TextureTranslator)
 	{
+		// TextureTranslator can can nullptr in case where
+		// UInterchangeTranslatorBase::CanImportSourceData implementation return false
 		return TOptional<UE::Interchange::FImportImage>();
 	}
 	AlternateTexturePath = PayLoadKey;
