Interchange: Add support for directional light source angle.

diff --git a/Engine/Plugins/Interchange/Extensions/OpenUSD/Source/Import/Private/InterchangeUsdTranslator.cpp b/Engine/Plugins/Interchange/Extensions/OpenUSD/Source/Import/Private/InterchangeUsdTranslator.cpp
--- a/Engine/Plugins/Interchange/Extensions/OpenUSD/Source/Import/Private/InterchangeUsdTranslator.cpp
+++ b/Engine/Plugins/Interchange/Extensions/OpenUSD/Source/Import/Private/InterchangeUsdTranslator.cpp
@@ -3451,8 +3451,9 @@ namespace UE::InterchangeUsdTranslator::Private
 			Intensity = UsdToUnreal::ConvertLightIntensityAttr(Intensity, Exposure);
 			LightNode->SetCustomIntensity(Intensity);
 
-			// LightSourceAngle currently not supported by UInterchangeDirectionalLightNode
-			// float Angle = UsdUtils::GetAttributeValue<float>(Prim, TEXTVIEW("inputs:angle"));
+			static const FString SourceAngleToken = UsdToUnreal::ConvertToken(pxr::UsdLuxTokens->inputsAngle);
+			float Angle = UsdUtils::GetAttributeValue<float>(Prim, SourceAngleToken);
+			LightNode->SetCustomLightSourceAngle(Angle);
 		}
 		else if (Prim.IsA(TEXT("SphereLight")))
 		{
diff --git a/Engine/Plugins/Interchange/Runtime/Source/FactoryNodes/Private/InterchangeLightFactoryNode.cpp b/Engine/Plugins/Interchange/Runtime/Source/FactoryNodes/Private/InterchangeLightFactoryNode.cpp
--- a/Engine/Plugins/Interchange/Runtime/Source/FactoryNodes/Private/InterchangeLightFactoryNode.cpp
+++ b/Engine/Plugins/Interchange/Runtime/Source/FactoryNodes/Private/InterchangeLightFactoryNode.cpp
@@ -2,6 +2,7 @@
 
 #include "InterchangeLightFactoryNode.h"
 
+#include "Components/DirectionalLightComponent.h"
 #include "Components/RectLightComponent.h"
 #include "Components/SkyLightComponent.h"
 #include "Components/SpotLightComponent.h"
@@ -70,6 +71,32 @@ void UInterchangeBaseLightFactoryNode::OnRestoreAllCustomAttributeDelegates()
 	REFILL_CUSTOM_ATTRIBUTE_APPLY_DELEGATE_EXPLICIT(Macro_CustombUseTemperatureKey.Key, GetCustomUseTemperature, SetCustomUseTemperature, bool);
 }
 
+bool UInterchangeDirectionalLightFactoryNode::GetCustomLightSourceAngle(float& AttributeValue) const
+{
+	IMPLEMENT_NODE_ATTRIBUTE_GETTER(LightSourceAngle, float)
+}
+
+bool UInterchangeDirectionalLightFactoryNode::SetCustomLightSourceAngle(float AttributeValue, bool bAddApplyDelegate)
+{
+	IMPLEMENT_NODE_ATTRIBUTE_SETTER(UInterchangeDirectionalLightFactoryNode, LightSourceAngle, float, UDirectionalLightComponent)
+}
+
+void UInterchangeDirectionalLightFactoryNode::CopyWithObject(const UInterchangeFactoryBaseNode* SourceNode, UObject* Object)
+{
+	Super::CopyWithObject(SourceNode, Object);
+
+	if (const UInterchangeDirectionalLightFactoryNode* LightFactoryNode = Cast<UInterchangeDirectionalLightFactoryNode>(SourceNode))
+	{
+		COPY_NODE_DELEGATES(LightFactoryNode, LightSourceAngle, float, UDirectionalLightComponent)
+	}
+}
+
+void UInterchangeDirectionalLightFactoryNode::OnRestoreAllCustomAttributeDelegates()
+{
+	Super::OnRestoreAllCustomAttributeDelegates();
+	REFILL_CUSTOM_ATTRIBUTE_APPLY_DELEGATE(LightSourceAngle, float);
+}
+
 bool UInterchangeLightFactoryNode::GetCustomIntensityUnits(ELightUnits& AttributeValue) const
 {
 	IMPLEMENT_NODE_ATTRIBUTE_GETTER(IntensityUnits, ELightUnits)
diff --git a/Engine/Plugins/Interchange/Runtime/Source/FactoryNodes/Public/InterchangeLightFactoryNode.h b/Engine/Plugins/Interchange/Runtime/Source/FactoryNodes/Public/InterchangeLightFactoryNode.h
--- a/Engine/Plugins/Interchange/Runtime/Source/FactoryNodes/Public/InterchangeLightFactoryNode.h
+++ b/Engine/Plugins/Interchange/Runtime/Source/FactoryNodes/Public/InterchangeLightFactoryNode.h
@@ -56,6 +56,20 @@ UCLASS(MinimalAPI, BlueprintType)
 class UInterchangeDirectionalLightFactoryNode : public UInterchangeBaseLightFactoryNode
 {
 	GENERATED_BODY()
+
+public:
+	UFUNCTION(BlueprintCallable, Category = "Interchange | Node | DirectionalLightFactory")
+	UE_API bool GetCustomLightSourceAngle(float& AttributeValue) const;
+
+	UFUNCTION(BlueprintCallable, Category = "Interchange | Node | DirectionalLightFactory")
+	UE_API bool SetCustomLightSourceAngle(float AttributeValue, bool bAddApplyDelegate = true);
+
+	UE_API virtual void CopyWithObject(const UInterchangeFactoryBaseNode* SourceNode, UObject* Object) override;
+
+	virtual void OnRestoreAllCustomAttributeDelegates() override;
+
+private:
+	IMPLEMENT_NODE_ATTRIBUTE_KEY(LightSourceAngle)
 };
 
 UCLASS(MinimalAPI, BlueprintType)
diff --git a/Engine/Plugins/Interchange/Runtime/Source/Nodes/Private/InterchangeLightNode.cpp b/Engine/Plugins/Interchange/Runtime/Source/Nodes/Private/InterchangeLightNode.cpp
--- a/Engine/Plugins/Interchange/Runtime/Source/Nodes/Private/InterchangeLightNode.cpp
+++ b/Engine/Plugins/Interchange/Runtime/Source/Nodes/Private/InterchangeLightNode.cpp
@@ -200,6 +200,16 @@ FString UInterchangeDirectionalLightNode::GetTypeName() const
 	return TEXT("DirectionalLightNode");
 }
 
+bool UInterchangeDirectionalLightNode::GetCustomLightSourceAngle(float& AttributeValue) const
+{
+	IMPLEMENT_NODE_ATTRIBUTE_GETTER(LightSourceAngle, float)
+}
+
+bool UInterchangeDirectionalLightNode::SetCustomLightSourceAngle(float AttributeValue)
+{
+	IMPLEMENT_NODE_ATTRIBUTE_SETTER_NODELEGATE(LightSourceAngle, float)
+}
+
 FString UInterchangeSkyLightNode::GetTypeName() const
 {
 	return TEXT("SkyLightNode");
diff --git a/Engine/Plugins/Interchange/Runtime/Source/Nodes/Public/InterchangeLightNode.h b/Engine/Plugins/Interchange/Runtime/Source/Nodes/Public/InterchangeLightNode.h
--- a/Engine/Plugins/Interchange/Runtime/Source/Nodes/Public/InterchangeLightNode.h
+++ b/Engine/Plugins/Interchange/Runtime/Source/Nodes/Public/InterchangeLightNode.h
@@ -215,6 +215,15 @@ class UInterchangeDirectionalLightNode : public UInterchangeBaseLightNode
 public:
 
 	UE_API virtual FString GetTypeName() const override;
+
+	UFUNCTION(BlueprintCallable, Category = "Interchange | Node | DirectionalLight")
+	UE_API bool GetCustomLightSourceAngle(float& AttributeValue) const;
+
+	UFUNCTION(BlueprintCallable, Category = "Interchange | Node | DirectionalLight")
+	UE_API bool SetCustomLightSourceAngle(float AttributeValue);
+
+private:
+	IMPLEMENT_NODE_ATTRIBUTE_KEY(LightSourceAngle)
 };
 
 UCLASS(MinimalAPI, BlueprintType)
diff --git a/Engine/Plugins/Interchange/Runtime/Source/Pipelines/Private/InterchangeGenericScenesPipeline.cpp b/Engine/Plugins/Interchange/Runtime/Source/Pipelines/Private/InterchangeGenericScenesPipeline.cpp
--- a/Engine/Plugins/Interchange/Runtime/Source/Pipelines/Private/InterchangeGenericScenesPipeline.cpp
+++ b/Engine/Plugins/Interchange/Runtime/Source/Pipelines/Private/InterchangeGenericScenesPipeline.cpp
@@ -1129,6 +1129,15 @@ void UInterchangeGenericLevelPipeline::SetUpFactoryNode(UInterchangeActorFactory
 				}
 			}
 
+			// DirectionalLight
+			if (const UInterchangeDirectionalLightNode* DirectionalLightNode = Cast<UInterchangeDirectionalLightNode>(BaseLightNode))
+			{
+				if (UInterchangeDirectionalLightFactoryNode* DirectionalLightFactoryNode = Cast<UInterchangeDirectionalLightFactoryNode>(BaseLightFactoryNode))
+				{
+					COPY_FROM_TRANSLATED_TO_FACTORY(DirectionalLightNode, DirectionalLightFactoryNode, LightSourceAngle, float)
+				}
+			}
+
 			// Sky Light
 			if (const UInterchangeSkyLightNode* SkyLightNode = Cast<UInterchangeSkyLightNode>(BaseLightNode))
 			{
