[Interchange] ufbx fixes

Fixed creating animation track for bones without transform animation, Editor.Interchange.Import.Asset.FBX.TP_SK_Socket UE-318723
Calculated animation time per skeleton. Editor.Interchange.Import.Asset.FBX.TP_FBX_SK_Anim_Time_animated_time UE-318723
Fixed embedded textures. Editor.Python.EngineTest.Interchange.Import.Asset.FBX.test_textutebasecolor UE-349647
Fixed ensure with non-file(procedural) texmap.  EngineTest\Content\Tests\Interchange\assets\FBX\Animated panels.fbx

diff --git a/Engine/Plugins/Interchange/Runtime/Source/Parsers/Fbx/Private/Ufbx/UfbxAnimation.cpp b/Engine/Plugins/Interchange/Runtime/Source/Parsers/Fbx/Private/Ufbx/UfbxAnimation.cpp
--- a/Engine/Plugins/Interchange/Runtime/Source/Parsers/Fbx/Private/Ufbx/UfbxAnimation.cpp
+++ b/Engine/Plugins/Interchange/Runtime/Source/Parsers/Fbx/Private/Ufbx/UfbxAnimation.cpp
@@ -29,6 +29,33 @@ namespace UE::Interchange::Private
 			return Keys[KeyIndex].value;
 		};
 
+		inline FVector ConvertValue(const ufbx_vec3& V)
+		{
+			return Convert::ConvertVec3(V);
+		}
+
+		inline FQuat ConvertValue(const ufbx_quat& V)
+		{
+			return Convert::ConvertQuat(V);
+		}
+
+		// Get inverval where the values are actually animated
+		template<typename ValueList, typename Value>
+		void UpdateAnimationInterval(TInterval<double>& Interval, const ValueList& Keys, const bool& IsConstant, const Value& Reference)
+		{
+			if (!IsConstant 
+				|| (Keys.count > 0 
+					&& !ConvertValue(Keys[0].value).Equals(
+						ConvertValue(Reference), UE_SMALL_NUMBER)))
+			{
+				if (Keys.count > 1)
+				{
+					Interval.Include(Keys[0].time);
+					Interval.Include(Keys[Keys.count-1].time);
+				}
+			}
+		}
+
 		// inspired by UE::Interchange::Private::ImportCurve - Engine/Plugins/Interchange/Runtime/Source/Parsers/Fbx/Private/FbxAnimation.cpp
 		void ConvertCurve(const ufbx_anim_curve& SourceFloatCurves, const float ScaleValue, TArray<FInterchangeCurveKey>& DestinationFloatCurve)
 		{
@@ -489,7 +516,7 @@ namespace UE::Interchange::Private
 
 			FString DisplayString = Convert::ToUnrealString(Stack->name);
 
-			auto GetTrackNode = [&](FString SkeletonRootNodeUid)
+			auto GetTrackNode = [&](FString SkeletonRootNodeUid, TInterval<double> AnimInterval=TInterval<double>())
 			{
 				UInterchangeSkeletalAnimationTrackNode*& SkeletalAnimationTrackNode = TrackNodePerSkeletonRoot.FindOrAdd(SkeletonRootNodeUid);
 				if (!SkeletalAnimationTrackNode)
@@ -502,14 +529,29 @@ namespace UE::Interchange::Private
 					SkeletalAnimationTrackNode->SetCustomAnimationSampleRate(Parser.Scene->settings.frames_per_second);
 
 					SkeletalAnimationTrackNode->SetCustomSkeletonNodeUid(SkeletonRootNodeUid);
-				
-					SkeletalAnimationTrackNode->SetCustomAnimationStartTime(Stack->anim->time_begin);
-					SkeletalAnimationTrackNode->SetCustomAnimationStopTime(Stack->anim->time_end);
 
 					SkeletalAnimationTrackNode->SetCustomSourceTimelineAnimationStartTime(Stack->time_begin);
 					SkeletalAnimationTrackNode->SetCustomSourceTimelineAnimationStopTime(Stack->time_end);
 					NodeContainer.AddNode(SkeletalAnimationTrackNode);
 				}
+
+				double CustomAnimationStartTime;
+				double CustomAnimationStopTime;
+				if (SkeletalAnimationTrackNode->GetCustomAnimationStartTime(CustomAnimationStartTime))
+				{
+					AnimInterval.Include(CustomAnimationStartTime);
+				}
+				if (SkeletalAnimationTrackNode->GetCustomAnimationStopTime(CustomAnimationStopTime))
+				{
+					AnimInterval.Include(CustomAnimationStopTime);
+				}
+
+				if (AnimInterval.IsValid())
+				{
+					SkeletalAnimationTrackNode->SetCustomAnimationStartTime(AnimInterval.Min);
+					SkeletalAnimationTrackNode->SetCustomAnimationStopTime(AnimInterval.Max);
+				}
+
 				return SkeletalAnimationTrackNode;
 			};
 
@@ -527,19 +569,26 @@ namespace UE::Interchange::Private
 				{
 					FString NodeUid = Parser.GetNodeUid(*Node);
 
-					UInterchangeSkeletalAnimationTrackNode* SkeletalAnimationTrackNode = GetTrackNode(Scene.GetSkeletonRoot(Node));
+					TInterval<double> Interval;
+					AnimUtils::UpdateAnimationInterval(Interval, BakedNode.translation_keys, BakedNode.constant_translation, Node->local_transform.translation);
+					AnimUtils::UpdateAnimationInterval(Interval, BakedNode.rotation_keys, BakedNode.constant_rotation, Node->local_transform.rotation);
+					AnimUtils::UpdateAnimationInterval(Interval, BakedNode.scale_keys, BakedNode.constant_scale, Node->local_transform.scale);
+
+					if (Interval.IsValid())
+					{
+						UInterchangeSkeletalAnimationTrackNode* SkeletalAnimationTrackNode = GetTrackNode(Scene.GetSkeletonRoot(Node), Interval);
 
-					// #ufbx_todo: Need check animated parents - whole subtree has to be animated too?
-					// see UE::Interchange::Private::FFbxAnimation::AddSkeletalTransformAnimation
-					// or maybe there's an option in ufbx baking that already does that?
+						// #ufbx_todo: Need check animated parents - whole subtree has to be animated too?
+						// see UE::Interchange::Private::FFbxAnimation::AddSkeletalTransformAnimation
+						// or maybe there's an option in ufbx baking that already does that?
 
-					FString PayLoadKey = NodeUid + TEXT("_") + FString::FromInt(AnimStackIndex) + TEXT("_SkeletalAnimationPayloadKey");
+						FString PayLoadKey = NodeUid + TEXT("_") + FString::FromInt(AnimStackIndex) + TEXT("_SkeletalAnimationPayloadKey");
 
-					Parser.PayloadContexts.Add(PayLoadKey, FSkeletalAnimationPayloadContext{NodeUid, AnimStackIndex, BakedNodeIndex, BakedAnim});
+						Parser.PayloadContexts.Add(PayLoadKey, FSkeletalAnimationPayloadContext{NodeUid, AnimStackIndex, BakedNodeIndex, BakedAnim});
 
-					SkeletalAnimationTrackNode->SetAnimationPayloadKeyForSceneNodeUid(NodeUid, PayLoadKey, EInterchangeAnimationPayLoadType::BAKED);
+						SkeletalAnimationTrackNode->SetAnimationPayloadKeyForSceneNodeUid(NodeUid, PayLoadKey, EInterchangeAnimationPayLoadType::BAKED);
+					}
 				}
-
 			}
 
 			// Collect curves for other animated properties - morph, ..
diff --git a/Engine/Plugins/Interchange/Runtime/Source/Parsers/Fbx/Private/Ufbx/UfbxMaterial.cpp b/Engine/Plugins/Interchange/Runtime/Source/Parsers/Fbx/Private/Ufbx/UfbxMaterial.cpp
--- a/Engine/Plugins/Interchange/Runtime/Source/Parsers/Fbx/Private/Ufbx/UfbxMaterial.cpp
+++ b/Engine/Plugins/Interchange/Runtime/Source/Parsers/Fbx/Private/Ufbx/UfbxMaterial.cpp
@@ -104,17 +104,30 @@ namespace UE::Interchange::Private
 
 		const UInterchangeTexture2DNode* CreateTexture2DNode(const ufbx_texture& Texture)
 		{
+			if (Texture.type != UFBX_TEXTURE_FILE)
+			{
+				return nullptr;
+			}
 			const FString TextureFilename = FindTextureFile(Texture);
 			if (TextureFilename.IsEmpty() || !FPaths::FileExists(TextureFilename))
 			{
-				if (!GIsAutomationTesting)
+				// Texture might be embedded
+				if (Texture.content.data && Texture.content.size > 0)
 				{
-					UInterchangeResultTextureDisplay_TextureFileDoNotExist* Message = Parser.AddMessage<UInterchangeResultTextureDisplay_TextureFileDoNotExist>();
-					Message->TextureName = TextureFilename;
-					Message->MaterialName = MaterialInstance.GetDisplayLabel();
+#if WITH_ENGINE
+					Parser.PayloadContexts.Add(TextureFilename, FPayloadContext(FPayloadContext::Element, Texture.element_id));
+#endif
+				}
+				else 
+				{
+					if (!GIsAutomationTesting)
+					{
+						UInterchangeResultTextureDisplay_TextureFileDoNotExist* Message = Parser.AddMessage<UInterchangeResultTextureDisplay_TextureFileDoNotExist>();
+						Message->TextureName = TextureFilename;
+						Message->MaterialName = MaterialInstance.GetDisplayLabel();
+					}
+					return nullptr;
 				}
-
-				return nullptr;
 			}
 			return CreateTexture2DNode(TextureFilename);
 		}
diff --git a/Engine/Plugins/Interchange/Runtime/Source/Parsers/Fbx/Private/Ufbx/UfbxParser.cpp b/Engine/Plugins/Interchange/Runtime/Source/Parsers/Fbx/Private/Ufbx/UfbxParser.cpp
--- a/Engine/Plugins/Interchange/Runtime/Source/Parsers/Fbx/Private/Ufbx/UfbxParser.cpp
+++ b/Engine/Plugins/Interchange/Runtime/Source/Parsers/Fbx/Private/Ufbx/UfbxParser.cpp
@@ -88,6 +88,11 @@ FString FUfbxParser::GetNodeUid(const ufbx_node& Node) const
 	// #ufbx_todo: check that element_id is enough
 	// reference FBX parser's uses UE::Interchange::Private::FFbxHelper::GetFbxNodeHierarchyName to create NodeUid
 
+	if (FString* Found = NodeToUid.Find(&Node))
+	{
+		return *Found;
+	}
+
 	FString NodeUid = GetElementNameDeduplicated(Node.element);
 	
 	if (Node.parent)
@@ -98,9 +103,13 @@ FString FUfbxParser::GetNodeUid(const ufbx_node& Node) const
 		{
 			const UInterchangeSceneNode* ParentSceneNode = *ParentSceneNodeFound;
 	
-			return ParentSceneNode->GetUniqueID()+TEXT(".")+NodeUid;
+			NodeUid = ParentSceneNode->GetUniqueID()+TEXT(".")+NodeUid;
 		}
 	}
+
+	NodeToUid.Add(&Node, NodeUid);
+	UidToNode.Add(NodeUid, &Node);
+
 	return NodeUid;
 }
 
@@ -214,9 +223,12 @@ void FUfbxParser::SetResultContainer(UInterchangeResultsContainer* Result)
 
 void FUfbxParser::Reset()
 {
+	UidToNode.Reset();
+	NodeToUid.Reset();
 	NameToElement.Reset();
 	ElementToName.Reset();
 
+
 	ElementIdToSceneNode.Reset();
 
 	PayloadContexts.Reset();
diff --git a/Engine/Plugins/Interchange/Runtime/Source/Parsers/Fbx/Private/Ufbx/UfbxParser.h b/Engine/Plugins/Interchange/Runtime/Source/Parsers/Fbx/Private/Ufbx/UfbxParser.h
--- a/Engine/Plugins/Interchange/Runtime/Source/Parsers/Fbx/Private/Ufbx/UfbxParser.h
+++ b/Engine/Plugins/Interchange/Runtime/Source/Parsers/Fbx/Private/Ufbx/UfbxParser.h
@@ -288,5 +288,8 @@ namespace UE::Interchange::Private
 		mutable TMap<FString, const ufbx_element*> NameToElement;
 		mutable TMap<const ufbx_element*, FString> ElementToName;
 
+		mutable TMap<FString, const ufbx_node*> UidToNode;
+		mutable TMap<const ufbx_node*, FString> NodeToUid;
+
 	};
 }
