[Interchange] ufbx fixed labels of elements - nodes, meshes, camera, bones

- Editor.Python.EngineTest.Interchange.Import.Asset.FBX.test_name_clash and other naming tests
- Supported bKeepFbxNamespace

diff --git a/Engine/Plugins/Interchange/Runtime/Source/Parsers/Fbx/Private/Ufbx/UfbxAnimation.cpp b/Engine/Plugins/Interchange/Runtime/Source/Parsers/Fbx/Private/Ufbx/UfbxAnimation.cpp
--- a/Engine/Plugins/Interchange/Runtime/Source/Parsers/Fbx/Private/Ufbx/UfbxAnimation.cpp
+++ b/Engine/Plugins/Interchange/Runtime/Source/Parsers/Fbx/Private/Ufbx/UfbxAnimation.cpp
@@ -678,7 +678,7 @@ namespace UE::Interchange::Private
 							for (const ufbx_blend_keyframe& Keyframe : Channel->keyframes)
 							{
 								MorphPayloadContext.InbetweenFullWeights.Add(Keyframe.target_weight);
-								MorphPayloadContext.InbetweenCurveNames.Add(Parser.GetElementNameDeduplicated(Keyframe.shape->element));
+								MorphPayloadContext.InbetweenCurveNames.Add(Parser.GetElementName(Keyframe.shape->element));
 							}
 
 							Parser.PayloadContexts.Add(PayloadKey, MorphPayloadContext);
diff --git a/Engine/Plugins/Interchange/Runtime/Source/Parsers/Fbx/Private/Ufbx/UfbxConvert.cpp b/Engine/Plugins/Interchange/Runtime/Source/Parsers/Fbx/Private/Ufbx/UfbxConvert.cpp
--- a/Engine/Plugins/Interchange/Runtime/Source/Parsers/Fbx/Private/Ufbx/UfbxConvert.cpp
+++ b/Engine/Plugins/Interchange/Runtime/Source/Parsers/Fbx/Private/Ufbx/UfbxConvert.cpp
@@ -15,10 +15,4 @@ Convert::FMeshNameAndUid::FMeshNameAndUid(FUfbxParser& Parser, const ufbx_elemen
 {
 }
 
-Convert::FNodeNameAndUid::FNodeNameAndUid(FUfbxParser& Parser, UInterchangeSceneNode* ParentSceneNode, const ufbx_node& Node)
-	: Label(ParentSceneNode ?  Parser.GetElementName(Node.element) : TEXT("RootNode"))
-	, UniqueID(Parser.GetNodeUid(Node))
-{
-}
-
 }
diff --git a/Engine/Plugins/Interchange/Runtime/Source/Parsers/Fbx/Private/Ufbx/UfbxConvert.h b/Engine/Plugins/Interchange/Runtime/Source/Parsers/Fbx/Private/Ufbx/UfbxConvert.h
--- a/Engine/Plugins/Interchange/Runtime/Source/Parsers/Fbx/Private/Ufbx/UfbxConvert.h
+++ b/Engine/Plugins/Interchange/Runtime/Source/Parsers/Fbx/Private/Ufbx/UfbxConvert.h
@@ -30,16 +30,6 @@ namespace Convert
 		FString UniqueID;
 	};
 
-	struct FNodeNameAndUid
-	{
-		// #ufbx_todo: currently replicates FBX SDK's behavior to name root node "RootNode".
-		//		But the root node doesn't have actual name, in fact. Also it might be better to name this node after the file name?
-		FNodeNameAndUid(FUfbxParser& Parser, UInterchangeSceneNode* ParentSceneNode, const ufbx_node& Node);
-
-		FString Label;
-		FString UniqueID;
-	};
-
 	inline FVector2d ConvertVec2(const ufbx_vec2& V)
 	{
 		return FVector2d(V.x, V.y);
diff --git a/Engine/Plugins/Interchange/Runtime/Source/Parsers/Fbx/Private/Ufbx/UfbxMaterial.cpp b/Engine/Plugins/Interchange/Runtime/Source/Parsers/Fbx/Private/Ufbx/UfbxMaterial.cpp
--- a/Engine/Plugins/Interchange/Runtime/Source/Parsers/Fbx/Private/Ufbx/UfbxMaterial.cpp
+++ b/Engine/Plugins/Interchange/Runtime/Source/Parsers/Fbx/Private/Ufbx/UfbxMaterial.cpp
@@ -177,7 +177,7 @@ namespace UE::Interchange::Private
 				continue;
 			}
 
-			FString MaterialLabel = Parser.GetMaterialLabel(*Material).ToString();
+			FString MaterialLabel = Parser.GetMaterialLabel(*Material);
 			
 			FString MaterialUid = Parser.GetMaterialUid(*Material);
 			UInterchangeMaterialInstanceNode* MaterialInstanceNode = CreateMaterialInstanceNode(NodeContainer, MaterialUid, MaterialLabel);
diff --git a/Engine/Plugins/Interchange/Runtime/Source/Parsers/Fbx/Private/Ufbx/UfbxMesh.cpp b/Engine/Plugins/Interchange/Runtime/Source/Parsers/Fbx/Private/Ufbx/UfbxMesh.cpp
--- a/Engine/Plugins/Interchange/Runtime/Source/Parsers/Fbx/Private/Ufbx/UfbxMesh.cpp
+++ b/Engine/Plugins/Interchange/Runtime/Source/Parsers/Fbx/Private/Ufbx/UfbxMesh.cpp
@@ -276,14 +276,14 @@ namespace UE::Interchange::Private
 
 			for (ufbx_blend_deformer* Deformer : Mesh->blend_deformers)
 			{
-				const FString MorphTargetName = Parser.GetElementNameDeduplicated(Deformer->element);
+				const FString MorphTargetName = Parser.GetElementName(Deformer->element);
 				for (int32 ChannelIndex = 0; ChannelIndex < Deformer->channels.count; ++ChannelIndex)
 				{
 					const ufbx_blend_channel* Channel = Deformer->channels[ChannelIndex];
 
 					const int32 CurrentChannelMorphTargetCount = Channel->keyframes.count;
 
-					FString ChannelName = Parser.GetElementNameDeduplicated(Channel->element);
+					FString ChannelName = Parser.GetElementName(Channel->element);
 					// Maya adds the name of the MorphTarget and an underscore to the front of the channel name, so remove it
 					if (ChannelName.StartsWith(MorphTargetName) && (ChannelName.Compare(MorphTargetName, ESearchCase::IgnoreCase) != 0))
 					{
@@ -296,7 +296,7 @@ namespace UE::Interchange::Private
 						const ufbx_blend_shape* Shape = Keyframe.shape;
 
 						FString ShapeName = CurrentChannelMorphTargetCount > 1
-							? Parser.GetElementNameDeduplicated(Shape->element)
+							? Parser.GetElementName(Shape->element)
 							// Maya concatenates the number of the shape to the end of its name, so instead use the name of the channel
 							: ChannelName;
 
@@ -576,7 +576,7 @@ namespace UE::Interchange::Private
 					if (BoneNode && !SortedJoints.Find(BoneNode, BoneIndex))
 					{
 						BoneIndex = SortedJoints.Add(BoneNode);
-						FString JointNodeUniqueID = Parser.GetBoneNodeUid(*BoneNode);
+						FString JointNodeUniqueID = Parser.GetBoneNodeName(*BoneNode);
 						OutJointUniqueNames.Add(JointNodeUniqueID);
 					}
 
diff --git a/Engine/Plugins/Interchange/Runtime/Source/Parsers/Fbx/Private/Ufbx/UfbxParser.cpp b/Engine/Plugins/Interchange/Runtime/Source/Parsers/Fbx/Private/Ufbx/UfbxParser.cpp
--- a/Engine/Plugins/Interchange/Runtime/Source/Parsers/Fbx/Private/Ufbx/UfbxParser.cpp
+++ b/Engine/Plugins/Interchange/Runtime/Source/Parsers/Fbx/Private/Ufbx/UfbxParser.cpp
@@ -14,6 +14,7 @@
 #include "Mesh/InterchangeMeshPayload.h"
 #endif
 
+#include "InterchangeCameraNode.h"
 #include "Fbx/InterchangeFbxMessages.h"
 #include "InterchangeHelper.h"
 #include "InterchangeCommonAnimationPayload.h"
@@ -34,209 +35,38 @@ namespace UE::Interchange::Private
 
 FUfbxParser::FUfbxParser(TWeakObjectPtr<UInterchangeResultsContainer> InResultsContainer)
 	: ResultsContainer(InResultsContainer)
-{
-}
-
-FUfbxParser::~FUfbxParser()
-{
-	Reset();
-}
-
-FString FUfbxParser::GetElementNameRaw(const ufbx_element& Element) const
-{
-	return Convert::ToUnrealString(Element.name);
-}
-
-FString FUfbxParser::GetElementNameDeduplicated(const ufbx_element& Element, bool bIsJoint) const
-{
-	if (FString* Found = ElementToName.Find(&Element))
-	{
-		return *Found;
-	}
-
-	const FString BaseName = GetElementName(Element, bIsJoint);
-	FString Name = BaseName;
-	for (int32 Count = 0;;++Count)
+	, ElementNames(*this)
+	, NodeNames(*this)
+	, BoneNames(*this)
+	, MeshNames(*this)
+	, CameraNames(*this)
+	, MaterialNames(*this)
 	{
-		if (!NameToElement.Contains(Name))
-		{
-			break;
-		}
-		Name = BaseName + TEXT("_ncl_") + FString::FromInt(Count+1);
-	}
-	NameToElement.Add(Name, &Element);
-	ElementToName.Add(&Element, Name);
-	return Name;
-}
+		NodeNames.DefaultName = TEXT("Node");
+		NodeNames.bUseNCL = false;
+		NodeNames.bUseNodeRules = true;
 
-// inspired by UE::Interchange::Private::FFbxHelper::GetFbxObjectName
-FString FUfbxParser::GetElementName(const ufbx_element& Element, bool bIsJoint) const
-{
-	FString ElementName = FUfbxParser::GetElementNameRaw(Element);
-	UE::Interchange::SanitizeName(ElementName, bIsJoint);
-	return ElementName;
-}
+		BoneNames.DefaultName = TEXT("Node");
+		BoneNames.bIsJoint = true;
 
-// #ufbx_todo: see FFbxHelper::GetMeshName
-FString FUfbxParser::GetMeshLabel(const ufbx_element& Mesh) const
-{
-	return UTF8_TO_TCHAR(Mesh.name.data);
-}
+		MeshNames.DefaultName = TEXT("Mesh");
 
-FString FUfbxParser::GetNodeUid(const ufbx_node& Node) const
-{
-	// #ufbx_todo: check that element_id is enough
-	// reference FBX parser's uses UE::Interchange::Private::FFbxHelper::GetFbxNodeHierarchyName to create NodeUid
+		CameraNames.DefaultName = TEXT("Camera");
 
-	if (FString* Found = NodeToUid.Find(&Node))
-	{
-		return *Found;
-	}
+		MaterialNames.DefaultName = TEXT("Material");
 
-	FString NodeUid = GetElementNameDeduplicated(Node.element);
-	
-	if (Node.parent)
-	{
-		const UInterchangeSceneNode*const* ParentSceneNodeFound = ElementIdToSceneNode.Find(Node.parent->element_id);
-		// 
-		if (ensureMsgf(ParentSceneNodeFound, TEXT("Expected to have child node Uid generated after parent was processed first")))
-		{
-			const UInterchangeSceneNode* ParentSceneNode = *ParentSceneNodeFound;
-	
-			NodeUid = ParentSceneNode->GetUniqueID()+TEXT(".")+NodeUid;
-		}
 	}
 
-	NodeToUid.Add(&Node, NodeUid);
-	UidToNode.Add(NodeUid, &Node);
-
-	return NodeUid;
-}
-
-FString FUfbxParser::GetBoneNodeUid(const ufbx_node& Node) const
-{
-	// #ufbx_todo: need make sure not dup
-	return GetElementNameDeduplicated(Node.element, true);
-}
-
-FName FUfbxParser::GetMaterialLabel(const ufbx_material& Material) const
-{
-	return *GetElementName(Material.element);
-}
-
-FString FUfbxParser::GetMaterialUid(const ufbx_material& Material) const
-{
-	// Since materials names are not unique use its ufbx type_id and keep original name in the Uid for simpler debugging
-	return  FString::Printf(TEXT("\\Material\\%d_%s"), Material.element.typed_id, *GetElementNameRaw(Material.element));
-}
-
-FString FUfbxParser::GetMeshUid(const ufbx_element& Mesh) const
-{
-	return GetElementNameDeduplicated(Mesh);
-}
-
-void FUfbxParser::ConvertProperty(const ufbx_prop& Prop, UInterchangeBaseNode* SceneNode, const TOptional<FString>& PayloadKey)
+FUfbxParser::~FUfbxParser()
 {
-	if (Prop.flags & UFBX_PROP_FLAG_USER_DEFINED)
-	{
-		FString PropertyName = Convert::ToUnrealString(Prop.name);
-
-		switch (Prop.type)
-		{
-		case UFBX_PROP_UNKNOWN:
-			break;
-		case UFBX_PROP_BOOLEAN:
-			{
-				const bool PropertyValue = Prop.value_int != 0;
-				UInterchangeUserDefinedAttributesAPI::CreateUserDefinedAttribute(SceneNode, PropertyName, PropertyValue, PayloadKey);
-			}
-			break;
-		case UFBX_PROP_INTEGER:
-			{
-				const int32 PropertyValue = Prop.value_int;
-				UInterchangeUserDefinedAttributesAPI::CreateUserDefinedAttribute(SceneNode, PropertyName, PropertyValue, PayloadKey);
-			}
-			break;
-		case UFBX_PROP_NUMBER:
-			{
-				const float PropertyValue = Prop.value_real;
-				UInterchangeUserDefinedAttributesAPI::CreateUserDefinedAttribute(SceneNode, PropertyName, PropertyValue, PayloadKey);
-			}
-			break;
-		case UFBX_PROP_VECTOR:
-			{
-				const FVector4d PropertyValue = Convert::ConvertVec4(Prop.value_vec4);
-				UInterchangeUserDefinedAttributesAPI::CreateUserDefinedAttribute(SceneNode, PropertyName, PropertyValue, PayloadKey);
-			}
-			break;
-		case UFBX_PROP_COLOR:
-			{
-				const FVector PropertyValue = FVector(Convert::ConvertColor(Prop.value_vec4));
-				UInterchangeUserDefinedAttributesAPI::CreateUserDefinedAttribute(SceneNode, PropertyName, PropertyValue, PayloadKey);
-			}
-			break;
-		case UFBX_PROP_COLOR_WITH_ALPHA:
-			{
-				const FLinearColor PropertyValue = Convert::ConvertColor(Prop.value_vec4);
-				UInterchangeUserDefinedAttributesAPI::CreateUserDefinedAttribute(SceneNode, PropertyName, PropertyValue, PayloadKey);
-			}
-			break;
-		case UFBX_PROP_STRING:
-			{
-				const FString PropertyValue = Convert::ToUnrealString(Prop.value_str);
-				UInterchangeUserDefinedAttributesAPI::CreateUserDefinedAttribute(SceneNode, PropertyName, PropertyValue, PayloadKey);
-			}
-			break;
-		case UFBX_PROP_DATE_TIME:
-			{
-				const FDateTime PropertyValue(Prop.value_int);
-				UInterchangeUserDefinedAttributesAPI::CreateUserDefinedAttribute(SceneNode, PropertyName, PropertyValue, PayloadKey);
-			}
-			break;
-		// #ufbx_todo: support other useful custom attributes among these types
-		case UFBX_PROP_TRANSLATION:
-			break;
-		case UFBX_PROP_ROTATION:
-			break;
-		case UFBX_PROP_SCALING:
-			break;
-		case UFBX_PROP_DISTANCE:
-			break;
-		case UFBX_PROP_COMPOUND:
-			break;
-		case UFBX_PROP_BLOB:
-			break;
-		case UFBX_PROP_REFERENCE:
-			break;
-		case UFBX_PROP_TYPE_FORCE_32BIT:
-			break;
-		}
-	}
+	Reset();
 }
 
-
-
 void FUfbxParser::SetResultContainer(UInterchangeResultsContainer* Result)
 {
 	ResultsContainer = Result;
 }
 
-void FUfbxParser::Reset()
-{
-	UidToNode.Reset();
-	NodeToUid.Reset();
-	NameToElement.Reset();
-	ElementToName.Reset();
-
-
-	ElementIdToSceneNode.Reset();
-
-	PayloadContexts.Reset();
-
-	ufbx_free_scene(Scene);
-	Scene = nullptr;
-}
-
 bool FUfbxParser::LoadFbxFile(const FString& Filename, UInterchangeBaseNodeContainer& NodeContainer)
 {
 	TRACE_CPUPROFILER_EVENT_SCOPE(FUfbxParser::FillContainerWithFbxScene)
@@ -411,6 +241,214 @@ bool FUfbxParser::FetchMeshPayloadData(const FString& PayloadKey, const FTransfo
 	return false;	
 }
 
+void FUfbxParser::Reset()
+{
+	NodeToUid.Reset();
+
+	ElementNames.Reset();
+	NodeNames.Reset();
+	BoneNames.Reset();
+	MeshNames.Reset();
+	CameraNames.Reset();
+	MaterialNames.Reset();
+
+	ElementIdToSceneNode.Reset();
+
+	PayloadContexts.Reset();
+
+	ufbx_free_scene(Scene);
+	Scene = nullptr;
+}
+
+FString FUfbxParser::GetElementNameRaw(const ufbx_element& Element) const
+{
+	return Convert::ToUnrealString(Element.name);
+}
+
+FString FUfbxParser::GetElementName(const ufbx_element& Element) const
+{
+	return ElementNames.GetOrMake(Element, [&]
+		{
+			return GetElementNameRaw(Element);
+		});
+}
+
+FString FUfbxParser::GetNodeName(const ufbx_node& Node) const
+{
+	return NodeNames.GetSafe(Node.element);
+}
+
+FString FUfbxParser::GetNodeUid(const ufbx_node& Node) const
+{
+	if (const FString* Found = NodeToUid.Find(&Node))
+	{
+		return *Found;
+	}
+
+	FString NodeUid = GetNodeName(Node);
+	// As we don't want the name to have the namespace removed,
+	// so name for the node wasn't sanitized(to keep namespace name within it)
+	// we manually generate the name here:
+	UE::Interchange::SanitizeName(NodeUid);
+
+	// Compose NodeUid as node "path" - unique node names connected with dot
+	// #ufbx_todo: check that element_id is enough
+	// reference FBX parser's uses UE::Interchange::Private::FFbxHelper::GetFbxNodeHierarchyName to create NodeUid
+	if (Node.parent)
+	{
+		const UInterchangeSceneNode*const* ParentSceneNodeFound = ElementIdToSceneNode.Find(Node.parent->element_id);
+		
+		if (ensureMsgf(ParentSceneNodeFound, TEXT("Expected to have child node Uid generated after parent was processed first")))
+		{
+			const UInterchangeSceneNode* ParentSceneNode = *ParentSceneNodeFound;
+	
+			NodeUid = ParentSceneNode->GetUniqueID()+TEXT(".")+NodeUid;
+		}
+	}
+
+	NodeToUid.Add(&Node, NodeUid);
+	return NodeUid;
+}
+
+FString FUfbxParser::GetNodeLabel(const ufbx_node& Node) const
+{
+	FString Label(GetNodeName(Node));
+	NodeNames.ManageNamespace(Label);
+	UE::Interchange::SanitizeName(Label);
+	return Label;
+}
+
+FString FUfbxParser::GetBoneNodeName(const ufbx_node& Node) const
+{
+	return BoneNames.GetOrMake(Node.element, [&]()
+		{
+			return GetElementNameRaw(Node.element);
+		});
+}
+
+FString FUfbxParser::GetMeshUid(const ufbx_element& Mesh) const
+{
+	return MeshNames.GetOrMake(Mesh, [&]()
+		{	return GetElementNameRaw(Mesh); });
+}
+
+FString FUfbxParser::GetMeshLabel(const ufbx_element& Mesh) const
+{
+	return GetMeshUid(Mesh);
+}
+
+FString FUfbxParser::GetMaterialName(const ufbx_material& Material) const
+{
+	return MaterialNames.GetOrMake(Material.element, [&]
+		{
+			return GetElementNameRaw(Material.element);
+		});
+}
+
+FString FUfbxParser::GetMaterialUid(const ufbx_material& Material) const
+{
+	return  TEXT("\\Material\\") + GetMaterialName(Material);
+}
+
+FString FUfbxParser::GetMaterialLabel(const ufbx_material& Material) const
+{
+	return GetMaterialName(Material);
+}
+
+FString FUfbxParser::GetCameraName(const ufbx_camera& Camera) const
+{
+	return CameraNames.GetOrMake(Camera.element, [&]
+		{
+			FString CameraName = GetElementNameRaw(Camera.element);
+
+			if (CameraName.IsEmpty())
+			{
+				FString DefaultNamePrefix(UInterchangePhysicalCameraNode::StaticAssetTypeName());
+				CameraName = DefaultNamePrefix + TEXT("_") + FString::FromInt(Camera.typed_id);
+			}
+			return CameraName;
+		});
+}
+
+void FUfbxParser::ConvertProperty(const ufbx_prop& Prop, UInterchangeBaseNode* SceneNode, const TOptional<FString>& PayloadKey)
+{
+	if (Prop.flags & UFBX_PROP_FLAG_USER_DEFINED)
+	{
+		FString PropertyName = Convert::ToUnrealString(Prop.name);
+
+		switch (Prop.type)
+		{
+		case UFBX_PROP_UNKNOWN:
+			break;
+		case UFBX_PROP_BOOLEAN:
+			{
+				const bool PropertyValue = Prop.value_int != 0;
+				UInterchangeUserDefinedAttributesAPI::CreateUserDefinedAttribute(SceneNode, PropertyName, PropertyValue, PayloadKey);
+			}
+			break;
+		case UFBX_PROP_INTEGER:
+			{
+				const int32 PropertyValue = Prop.value_int;
+				UInterchangeUserDefinedAttributesAPI::CreateUserDefinedAttribute(SceneNode, PropertyName, PropertyValue, PayloadKey);
+			}
+			break;
+		case UFBX_PROP_NUMBER:
+			{
+				const float PropertyValue = Prop.value_real;
+				UInterchangeUserDefinedAttributesAPI::CreateUserDefinedAttribute(SceneNode, PropertyName, PropertyValue, PayloadKey);
+			}
+			break;
+		case UFBX_PROP_VECTOR:
+			{
+				const FVector4d PropertyValue = Convert::ConvertVec4(Prop.value_vec4);
+				UInterchangeUserDefinedAttributesAPI::CreateUserDefinedAttribute(SceneNode, PropertyName, PropertyValue, PayloadKey);
+			}
+			break;
+		case UFBX_PROP_COLOR:
+			{
+				const FVector PropertyValue = FVector(Convert::ConvertColor(Prop.value_vec4));
+				UInterchangeUserDefinedAttributesAPI::CreateUserDefinedAttribute(SceneNode, PropertyName, PropertyValue, PayloadKey);
+			}
+			break;
+		case UFBX_PROP_COLOR_WITH_ALPHA:
+			{
+				const FLinearColor PropertyValue = Convert::ConvertColor(Prop.value_vec4);
+				UInterchangeUserDefinedAttributesAPI::CreateUserDefinedAttribute(SceneNode, PropertyName, PropertyValue, PayloadKey);
+			}
+			break;
+		case UFBX_PROP_STRING:
+			{
+				const FString PropertyValue = Convert::ToUnrealString(Prop.value_str);
+				UInterchangeUserDefinedAttributesAPI::CreateUserDefinedAttribute(SceneNode, PropertyName, PropertyValue, PayloadKey);
+			}
+			break;
+		case UFBX_PROP_DATE_TIME:
+			{
+				const FDateTime PropertyValue(Prop.value_int);
+				UInterchangeUserDefinedAttributesAPI::CreateUserDefinedAttribute(SceneNode, PropertyName, PropertyValue, PayloadKey);
+			}
+			break;
+		// #ufbx_todo: support other useful custom attributes among these types
+		case UFBX_PROP_TRANSLATION:
+			break;
+		case UFBX_PROP_ROTATION:
+			break;
+		case UFBX_PROP_SCALING:
+			break;
+		case UFBX_PROP_DISTANCE:
+			break;
+		case UFBX_PROP_COMPOUND:
+			break;
+		case UFBX_PROP_BLOB:
+			break;
+		case UFBX_PROP_REFERENCE:
+			break;
+		case UFBX_PROP_TYPE_FORCE_32BIT:
+			break;
+		}
+	}
+}
+
 #if WITH_ENGINE
 bool FUfbxParser::FetchMeshPayloadData(const FString& PayloadKey, const FTransform& MeshGlobalTransform, FMeshPayloadData& OutMeshPayloadData)
 {
diff --git a/Engine/Plugins/Interchange/Runtime/Source/Parsers/Fbx/Private/Ufbx/UfbxParser.h b/Engine/Plugins/Interchange/Runtime/Source/Parsers/Fbx/Private/Ufbx/UfbxParser.h
--- a/Engine/Plugins/Interchange/Runtime/Source/Parsers/Fbx/Private/Ufbx/UfbxParser.h
+++ b/Engine/Plugins/Interchange/Runtime/Source/Parsers/Fbx/Private/Ufbx/UfbxParser.h
@@ -12,6 +12,7 @@
 
 // #ufbx_todo: move IFbxParser to separate file and include IFbxParser.h instead
 #include "FbxAPI.h"
+#include "InterchangeHelper.h"
 
 #include "MeshDescription.h"
 
@@ -85,16 +86,16 @@ namespace UE::Interchange::Private
 	{
 	public:
 		explicit FUfbxParser(TWeakObjectPtr<UInterchangeResultsContainer> InResultsContainer);
-		~FUfbxParser();
+		virtual ~FUfbxParser() override;
 
 		virtual const TCHAR* GetName() override { return TEXT("uFBX"); }
 
 		virtual bool IsThreadSafe() override { return true; }
 
 		virtual void SetResultContainer(UInterchangeResultsContainer* Result) override;
-		virtual void SetConvertSettings(const bool InbConvertScene, const bool InbForceFrontXAxis, const bool InbConvertSceneUnit, const bool InbKeepFbxNamespace)  override
+		virtual void SetConvertSettings(const bool InbConvertScene, const bool InbForceFrontXAxis, const bool InbConvertSceneUnit, const bool InbKeepFbxNamespace) override
 		{
-			// #ufbx_todo: 
+			bKeepFbxNamespace = InbKeepFbxNamespace;
 		}
 
 		virtual bool LoadFbxFile(const FString& Filename, UInterchangeBaseNodeContainer& NodeContainer) override;
@@ -149,35 +150,43 @@ namespace UE::Interchange::Private
 		 * @param bIsJoint 
 		 * @return 
 		 */
-		FString GetElementName(const ufbx_element& Element, bool bIsJoint = false) const;
 		FString GetElementNameRaw(const ufbx_element& Element) const;
+
 		/**
 		 * Returns name build for element making sure it's unique for each element
 		 * - name based on sanitized name of the element so it can be used in any place when sanitized string is needed
 		 * @param Element *
 		 * @return 
 		 */
-		FString GetElementNameDeduplicated(const ufbx_element& Element, bool bIsJoint = false) const;
+		FString GetElementName(const ufbx_element& Element) const;
 
+		/**
+		 * Get/Make "Name" of the node, unique(taking into account sanitized form) bu tnot sanitized(to keep namespace which contains colon), used to base Uid on, and set as a label(except bone nodes - they have different label, see @ref GetBoneNodeName)
+		 * @param Node 
+		 * @return 
+		 */
+		FString GetNodeName(const ufbx_node& Node) const;
+		FString GetNodeLabel(const ufbx_node& Node) const;
 		FString GetNodeUid(const ufbx_node& Node) const;
-		FString GetBoneNodeUid(const ufbx_node& Node) const;
+
+		// Set name for the joint node, same name to be passed to UE::Interchange::FMeshPayloadData::JointNames when fetching the skinned mesh payload
+		// sanitized to replace underscore with dash
+		FString GetBoneNodeName(const ufbx_node& Node) const;
 
 		FString GetMeshUid(const ufbx_element& Node) const;
 		FString GetMeshLabel(const ufbx_element& Mesh) const;
 
+		FString GetMaterialName(const ufbx_material& Material) const;
 		FString GetMaterialUid(const ufbx_material& Material) const;
-		FName GetMaterialLabel(const ufbx_material& Material) const;
+		FString GetMaterialLabel(const ufbx_material& Material) const;
 
 		FName GetMaterialSlotName(const ufbx_material* Material) const
 		{
-			if (!Material)
-			{
-				return NAME_None;
-			}
-
-			return {Material ? FName(*GetElementNameDeduplicated(Material->element)) : NAME_None};
+			return Material ? FName(*GetMaterialName(*Material)) : NAME_None;
 		}
 
+		FString GetCameraName(const ufbx_camera& Camera) const;
+
 		static void ConvertProperty(const ufbx_prop& Prop, UInterchangeBaseNode* SceneNode, const TOptional<FString>& PayloadKey = TOptional<FString>());
 
 		ufbx_scene* Scene = nullptr;
@@ -284,12 +293,158 @@ namespace UE::Interchange::Private
 		TMap<uint32_t, UInterchangeSceneNode*> ElementIdToSceneNode;
 		FPayloadContexts PayloadContexts;
 
-		// Keeps track of unique names generated for elements
-		mutable TMap<FString, const ufbx_element*> NameToElement;
-		mutable TMap<const ufbx_element*, FString> ElementToName;
+		/**
+		 * Maintain map of elements unique names
+		 */
+		struct FUniqueNames
+		{
+			FUniqueNames(const FUfbxParser& InParser)
+				: Parser(InParser)
+				, bUseNCL(true)
+				, bUseNodeRules(false)
+				, bIsJoint(false)
+			{
+			}
+
+			/**
+			 * Ensures that name is unique modulo sanitation. Although stored name may not be sanitized(Node name need to keep namespace which contains colon)
+			 * still to pass Uid(based on Node name) to Interchange that Uid should dbe sanitized.
+			 * @param Element 
+			 * @param BaseName 
+			 * @param bOverrideUseNCL whether to use '_ncl_' suffix when making unique name, or just add a number
+			 * @return 
+			 */
+			FString MakeUniqueName(const ufbx_element& Element, const FString& BaseName, bool bOverrideUseNCL=false) const 
+			{
+				// Fix base name, non-node elements might need to have namespace removed
+				const FString FixedBaseName = [&]
+					{
+						FString Result = BaseName.IsEmpty() ? DefaultName : BaseName;
+
+						if (!bUseNodeRules)
+						{
+							ManageNamespace(Result);
+						}
+						return Result;
+					}();
+
+				// Start with base name.
+				FString Name = FixedBaseName;
+				FString SanitizedName = Name;
+				UE::Interchange::SanitizeName(SanitizedName, bIsJoint);
+
+				const FString SanitizedNameInitial = SanitizedName;
+
+				if (!bUseNodeRules)
+				{
+					Name = SanitizedName;
+				}
 
-		mutable TMap<FString, const ufbx_node*> UidToNode;
-		mutable TMap<const ufbx_node*, FString> NodeToUid;
+				for (int32 Count = 0;;++Count)
+				{
+					if (!NamesUsed.Contains(SanitizedName))
+					{
+						break;
+					}
+
+					FString Suffix = (bUseNCL || bOverrideUseNCL ? TEXT("_ncl_") : TEXT("")) + FString::FromInt(Count+1);
+					SanitizedName = SanitizedNameInitial + Suffix;
+
+					if (bUseNodeRules)
+					{
+						Name = FixedBaseName + Suffix;
+					}
+					else
+					{
+						Name = SanitizedName;
+					}
+				}
+				NamesUsed.Add(SanitizedName);
+
+				ElementToName.Add(&Element, Name);
+				return Name;
+			}
+
+			void ManageNamespace(FString& ObjectName) const
+			{
+				if (Parser.bKeepFbxNamespace)
+				{
+					if (ObjectName.Contains(TEXT(":")))
+					{
+						ObjectName.ReplaceInline(TEXT(":"), TEXT("_"));
+					}
+				}
+				else
+				{
+					// Remove namespaces
+					int32 LastNamespaceTokenIndex = INDEX_NONE;
+					if (ObjectName.FindLastChar(TEXT(':'), LastNamespaceTokenIndex))
+					{
+						//+1 to remove the ':' character we found
+						ObjectName.RightChopInline(LastNamespaceTokenIndex + 1, EAllowShrinking::Yes);
+					}
+				}
+			}
+
+			const FString* FindName(const ufbx_element& Element) const
+			{
+				return ElementToName.Find(&Element);
+			}
+
+			FString GetOrMake(const ufbx_element& Element, const TFunctionRef<FString()>& GetBaseName) const
+			{
+				if (const FString* Found = ElementToName.Find(&Element))
+				{
+					return *Found;
+				}
+				return MakeUniqueName(Element, GetBaseName());
+			}
+
+			FString GetSafe(const ufbx_element& Element) const
+			{
+				if (const FString* Found = FindName(Element))
+				{
+					return *Found;
+				}
+				// Shouldn't happen!
+				ensure(false);
+				return TEXT("<not found>");
+			}
 
+			void Reset()
+			{
+				NamesUsed.Reset();
+				ElementToName.Reset();
+			}
+
+			// Unique sanitized name already in use
+			mutable TSet<FString> NamesUsed;
+
+			// Name (unique) registered for each element, maybe non-sanitized(e.g. for Nodes)! But having it sanitized should still produce unique
+			mutable TMap<const ufbx_element*, FString> ElementToName;
+
+			const FUfbxParser& Parser;
+			// use "_ncl_" suffix to create uniques, followed by number, or just ""
+			bool bUseNCL;
+			// Keeping namespace name for nodes and generating unique names without removing namespace 
+			bool bUseNodeRules;
+			// Name to use when source base name is empty
+			FString DefaultName;
+			// Joins have different rules for name sanitation
+			bool bIsJoint;
+		};
+
+		bool bKeepFbxNamespace = false;
+
+		// Unique names for elements by type, to base Label and Uid on
+		FUniqueNames ElementNames;
+		FUniqueNames NodeNames;
+		FUniqueNames BoneNames;
+		FUniqueNames MeshNames;
+		FUniqueNames CameraNames;
+		FUniqueNames MaterialNames;
+
+		// Cache of node Uids
+		mutable TMap<const ufbx_node*, FString> NodeToUid;
 	};
 }
diff --git a/Engine/Plugins/Interchange/Runtime/Source/Parsers/Fbx/Private/Ufbx/UfbxScene.cpp b/Engine/Plugins/Interchange/Runtime/Source/Parsers/Fbx/Private/Ufbx/UfbxScene.cpp
--- a/Engine/Plugins/Interchange/Runtime/Source/Parsers/Fbx/Private/Ufbx/UfbxScene.cpp
+++ b/Engine/Plugins/Interchange/Runtime/Source/Parsers/Fbx/Private/Ufbx/UfbxScene.cpp
@@ -2,12 +2,15 @@
 
 #include "UfbxScene.h"
 
-#include "InterchangeFbxSettings.h"
 #include "UfbxConvert.h"
 #include "UfbxParser.h"
+#include "UfbxAnimation.h"
+
+#include "InterchangeFbxSettings.h"
 
 #include "InterchangeMeshNode.h"
-#include "UfbxAnimation.h"
+#include "InterchangeCameraNode.h"
+#include "Math/UnitConversion.h"
 #include "Nodes/InterchangeBaseNodeContainer.h"
 #include "Nodes/InterchangeUserDefinedAttribute.h"
 
@@ -43,9 +46,10 @@ namespace UE::Interchange::Private
 
 		void InitHierarchyRecursively(FUfbxParser& Parser, UInterchangeBaseNodeContainer& NodeContainer, const ufbx_node& Node, UInterchangeSceneNode* ParentSceneNode)
 		{
-			Convert::FNodeNameAndUid NodeNameAndUid(Parser, ParentSceneNode, Node);
+			FString Label(Parser.GetNodeLabel(Node));
+			FString UniqueID(Parser.GetNodeUid(Node));
 			UInterchangeSceneNode* SceneNode = CreateTransformNode(Parser,
-				NodeContainer, NodeNameAndUid.Label, NodeNameAndUid.UniqueID, Convert::ConvertTransform(Node.local_transform), ParentSceneNode ? ParentSceneNode->GetUniqueID() : TEXT(""));
+				NodeContainer, Label, UniqueID, Convert::ConvertTransform(Node.local_transform), ParentSceneNode ? ParentSceneNode->GetUniqueID() : TEXT(""));
 
 			Parser.ElementIdToSceneNode.Add(Node.element_id, SceneNode);
 
@@ -202,7 +206,7 @@ namespace UE::Interchange::Private
 			// To recognize Node as a Joint, see e.g. UInterchangePipelineMeshesUtilities::CreateInterchangePipelineMeshesUtilities
 			UnrealNode->AddSpecializedType(FSceneNodeStaticData::GetJointSpecializeTypeString());
 
-			FString JointNodeName = Parser.GetBoneNodeUid(InNode);
+			FString JointNodeName = Parser.GetBoneNodeName(InNode);
 			UnrealNode->SetDisplayLabel(JointNodeName);
 
 			//Get the bind pose transform for this joint
@@ -404,6 +408,29 @@ namespace UE::Interchange::Private
 		ufbx_node* RootNode = Parser.Scene->root_node;
 		Parser.ElementIdToSceneNode.Reserve(Parser.Scene->nodes.count);
 
+		// Make unique names for nodes
+		for (const ufbx_node* Node : Parser.Scene->nodes)
+		{
+			FString BaseName = Parser.GetElementNameRaw(Node->element);
+
+			bool bUseNCL = false;
+
+			if (BaseName.IsEmpty())
+			{
+				if (Node->is_root)
+				{
+					// ufbx doesn't forcefully initialize root node name, so in case it's empty - set it
+					BaseName = TEXT("RootNode");
+				}
+				else
+				{
+					bUseNCL = true;
+				}
+			}
+
+			Parser.NodeNames.MakeUniqueName(Node->element, BaseName, bUseNCL);
+		}
+
 		SceneUtils::InitHierarchyRecursively(Parser, NodeContainer, *RootNode, nullptr);
 
 		for (const ufbx_node* Node : Parser.Scene->nodes)
