[Interchange] ufbx rigid animation in some cases of coordsys conversion

diff --git a/Engine/Plugins/Interchange/Runtime/Source/Parsers/Fbx/Private/Ufbx/UfbxAnimation.cpp b/Engine/Plugins/Interchange/Runtime/Source/Parsers/Fbx/Private/Ufbx/UfbxAnimation.cpp
--- a/Engine/Plugins/Interchange/Runtime/Source/Parsers/Fbx/Private/Ufbx/UfbxAnimation.cpp
+++ b/Engine/Plugins/Interchange/Runtime/Source/Parsers/Fbx/Private/Ufbx/UfbxAnimation.cpp
@@ -415,15 +415,48 @@ namespace UE::Interchange::Private
 
 		// #ufbx_todo: FAnimationPayloadContext::InternalFetchCurveNodePayloadToFile does some pivot handling, see ResetPivotsPrePostRotationsAndSetRotationOrder
 
-		AddCurve(Animation.TranslationProp, 0);
-		AddCurve(Animation.TranslationProp, 2, -1.0);
-		AddCurve(Animation.TranslationProp, 1);
-		AddCurve(Animation.RotationProp, 0);
-		AddCurve(Animation.RotationProp, 1);
-		AddCurve(Animation.RotationProp, 2);
-		AddCurve(Animation.ScalingProp, 0);
-		AddCurve(Animation.ScalingProp, 2);
-		AddCurve(Animation.ScalingProp, 1);
+		// ufbx doesn't convert transform curve data to desired coord space(or unit scale), so need to handle node transform curve mapping here
+		ufbx_real TranslationScaling = 1;
+
+		int32 AxisMap[3] = {0, 1, 2};
+		double AxisSign[3] = {1.0, 1.0, 1.0};
+		int32 AxisRotationMap[3] = {0, 1, 2};
+		double AxisRotationSign[3] = {1.0, 1.0, 1.0};
+
+		const ufbx_node* Node = Animation.Node;
+
+		if (Node && Node->has_adjust_transform && (Parser.Scene->settings.axes.up == UFBX_COORDINATE_AXIS_POSITIVE_Y))
+		{
+			// Nodes that are just below the RootNode(artificial) have adjust_transform and have they YZ swapped
+			// When FBX was exported as Yup
+			AxisMap[1] = 2;
+			AxisMap[2] = 1;
+		}
+		else
+		{
+			// Other nodes just need to fix handedness by mirroring on Y axis
+			AxisSign[1] = -1.0f;
+		}
+
+		// Rotations need to be inverted due to handedness change
+		AxisRotationSign[1] = -1.0f;
+		AxisRotationSign[2] = -1.0f;
+
+		if (Node)
+		{
+			TranslationScaling = Node->adjust_pre_scale;
+		}
+
+		AddCurve(Animation.TranslationProp, AxisMap[0], AxisSign[0]*TranslationScaling);
+		AddCurve(Animation.TranslationProp, AxisMap[1], AxisSign[1]*TranslationScaling);
+		AddCurve(Animation.TranslationProp, AxisMap[2], AxisSign[2]*TranslationScaling);
+
+		AddCurve(Animation.RotationProp, AxisRotationMap[0], AxisRotationSign[0]);
+		AddCurve(Animation.RotationProp, AxisRotationMap[1], AxisRotationSign[1]);
+		AddCurve(Animation.RotationProp, AxisRotationMap[2], AxisRotationSign[2]);
+		AddCurve(Animation.ScalingProp, AxisMap[0], 1);
+		AddCurve(Animation.ScalingProp, AxisMap[1], 1);
+		AddCurve(Animation.ScalingProp, AxisMap[2], 1);
 
 		return true;
 	}
@@ -691,14 +724,20 @@ namespace UE::Interchange::Private
 					{
 						if (FCStringAnsi::Strcmp(UFBX_Lcl_Translation, Prop.prop_name.data)==0)
 						{
-							TransformationPayloadContexts.FindOrAdd(Node).TranslationProp = &Prop;
+							FRigidAnimationPayloadContext& RigidAnimationPayloadContext = TransformationPayloadContexts.FindOrAdd(Node);
+							RigidAnimationPayloadContext.Node = Node;
+							RigidAnimationPayloadContext.TranslationProp = &Prop;
 						}
 						if (FCStringAnsi::Strcmp(UFBX_Lcl_Rotation, Prop.prop_name.data)==0)
 						{
+							FRigidAnimationPayloadContext& RigidAnimationPayloadContext = TransformationPayloadContexts.FindOrAdd(Node);
+							RigidAnimationPayloadContext.Node = Node;
 							TransformationPayloadContexts.FindOrAdd(Node).RotationProp = &Prop;
 						}
 						if (FCStringAnsi::Strcmp(UFBX_Lcl_Scaling, Prop.prop_name.data)==0)
 						{
+							FRigidAnimationPayloadContext& RigidAnimationPayloadContext = TransformationPayloadContexts.FindOrAdd(Node);
+							RigidAnimationPayloadContext.Node = Node;
 							TransformationPayloadContexts.FindOrAdd(Node).ScalingProp = &Prop;
 						}
 					}
diff --git a/Engine/Plugins/Interchange/Runtime/Source/Parsers/Fbx/Private/Ufbx/UfbxParser.h b/Engine/Plugins/Interchange/Runtime/Source/Parsers/Fbx/Private/Ufbx/UfbxParser.h
--- a/Engine/Plugins/Interchange/Runtime/Source/Parsers/Fbx/Private/Ufbx/UfbxParser.h
+++ b/Engine/Plugins/Interchange/Runtime/Source/Parsers/Fbx/Private/Ufbx/UfbxParser.h
@@ -70,6 +70,7 @@ namespace UE::Interchange::Private
 	struct FRigidAnimationPayloadContext
 	{
 		// Transform animation curves array in the order Interchange accepts it - Tx, Ty, Tz, Rx, Ry, Rz, Sx, Sy, Sz
+		const ufbx_node* Node = nullptr;
 		const ufbx_anim_prop* TranslationProp = nullptr;
 		const ufbx_anim_prop* RotationProp = nullptr;
 		const ufbx_anim_prop* ScalingProp = nullptr;
