Interchange:
- Add a facility to auto detect long/lat cubemap textures

diff --git a/Engine/Plugins/Interchange/Runtime/Source/Pipelines/InterchangePipelines.Build.cs b/Engine/Plugins/Interchange/Runtime/Source/Pipelines/InterchangePipelines.Build.cs
--- a/Engine/Plugins/Interchange/Runtime/Source/Pipelines/InterchangePipelines.Build.cs
+++ b/Engine/Plugins/Interchange/Runtime/Source/Pipelines/InterchangePipelines.Build.cs
@@ -30,6 +30,7 @@ namespace UnrealBuildTool.Rules
 					"GeometryCache",
 					"GLTFCore",
 					"HairStrandsCore",
+					"ImageCore",
 					"InputCore",
 					"InterchangeImport",
 					"LevelSequence",
diff --git a/Engine/Plugins/Interchange/Runtime/Source/Pipelines/Private/InterchangeGenericTexturePipeline.cpp b/Engine/Plugins/Interchange/Runtime/Source/Pipelines/Private/InterchangeGenericTexturePipeline.cpp
--- a/Engine/Plugins/Interchange/Runtime/Source/Pipelines/Private/InterchangeGenericTexturePipeline.cpp
+++ b/Engine/Plugins/Interchange/Runtime/Source/Pipelines/Private/InterchangeGenericTexturePipeline.cpp
@@ -2,9 +2,11 @@
 
 #include "InterchangeGenericTexturePipeline.h"
 
+#include "Async/ParallelFor.h"
 #include "Engine/Texture.h"
 #include "Engine/TextureCube.h"
 #include "Engine/TextureLightProfile.h"
+#include "ImageUtils.h"
 #include "InterchangePipelineLog.h"
 #include "InterchangeTexture2DArrayFactoryNode.h"
 #include "InterchangeTexture2DArrayNode.h"
@@ -31,6 +33,25 @@
 #include "UDIMUtilities.h"
 #endif //WITH_EDITOR
 
+int32 GInterchangeMinTextureCubeMapDetection = 512;
+static FAutoConsoleVariableRef CVarInterchangeMinTextureCubeMapDetection(
+	TEXT("Interchange.Textures.MinCubeMapDetection"),
+	GInterchangeMinTextureCubeMapDetection,
+	TEXT("The minimum size to auto detect that an image is a long-lat cubemap texture during import. Default: 512, Minimum value: 8")
+);
+
+static FAutoConsoleVariableSink CVarInterchangeMinTextureCubeMapDetectionSink
+(FConsoleCommandDelegate::CreateLambda([]()
+ {
+	 int32 MinValue = CVarInterchangeMinTextureCubeMapDetection->GetInt();
+	 if (MinValue < 8)
+	 {
+		 CVarInterchangeMinTextureCubeMapDetection->Set(8, ECVF_SetByConsole);
+		 UE_LOG(LogInterchangePipeline, Warning, TEXT("Interchange.Textures.MinCubeMapDetection cannot be less than 8. Reset to 8."));
+	 }
+ }
+));
+
 namespace UE::Interchange::Private
 {
 	UClass* GetDefaultFactoryClassFromTextureNodeClass(UClass* NodeClass)
@@ -108,6 +129,244 @@ namespace UE::Interchange::Private
 			return TA_Wrap;
 		}
 	}
+
+	// Compute linear rec709 luminance
+	float ComputeLuminance(const FLinearColor& Color)
+	{
+		return 0.2126f * Color.R + 0.7152f * Color.G + 0.0722f * Color.B;
+	}
+	
+	/** Compute the global average luminance of an image */
+	double ComputeGlobalLuminance(const TArrayView64<FLinearColor>& Pixels, int32 Width, int32 Height)
+	{
+		TArray<double> LocalLuminances;
+		LocalLuminances.SetNum(Height);
+				
+		ParallelFor
+		(Height, [&](int32 Y)
+		 {
+			 double RowSum = 0.0;
+			 for (int32 X = 0; X < Width; ++X)
+			 {
+				 RowSum += ComputeLuminance(Pixels[Y * Width + X]);
+			 }
+			 LocalLuminances[Y] = RowSum;
+		 });
+
+		double GlobalLuminance = 0;
+		for (double Luminance : LocalLuminances)
+		{
+			GlobalLuminance += Luminance;
+		}
+
+		GlobalLuminance /= (Width * Height);
+		return GlobalLuminance;
+	}
+
+	/** 
+	 * For equirectangular maps, since there's a horizontal wrap then the edges should be very similar
+	 * We both compare color and luminance average differences
+	 */
+	bool CompareEdgeSimilarity(const TArrayView64<FLinearColor>& Pixels, int32 Width, int32 Height)
+	{
+		struct FLocalEdgeDiff
+		{
+			double RGB = 0.0;
+			double Luminance = 0.0;
+		};
+
+		TArray<FLocalEdgeDiff> EdgeDiffs;
+		EdgeDiffs.SetNumZeroed(Height);
+
+		ParallelFor
+		(Height, [&](int32 Y)
+		 {
+			 FLocalEdgeDiff& EdgeDiff= EdgeDiffs[Y];
+
+			 const FLinearColor& Left = Pixels[Y * Width];
+			 const FLinearColor& Right = Pixels[Y * Width + (Width - 1)];
+
+			 const float DeltaR = Left.R - Right.R;
+			 const float DeltaG = Left.G - Right.G;
+			 const float DeltaB = Left.B - Right.B;
+
+			 const float LLeft = ComputeLuminance(Left);
+			 const float LRight = ComputeLuminance(Right);
+
+			 EdgeDiff.Luminance += FMath::Abs(LLeft - LRight);
+			 EdgeDiff.RGB += FMath::Sqrt(DeltaR * DeltaR + DeltaG * DeltaG + DeltaB * DeltaB);
+		 });
+
+		double MeanErrorRGB = 0.0;
+		double MeanErrorLuminance = 0.0;
+
+		for (const FLocalEdgeDiff& Pair : EdgeDiffs)
+		{
+			MeanErrorRGB += Pair.RGB;
+			MeanErrorLuminance += Pair.Luminance;
+		}
+
+		MeanErrorRGB /= Height;
+		MeanErrorLuminance /= Height;
+
+		return MeanErrorRGB < 0.1 && MeanErrorLuminance < 0.1;
+	}
+
+	/**
+	 * Compute the variance of the luminance using Welford's online algorithm
+	 */
+	double ComputeLuminanceVarianceWelford(const TArrayView64<FLinearColor>& Pixels, int32 Width, int32 YStart, int32 YEnd)
+	{
+		struct FLocalVariance
+		{
+			double Mean = 0.0;
+			double M2 = 0.0; // Sum of squares of differences from the current mean
+			int64 Count = 0;
+
+			void Accumulate(double Value)
+			{
+				++Count;
+				double Delta = Value - Mean;
+				Mean += Delta / Count;
+				double Delta2 = Value - Mean;
+				M2 += Delta * Delta2;
+			}
+
+			void Combine(const FLocalVariance& Other)
+			{
+				if (Other.Count == 0)
+				{
+					return;
+				}
+
+				double Delta = Other.Mean - Mean;
+				int64 NewCount = Count + Other.Count;
+
+				M2 += Other.M2 + Delta * Delta * Count * Other.Count / NewCount;
+				Mean += Delta * Other.Count / NewCount;
+				Count = NewCount;
+			}
+
+			double Variance() const
+			{
+				return (Count > 1) ? M2 / (Count - 1) : 0.0;
+			}
+		};
+
+
+		int32 SampleHeight = YEnd - YStart;
+		TArray<FLocalVariance> LocalStats;
+		LocalStats.SetNum(SampleHeight);
+
+		const double Epsilon = 1e-6;
+
+		ParallelFor(
+			SampleHeight,
+			[&Pixels, Width, &LocalStats, YStart, Epsilon](int32 LocalY)
+			{
+				int32 Y = YStart + LocalY;
+				FLocalVariance Stats;
+
+				for (int32 X = 0; X < Width; ++X)
+				{
+					// use log to flatten the values in case of high disparities
+					const double L = FMath::Loge(FMath::Max(ComputeLuminance(Pixels[Y * Width + X]), Epsilon));
+					Stats.Accumulate(L);
+				}
+
+				LocalStats[LocalY] = Stats;
+			});
+
+		// Combine all local stats into a global one
+		FLocalVariance GlobalStats;
+		for (const FLocalVariance& Stats : LocalStats)
+		{
+			GlobalStats.Combine(Stats);
+		}
+
+		return GlobalStats.Variance();
+	}
+
+	/**
+	 * Usually for equirectangular textures, the polar regions(top / bottom) appear visually compressed
+	 * which results in lower luminance variance compared to the equatorial (middle)
+	 * We detect this by comparing the pixels distribution at the poles (10% of the whole image) against the equatorial (80% of the rest of the image)
+	 * If we have a pole repartion lesser than 70% for the equatorial then it should be a good candidate
+	 */
+	bool HasPoleCompression(const TArrayView64<FLinearColor>& Pixels, int32 Width, int32 Height)
+	{
+		const int32 Pole = FMath::Clamp(Height / 10, 1, Height / 4); // avoid 0 or really small bands
+		const int32 MiddleStart = Pole;
+		const int32 MiddleEnd = Height - Pole;
+
+		// We compute the global luminance of the image, probably a bit overkill
+		// but it can help mitigate high discrepancies between bright and dark areas making it intensity and exposure-independent
+		const double Luminance = ComputeGlobalLuminance(Pixels, Width, Height);
+		double Luminance2 = FMath::Max(Luminance * Luminance, 1e-6);
+
+		double VarianceTop = ComputeLuminanceVarianceWelford(Pixels, Width, 0, Pole) / Luminance2;
+		double VarianceMiddle = ComputeLuminanceVarianceWelford(Pixels, Width, MiddleStart, MiddleEnd) / Luminance2;
+		double VarianceBottom = ComputeLuminanceVarianceWelford(Pixels, Width, Height - Pole, Height) / Luminance2;
+
+		// 80% is probably too permissive?
+		double VarianceTolerance = VarianceMiddle * 0.8;
+		if (VarianceTop < VarianceTolerance && VarianceBottom < VarianceTolerance)
+		{
+			return true;
+		}
+
+		// the first test might be sometimes too restrictive whereas we still have an equirectangular map
+		// for example for hdr file with bright light on top (like studios or indoor with neons)
+		// the variance at the top can be a bit greater than 80% (probably a bit too permissive) of the variance of the middle
+		// by averaging the poles we ensure that at least we're not exceluding a true equirectangular map
+		return (0.5 * (VarianceTop + VarianceBottom)) < VarianceTolerance;
+	}
+
+	/**
+	 * A texture should be a cubemap texture if:
+	 * - its ratio is 2:1 (we try to not be as strict as possible, giving a very small marge of error)
+	 * - it should wrap horizontally without seams
+	 * - if the poles (top/bottom of the image) appear more compressed than the middle 
+	 */
+	bool IsEquirectangular(const FString& SourceFile)
+	{
+		using namespace UE::Interchange::Private;
+
+		bool bEquirectangular = false;
+
+		FImage Image, ImageFloat32;
+		const TCHAR* Filename = *SourceFile;
+		if (FImageUtils::LoadImage(Filename, Image))
+		{
+			// aspect ratio of equirectangular maps should be 2:1
+			const bool bIsAspectRatio_2_1 =
+				Image.GetWidth() > 0 &&
+				Image.GetHeight() >= GInterchangeMinTextureCubeMapDetection &&
+				Image.GetNumPixels() &&
+				FMath::IsNearlyEqual(float(Image.GetWidth()) / float(Image.GetHeight()), 2.);
+
+			if (bIsAspectRatio_2_1)
+			{
+				TArrayView64<FLinearColor> Pixels;
+				// for a proper detection we want a full float image format in linear space which is generally the case for .hdr or.exr image
+				// but definitely not for jpg or png, and we want to avoid unnecessary conversions
+				if (Image.Format == ERawImageFormat::RGBA32F && Image.GammaSpace == EGammaSpace::Linear)
+				{
+					Pixels = Image.AsRGBA32F();
+				}
+				else
+				{
+					Image.CopyTo(ImageFloat32, ERawImageFormat::RGBA32F, EGammaSpace::Linear);
+					Pixels = ImageFloat32.AsRGBA32F();
+				}
+				
+				bool bHasEdgeSimilarity = CompareEdgeSimilarity(Pixels, Image.GetWidth(), Image.GetHeight());
+				bool bHasPoleCompression = HasPoleCompression(Pixels, Image.GetWidth(), Image.GetHeight());
+				bEquirectangular = (bHasEdgeSimilarity && bHasPoleCompression);
+			}
+		}
+		return bEquirectangular;
+	}
 }
 
 FString UInterchangeGenericTexturePipeline::GetPipelineCategory(UClass* AssetClass)
@@ -266,10 +525,21 @@ UInterchangeTextureFactoryNode* UInterchangeGenericTexturePipeline::HandleCreati
 		}
 		else if (SourceFile)
 		{
-			const FString Extension = FPaths::GetExtension(SourceFile.GetValue()).ToLower();
-			if (FileExtensionsToImportAsLongLatCubemap.Contains(Extension))
+			using namespace UE::Interchange::Private;
+			if (bDetectLongLatCubemap)
+			{
+				if(IsEquirectangular(*SourceFile))
+				{
+					FactoryClass = UInterchangeTextureCubeFactoryNode::StaticClass();
+				}
+			}
+			else
 			{
-				FactoryClass = UInterchangeTextureCubeFactoryNode::StaticClass();
+				const FString Extension = FPaths::GetExtension(SourceFile.GetValue()).ToLower();
+				if (FileExtensionsToImportAsLongLatCubemap.Contains(Extension))
+				{
+					FactoryClass = UInterchangeTextureCubeFactoryNode::StaticClass();
+				}
 			}
 		}
 	}
diff --git a/Engine/Plugins/Interchange/Runtime/Source/Pipelines/Public/InterchangeGenericTexturePipeline.h b/Engine/Plugins/Interchange/Runtime/Source/Pipelines/Public/InterchangeGenericTexturePipeline.h
--- a/Engine/Plugins/Interchange/Runtime/Source/Pipelines/Public/InterchangeGenericTexturePipeline.h
+++ b/Engine/Plugins/Interchange/Runtime/Source/Pipelines/Public/InterchangeGenericTexturePipeline.h
@@ -51,8 +51,12 @@ public:
 	UPROPERTY(EditAnywhere, BlueprintReadWrite, Category = "Textures", Meta=(EditCondition="bImportTextures"))
 	bool bImportUDIMs = true;
 
+	/** Detect if a texture should be imported as long/lat cubemap*/
+	UPROPERTY(EditAnywhere, BlueprintReadWrite, Category = "Textures", Meta = (EditCondition = "bImportTextures"))
+	bool bDetectLongLatCubemap = false;
+
 	/** Specify the file types that should be imported as long/lat cubemaps. */
-	UPROPERTY(EditAnywhere, BlueprintReadWrite, Category = "Textures", Meta=(EditCondition="bImportTextures"))
+	UPROPERTY(EditAnywhere, BlueprintReadWrite, Category = "Textures", Meta=(EditCondition="bImportTextures && (bDetectLongLatCubemap == false)"))
 	TSet<FString> FileExtensionsToImportAsLongLatCubemap = {"hdr"};
 
 	/** 
