Interchange and USD: Improve support for camera parameters

#Synced-CL 47704025

diff --git a/Engine/Plugins/Importers/USDImporter/Source/USDStage/Private/USDLevelSequenceHelper.cpp b/Engine/Plugins/Importers/USDImporter/Source/USDStage/Private/USDLevelSequenceHelper.cpp
--- a/Engine/Plugins/Importers/USDImporter/Source/USDStage/Private/USDLevelSequenceHelper.cpp
+++ b/Engine/Plugins/Importers/USDImporter/Source/USDStage/Private/USDLevelSequenceHelper.cpp
@@ -2606,6 +2606,9 @@ namespace UE::USDLevelSequenceHelper::Private
 		UnrealIdentifiers::OrthoFarClipPlanePropertyName,
 		UnrealIdentifiers::OrthoNearClipPlanePropertyName,
 		UnrealIdentifiers::CustomNearClipppingPlanePropertyName,
+		UnrealIdentifiers::DepthOfFieldFstopPropertyName,
+		UnrealIdentifiers::CameraISOPropertyName,
+		UnrealIdentifiers::CameraShutterSpeedPropertyName
 	};
 }
 
@@ -4667,6 +4670,9 @@ void FUsdLevelSequenceHelperImpl::HandleMovieSceneChange(UMovieScene& MovieScene
 						RemoveTimeSamplesForPropertyIfNeeded(UsdPrim, Guid, UnrealIdentifiers::OrthoFarClipPlanePropertyName);
 						RemoveTimeSamplesForPropertyIfNeeded(UsdPrim, Guid, UnrealIdentifiers::OrthoNearClipPlanePropertyName);
 						RemoveTimeSamplesForPropertyIfNeeded(UsdPrim, Guid, UnrealIdentifiers::CustomNearClipppingPlanePropertyName);
+						RemoveTimeSamplesForPropertyIfNeeded(UsdPrim, Guid, UnrealIdentifiers::DepthOfFieldFstopPropertyName);
+						RemoveTimeSamplesForPropertyIfNeeded(UsdPrim, Guid, UnrealIdentifiers::CameraISOPropertyName);
+						RemoveTimeSamplesForPropertyIfNeeded(UsdPrim, Guid, UnrealIdentifiers::CameraShutterSpeedPropertyName);
 					}
 					else if (bIsLight)
 					{
diff --git a/Engine/Plugins/Interchange/Extensions/OpenUSD/Source/Import/Private/SchemaHandlers/CameraSchemaHandler.cpp b/Engine/Plugins/Interchange/Extensions/OpenUSD/Source/Import/Private/SchemaHandlers/CameraSchemaHandler.cpp
--- a/Engine/Plugins/Interchange/Extensions/OpenUSD/Source/Import/Private/SchemaHandlers/CameraSchemaHandler.cpp
+++ b/Engine/Plugins/Interchange/Extensions/OpenUSD/Source/Import/Private/SchemaHandlers/CameraSchemaHandler.cpp
@@ -25,7 +25,21 @@ namespace UE::GeomCameraSchemaHandler::Private
 {
 	using namespace UE::Interchange::USD;
 
-	void ConvertCameraPrim(const UE::FUsdPrim& Prim, UInterchangePhysicalCameraNode* CameraNode)
+	const FString ProjectionToken = UsdToUnreal::ConvertToken(pxr::UsdGeomTokens->projection);
+	const FString HorizontalApertureToken = UsdToUnreal::ConvertToken(pxr::UsdGeomTokens->horizontalAperture);
+	const FString VerticalApertureToken = UsdToUnreal::ConvertToken(pxr::UsdGeomTokens->verticalAperture);
+	const FString HorizontalApertureOffsetToken = UsdToUnreal::ConvertToken(pxr::UsdGeomTokens->horizontalApertureOffset);
+	const FString VerticalApertureOffsetToken = UsdToUnreal::ConvertToken(pxr::UsdGeomTokens->verticalApertureOffset);
+	const FString FocalLengthToken = UsdToUnreal::ConvertToken(pxr::UsdGeomTokens->focalLength);
+	const FString ClippingRangeToken = UsdToUnreal::ConvertToken(pxr::UsdGeomTokens->clippingRange);
+	const FString FstopToken = UsdToUnreal::ConvertToken(pxr::UsdGeomTokens->fStop);
+	const FString FocusDistanceToken = UsdToUnreal::ConvertToken(pxr::UsdGeomTokens->focusDistance);
+	const FString ExposureToken = UsdToUnreal::ConvertToken(pxr::UsdGeomTokens->exposure);
+	const FString ExposureIsoToken = UsdToUnreal::ConvertToken(pxr::UsdGeomTokens->exposureIso);
+	const FString ExposureTimeToken = UsdToUnreal::ConvertToken(pxr::UsdGeomTokens->exposureTime);
+	const FString ExposureFstopToken = UsdToUnreal::ConvertToken(pxr::UsdGeomTokens->exposureFStop);
+
+	void ConvertCameraPrim(const UE::FUsdPrim& Prim, UInterchangeBaseCameraNode* CameraNode)
 	{
 		TRACE_CPUPROFILER_EVENT_SCOPE(AddCameraNode)
 
@@ -33,29 +47,134 @@ namespace UE::GeomCameraSchemaHandler::Private
 		UE::FUsdStage Stage = Prim.GetStage();
 		FUsdStageInfo StageInfo{Stage};
 
-		static const FString FocalLengthToken = UsdToUnreal::ConvertToken(pxr::UsdGeomTokens->focalLength);
-		static const FString HorizontalApertureToken = UsdToUnreal::ConvertToken(pxr::UsdGeomTokens->horizontalAperture);
-		static const FString VerticalApertureToken = UsdToUnreal::ConvertToken(pxr::UsdGeomTokens->verticalAperture);
-
-		if (float FocalLength; UsdUtils::GetAuthoredAttributeValue<float>(Prim, FocalLengthToken, FocalLength))
+		FName Projection = UsdUtils::GetAttributeValue<FName>(Prim, ProjectionToken);
+		if (UInterchangePhysicalCameraNode* PhysicalCamera = Cast<UInterchangePhysicalCameraNode>(CameraNode))
 		{
-			FocalLength = UsdToUnreal::ConvertDistance(StageInfo, FocalLength);
-			CameraNode->SetCustomFocalLength(FocalLength);
-		}
+			if (float SensorWidth; UsdUtils::GetAuthoredAttributeValue<float>(Prim, HorizontalApertureToken, SensorWidth))
+			{
+				SensorWidth = UsdToUnreal::ConvertDistance(StageInfo, SensorWidth);
+				PhysicalCamera->SetCustomSensorWidth(SensorWidth);
+			}
+
+			if (float SensorHeight; UsdUtils::GetAuthoredAttributeValue<float>(Prim, VerticalApertureToken, SensorHeight))
+			{
+				SensorHeight = UsdToUnreal::ConvertDistance(StageInfo, SensorHeight);
+				PhysicalCamera->SetCustomSensorHeight(SensorHeight);
+			}
+
+			if (float SensorHorizontalOffset; UsdUtils::GetAuthoredAttributeValue<float>(Prim, HorizontalApertureOffsetToken, SensorHorizontalOffset))
+			{
+				SensorHorizontalOffset = UsdToUnreal::ConvertDistance(StageInfo, SensorHorizontalOffset);
+				PhysicalCamera->SetCustomSensorHorizontalOffset(SensorHorizontalOffset);
+			}
+
+			if (float SensorVerticalOffset; UsdUtils::GetAuthoredAttributeValue<float>(Prim, VerticalApertureOffsetToken, SensorVerticalOffset))
+			{
+				SensorVerticalOffset = UsdToUnreal::ConvertDistance(StageInfo, SensorVerticalOffset);
+				PhysicalCamera->SetCustomSensorVerticalOffset(SensorVerticalOffset);
+			}
+
+			if (float FocalLength; UsdUtils::GetAuthoredAttributeValue<float>(Prim, FocalLengthToken, FocalLength))
+			{
+				FocalLength = UsdToUnreal::ConvertDistance(StageInfo, FocalLength);
+				PhysicalCamera->SetCustomFocalLength(FocalLength);
+			}
+
+			if (FVector2f ClippingRange; UsdUtils::GetAuthoredAttributeValue<FVector2f>(Prim, ClippingRangeToken, ClippingRange))
+			{
+				ClippingRange.X = UsdToUnreal::ConvertDistance(StageInfo, ClippingRange.X);
+				PhysicalCamera->SetCustomNearClipPlane(ClippingRange.X);
+			}
 
-		if (float SensorWidth; UsdUtils::GetAuthoredAttributeValue<float>(Prim, HorizontalApertureToken, SensorWidth))
+			if (float Fstop; UsdUtils::GetAuthoredAttributeValue<float>(Prim, FstopToken, Fstop))
+			{
+				PhysicalCamera->SetCustomDepthOfFieldFstop(Fstop);
+
+				// Follows the documentation on the fStop attribute "Defaults to 0.0, which turns off depth of field effects."
+				// We check for timeSamples too however, otherwise we'd fully turn off depth of field for cameras that only had
+				// animated fStop values...
+				UE::FUsdAttribute FstopAttr = Prim.GetAttribute(*FstopToken);
+				if (FMath::IsNearlyZero(Fstop) && !FstopAttr.ValueMightBeTimeVarying())
+				{
+					PhysicalCamera->SetCustomEnableDepthOfField(false);
+				}
+			}
+
+			if (float FocusDistance; UsdUtils::GetAuthoredAttributeValue<float>(Prim, FocusDistanceToken, FocusDistance))
+			{
+				FocusDistance = UsdToUnreal::ConvertDistance(StageInfo, FocusDistance);
+				PhysicalCamera->SetCustomFocusDistance(FocusDistance);
+			}
+		}
+		else if (UInterchangeStandardCameraNode* StandardCamera = Cast<UInterchangeStandardCameraNode>(CameraNode))
 		{
-			SensorWidth = UsdToUnreal::ConvertDistance(StageInfo, SensorWidth);
-			CameraNode->SetCustomSensorWidth(SensorWidth);
+			StandardCamera->SetCustomProjectionMode(EInterchangeCameraProjectionType::Orthographic);
+
+			// We're a bit more verbose here so that we only try using an aspect ratio if we know we have both
+			// valid height and width values authored
+			if (float UsdWidth; UsdUtils::GetAuthoredAttributeValue<float>(Prim, HorizontalApertureToken, UsdWidth))
+			{
+				float Width = UsdToUnreal::ConvertDistance(StageInfo, UsdWidth);
+				StandardCamera->SetCustomWidth(Width);
+
+				if (float UsdHeight; UsdUtils::GetAuthoredAttributeValue<float>(Prim, VerticalApertureToken, UsdHeight))
+				{
+					float Height = UsdToUnreal::ConvertDistance(StageInfo, UsdHeight);
+
+					if (Height != 0.0f)
+					{
+						float AspectRatio = Width / Height;
+						StandardCamera->SetCustomAspectRatio(AspectRatio);
+					}
+				}
+			}
+
+			if (FVector2f ClippingRange; UsdUtils::GetAuthoredAttributeValue<FVector2f>(Prim, ClippingRangeToken, ClippingRange))
+			{
+				ClippingRange = UsdUtils::GetAttributeValue<FVector2f>(Prim, ClippingRangeToken);
+				ClippingRange.X = UsdToUnreal::ConvertDistance(StageInfo, ClippingRange.X);
+				ClippingRange.Y = UsdToUnreal::ConvertDistance(StageInfo, ClippingRange.Y);
+				StandardCamera->SetCustomNearClipPlane(ClippingRange.X);
+				StandardCamera->SetCustomFarClipPlane(ClippingRange.Y);
+			}
 		}
 
-		if (float SensorHeight; UsdUtils::GetAuthoredAttributeValue<float>(Prim, VerticalApertureToken, SensorHeight))
+		// Common attributes
+		if (CameraNode)
 		{
-			SensorHeight = UsdToUnreal::ConvertDistance(StageInfo, SensorHeight);
-			CameraNode->SetCustomSensorHeight(SensorHeight);
-		}
+			if (float ExposureCompensation; UsdUtils::GetAuthoredAttributeValue<float>(Prim, ExposureToken, ExposureCompensation))
+			{
+				CameraNode->SetCustomExposureCompensation(ExposureCompensation);
+			}
+
+			bool bNeedsManualExposure = false;
+
+			if (float ExposureIso; UsdUtils::GetAuthoredAttributeValue<float>(Prim, ExposureIsoToken, ExposureIso))
+			{
+				CameraNode->SetCustomExposureISO(ExposureIso);
+				bNeedsManualExposure = true;
+			}
+
+			if (float ExposureTime; UsdUtils::GetAuthoredAttributeValue<float>(Prim, ExposureTimeToken, ExposureTime))
+			{
+				CameraNode->SetCustomExposureShutterSpeed(ExposureTime != 0.0f ? (1.0f / ExposureTime) : 0.0f);
+				bNeedsManualExposure = true;
+			}
 
-		// Focus distance and FStop not currently supported
+			if (float ExposureFstop; UsdUtils::GetAuthoredAttributeValue<float>(Prim, ExposureFstopToken, ExposureFstop))
+			{
+				CameraNode->SetCustomExposureFstop(ExposureFstop);
+				bNeedsManualExposure = true;
+			}
+
+			// We likely want manual exposure if we have any of these authored.
+			// Note that ExposureCompensation itself doesn't matter though: It can still have an effect even with
+			// auto exposure enabled
+			if (bNeedsManualExposure)
+			{
+				CameraNode->SetCustomAutoExposureMethod(EInterchangeAutoExposureMethod::Manual);
+			}
+		}
 	}
 }	 // namespace UE::GeomCameraSchemaHandler::Private
 
@@ -86,31 +205,69 @@ namespace UE::Interchange::USD
 
 		const FString NewNodeUid = MakeNodeUid(CameraPrefix + Prim.GetPrimPath().GetString());
 		const FString NewNodeName{Prim.GetName().ToString()};
-		UInterchangePhysicalCameraNode* AssetNode = AccumulatedInfo.GetOrCreateAssetNode<UInterchangePhysicalCameraNode>(
-			*UsdContext.GetNodeContainer(),
-			NewNodeUid,
-			NewNodeName
-		);
+
+		// Let's try reusing whatever camera node was provided
+		UInterchangeBaseCameraNode* AssetNode = AccumulatedInfo.GetAssetNodeOfClass<UInterchangeBaseCameraNode>();
 		if (!AssetNode)
 		{
-			return false;
+			UInterchangeBaseNodeContainer* NodeContainer = UsdContext.GetNodeContainer();
+			if (!NodeContainer)
+			{
+				return false;
+			}
+
+			// If we get to create one ourselves then let's create a standard camera for orthographic projection,
+			// and use the physical camera for perspective
+			FName Projection = UsdUtils::GetAttributeValue<FName>(Prim, ProjectionToken);
+			if (Projection == TEXT("orthographic"))
+			{
+				AssetNode = NewObject<UInterchangeStandardCameraNode>(NodeContainer);
+			}
+			else
+			{
+				AssetNode = NewObject<UInterchangePhysicalCameraNode>(NodeContainer);
+			}
+
+			NodeContainer->SetupNode(AssetNode, NewNodeUid, NewNodeName, EInterchangeNodeContainerType::TranslatedAsset);
+			AccumulatedInfo.PrimAssetNodes.Add(AssetNode);
 		}
 
 		ConvertCameraPrim(Prim, AssetNode);
 
-		const static TArray<FString> UsdAttributeNames = {
-			UsdToUnreal::ConvertToken(pxr::UsdGeomTokens->focalLength),
-			UsdToUnreal::ConvertToken(pxr::UsdGeomTokens->focusDistance),
-			UsdToUnreal::ConvertToken(pxr::UsdGeomTokens->fStop),
-			UsdToUnreal::ConvertToken(pxr::UsdGeomTokens->horizontalAperture),
-			UsdToUnreal::ConvertToken(pxr::UsdGeomTokens->verticalAperture),
-			UsdToUnreal::ConvertToken(pxr::UsdGeomTokens->horizontalApertureOffset),
-			UsdToUnreal::ConvertToken(pxr::UsdGeomTokens->verticalApertureOffset),
-			UsdToUnreal::ConvertToken(pxr::UsdGeomTokens->exposure),
-			UsdToUnreal::ConvertToken(pxr::UsdGeomTokens->projection),
-			UsdToUnreal::ConvertToken(pxr::UsdGeomTokens->clippingRange)
+		const static TMap<FString, TArray<FInterchangeTrackInfo>> OrthoAttributeMapping = {
+			{HorizontalApertureToken, 	{{UnrealIdentifiers::OrthoWidthPropertyName, 			EInterchangePropertyTracks::CameraOrthoWidth},
+									   	 {UnrealIdentifiers::AspectRatioPropertyName, 			EInterchangePropertyTracks::CameraAspectRatio}}},
+			{VerticalApertureToken,   	{{UnrealIdentifiers::AspectRatioPropertyName, 			EInterchangePropertyTracks::CameraAspectRatio}}},
+			{ClippingRangeToken,      	{{UnrealIdentifiers::OrthoNearClipPlanePropertyName, 	EInterchangePropertyTracks::CameraOrthoNearClipPlane},
+									   	 {UnrealIdentifiers::OrthoFarClipPlanePropertyName, 	EInterchangePropertyTracks::CameraOrthoFarClipPlane}}},
+			{ExposureToken,      		{{UnrealIdentifiers::ExposureCompensationPropertyName, 	EInterchangePropertyTracks::CameraPostProcessSettingsAutoExposureBias}}},
+			{ExposureIsoToken,      	{{UnrealIdentifiers::CameraISOPropertyName, 			EInterchangePropertyTracks::CameraPostProcessSettingsCameraISO}}},
+			{ExposureTimeToken,      	{{UnrealIdentifiers::CameraShutterSpeedPropertyName, 	EInterchangePropertyTracks::CameraPostProcessSettingsCameraShutterSpeed}}},
+			{ExposureFstopToken,      	{{UnrealIdentifiers::DepthOfFieldFstopPropertyName, 	EInterchangePropertyTracks::CameraPostProcessSettingsDepthOfFieldFstop}}},
 		};
-		AddNodesForAnimatedAttributes(Prim, UsdAttributeNames, AccumulatedInfo, UsdContext);
+
+		const static TMap<FString, TArray<FInterchangeTrackInfo>> PerspectiveAttributeMapping = {
+			{HorizontalApertureToken, 		{{UnrealIdentifiers::SensorWidthPropertyName, 				EInterchangePropertyTracks::CameraFilmbackSensorWidth}}},
+			{VerticalApertureToken,   		{{UnrealIdentifiers::SensorHeightPropertyName, 				EInterchangePropertyTracks::CameraFilmbackSensorHeight}}},
+			{HorizontalApertureOffsetToken,	{{UnrealIdentifiers::SensorHorizontalOffsetPropertyName, 	EInterchangePropertyTracks::CameraFilmbackSensorHorizontalOffset}}},
+			{VerticalApertureOffsetToken, 	{{UnrealIdentifiers::SensorVerticalOffsetPropertyName, 		EInterchangePropertyTracks::CameraFilmbackSensorVerticalOffset}}},
+			{FocalLengthToken, 				{{UnrealIdentifiers::CurrentFocalLengthPropertyName, 		EInterchangePropertyTracks::CameraCurrentFocalLength}}},
+			{ClippingRangeToken,      		{{UnrealIdentifiers::CustomNearClipppingPlanePropertyName, 	EInterchangePropertyTracks::CameraCustomNearClippingPlane}}},
+			{FstopToken,      				{{UnrealIdentifiers::DepthOfFieldFstopPropertyName, 		EInterchangePropertyTracks::CameraPostProcessSettingsDepthOfFieldFstop}}},
+			{FocusDistanceToken,      		{{UnrealIdentifiers::ManualFocusDistancePropertyName, 		EInterchangePropertyTracks::CameraFocusSettingsManualFocusDistance}}},
+			{ExposureToken,      			{{UnrealIdentifiers::ExposureCompensationPropertyName, 		EInterchangePropertyTracks::CameraPostProcessSettingsAutoExposureBias}}},
+			{ExposureIsoToken,      		{{UnrealIdentifiers::CameraISOPropertyName, 				EInterchangePropertyTracks::CameraPostProcessSettingsCameraISO}}},
+			{ExposureTimeToken,      		{{UnrealIdentifiers::CameraShutterSpeedPropertyName, 		EInterchangePropertyTracks::CameraPostProcessSettingsCameraShutterSpeed}}},
+			{ExposureFstopToken,      		{{UnrealIdentifiers::CurrentAperturePropertyName, 			EInterchangePropertyTracks::CameraCurrentAperture}}},
+		};
+
+		const bool bNodeIsOrtho = AssetNode->IsA<UInterchangeStandardCameraNode>();
+		AddNodesForAnimatedAttributes(
+			Prim,
+			bNodeIsOrtho ? OrthoAttributeMapping : PerspectiveAttributeMapping, 
+			AccumulatedInfo, 
+			UsdContext
+		);
 
 		UInterchangeBaseNode* SceneNodeBase = AccumulatedInfo.GetOrCreateMainSceneNode(Prim, TraversalInfo, UsdContext);
 		if (UInterchangeSceneNode* SceneNode = Cast<UInterchangeSceneNode>(SceneNodeBase))
diff --git a/Engine/Plugins/Interchange/Extensions/OpenUSD/Source/Import/Private/SchemaHandlers/ImageableSchemaHandler.cpp b/Engine/Plugins/Interchange/Extensions/OpenUSD/Source/Import/Private/SchemaHandlers/ImageableSchemaHandler.cpp
--- a/Engine/Plugins/Interchange/Extensions/OpenUSD/Source/Import/Private/SchemaHandlers/ImageableSchemaHandler.cpp
+++ b/Engine/Plugins/Interchange/Extensions/OpenUSD/Source/Import/Private/SchemaHandlers/ImageableSchemaHandler.cpp
@@ -441,10 +441,11 @@ namespace UE::Interchange::USD
 		}
 
 		// Note: We're relying on the UsdGeomImageable handler to provide our animation payload data
-		const static TArray<FString> UsdAttributeNames = {
-			UsdToUnreal::ConvertToken(pxr::UsdGeomTokens->visibility),
+		// Binding visibility to the actor works better for cameras
+		const static TMap<FString, TArray<FInterchangeTrackInfo>> AttributeMapping = {
+			{UsdToUnreal::ConvertToken(pxr::UsdGeomTokens->visibility), {{UnrealIdentifiers::HiddenInGamePropertyName, 	EInterchangePropertyTracks::ActorHiddenInGame, EInterchangeAnimationPayLoadType::STEPCURVE}}},
 		};
-		AddNodesForAnimatedAttributes(Prim, UsdAttributeNames, AccumulatedInfo, UsdContext);
+		AddNodesForAnimatedAttributes(Prim, AttributeMapping, AccumulatedInfo, UsdContext);
 
 		// Store our purpose here to be filtered on the pipeline
 		const EUsdPurpose PrimPurpose = IUsdPrim::GetPurpose(Prim);
diff --git a/Engine/Plugins/Interchange/Extensions/OpenUSD/Source/Import/Private/SchemaHandlers/LightSchemaHandler.cpp b/Engine/Plugins/Interchange/Extensions/OpenUSD/Source/Import/Private/SchemaHandlers/LightSchemaHandler.cpp
--- a/Engine/Plugins/Interchange/Extensions/OpenUSD/Source/Import/Private/SchemaHandlers/LightSchemaHandler.cpp
+++ b/Engine/Plugins/Interchange/Extensions/OpenUSD/Source/Import/Private/SchemaHandlers/LightSchemaHandler.cpp
@@ -433,20 +433,80 @@ namespace UE::Interchange::USD
 
 		ConvertLightPrim(Prim, LightAssetNode, AccumulatedInfo, *UsdContext.GetNodeContainer());
 
-		const static TArray<FString> UsdAttributeNames = {
-			UsdToUnreal::ConvertToken(pxr::UsdLuxTokens->inputsIntensity),
-			UsdToUnreal::ConvertToken(pxr::UsdLuxTokens->inputsExposure),
-			UsdToUnreal::ConvertToken(pxr::UsdLuxTokens->inputsColor),
-			UsdToUnreal::ConvertToken(pxr::UsdLuxTokens->inputsColorTemperature),
-			UsdToUnreal::ConvertToken(pxr::UsdLuxTokens->inputsEnableColorTemperature),
-			UsdToUnreal::ConvertToken(pxr::UsdLuxTokens->inputsRadius),
-			UsdToUnreal::ConvertToken(pxr::UsdLuxTokens->inputsWidth),
-			UsdToUnreal::ConvertToken(pxr::UsdLuxTokens->inputsHeight),
-			UsdToUnreal::ConvertToken(pxr::UsdLuxTokens->inputsAngle),
-			UsdToUnreal::ConvertToken(pxr::UsdLuxTokens->inputsShapingConeAngle),
-			UsdToUnreal::ConvertToken(pxr::UsdLuxTokens->inputsShapingConeSoftness)
-		};
-		AddNodesForAnimatedAttributes(Prim, UsdAttributeNames, AccumulatedInfo, UsdContext);
+		// Add nodes for animated attributes
+		const TMap<FString, TArray<FInterchangeTrackInfo>>* AttributeMapping = nullptr;
+		{
+			const static TMap<FString, TArray<FInterchangeTrackInfo>> CommonAttributeMapping = {
+				{UsdToUnreal::ConvertToken(pxr::UsdLuxTokens->inputsIntensity), 				{{UnrealIdentifiers::IntensityPropertyName, 		EInterchangePropertyTracks::LightIntensity}}},
+				{UsdToUnreal::ConvertToken(pxr::UsdLuxTokens->inputsExposure), 					{{UnrealIdentifiers::IntensityPropertyName, 		EInterchangePropertyTracks::LightIntensity}}},
+				{UsdToUnreal::ConvertToken(pxr::UsdLuxTokens->inputsColor), 					{{UnrealIdentifiers::LightColorPropertyName, 		EInterchangePropertyTracks::LightColor}}},
+				{UsdToUnreal::ConvertToken(pxr::UsdLuxTokens->inputsColorTemperature), 			{{UnrealIdentifiers::TemperaturePropertyName, 		EInterchangePropertyTracks::LightTemperature}}},
+				{UsdToUnreal::ConvertToken(pxr::UsdLuxTokens->inputsEnableColorTemperature), 	{{UnrealIdentifiers::UseTemperaturePropertyName, 	EInterchangePropertyTracks::LightUseTemperature}}},
+			};
+	
+			// Per-light-type mapping
+			// Our main concern here is to capture the properties that make sense for this asset node. If we provide
+			// incorrect ones, the main consequence will be broken bindings on the LevelSequence and maybe some errors/warnings
+			if (AssetNode->IsA<UInterchangeSpotLightNode>())
+			{
+				const static TMap<FString, TArray<FInterchangeTrackInfo>> CombinedLightMapping = []()
+				{
+					TMap<FString, TArray<FInterchangeTrackInfo>> Combined = CommonAttributeMapping;
+					Combined.Add(UsdToUnreal::ConvertToken(pxr::UsdLuxTokens->inputsShapingConeAngle), 		{{UnrealIdentifiers::OuterConeAnglePropertyName, 	EInterchangePropertyTracks::LightOuterConeAngle},
+																											 {UnrealIdentifiers::IntensityPropertyName, 		EInterchangePropertyTracks::LightIntensity}});
+					Combined.Add(UsdToUnreal::ConvertToken(pxr::UsdLuxTokens->inputsShapingConeSoftness), 	{{UnrealIdentifiers::InnerConeAnglePropertyName, 	EInterchangePropertyTracks::LightInnerConeAngle},
+																											 {UnrealIdentifiers::IntensityPropertyName, 		EInterchangePropertyTracks::LightIntensity}});
+					return Combined;
+				}();
+				AttributeMapping = &CombinedLightMapping;
+			}
+			else if (AssetNode->IsA<UInterchangePointLightNode>())
+			{
+				const static TMap<FString, TArray<FInterchangeTrackInfo>> CombinedLightMapping = []()
+				{
+					TMap<FString, TArray<FInterchangeTrackInfo>> Combined = CommonAttributeMapping;
+					Combined.Add(UsdToUnreal::ConvertToken(pxr::UsdLuxTokens->inputsRadius), 	{{UnrealIdentifiers::SourceRadiusPropertyName, 		EInterchangePropertyTracks::LightSourceRadius},
+																								 {UnrealIdentifiers::IntensityPropertyName, 		EInterchangePropertyTracks::LightIntensity}});
+					return Combined;
+				}();
+				AttributeMapping = &CombinedLightMapping;
+			}
+			else if (AssetNode->IsA<UInterchangeRectLightNode>())
+			{
+				const static TMap<FString, TArray<FInterchangeTrackInfo>> CombinedLightMapping = []()
+				{
+					TMap<FString, TArray<FInterchangeTrackInfo>> Combined = CommonAttributeMapping;
+					// Animated inputsRadius on a RectLightNode means this is a DiskLight handled as a RectLight
+					Combined.Add(UsdToUnreal::ConvertToken(pxr::UsdLuxTokens->inputsRadius), 	{{UnrealIdentifiers::SourceWidthPropertyName, 		EInterchangePropertyTracks::LightSourceWidth},
+																								 {UnrealIdentifiers::SourceHeightPropertyName, 		EInterchangePropertyTracks::LightSourceHeight},
+																								 {UnrealIdentifiers::IntensityPropertyName, 		EInterchangePropertyTracks::LightIntensity}});
+					Combined.Add(UsdToUnreal::ConvertToken(pxr::UsdLuxTokens->inputsWidth), 	{{UnrealIdentifiers::SourceWidthPropertyName, 		EInterchangePropertyTracks::LightSourceWidth},
+																								 {UnrealIdentifiers::IntensityPropertyName, 		EInterchangePropertyTracks::LightIntensity}});
+					Combined.Add(UsdToUnreal::ConvertToken(pxr::UsdLuxTokens->inputsHeight), 	{{UnrealIdentifiers::SourceHeightPropertyName, 		EInterchangePropertyTracks::LightSourceHeight},
+																								 {UnrealIdentifiers::IntensityPropertyName, 		EInterchangePropertyTracks::LightIntensity}});
+					return Combined;
+				}();
+				AttributeMapping = &CombinedLightMapping;
+			}
+			else if (AssetNode->IsA<UInterchangeDirectionalLightNode>())
+			{
+				const static TMap<FString, TArray<FInterchangeTrackInfo>> CombinedLightMapping = []()
+				{
+					TMap<FString, TArray<FInterchangeTrackInfo>> Combined = CommonAttributeMapping;
+					Combined.Add(UsdToUnreal::ConvertToken(pxr::UsdLuxTokens->inputsAngle), 	{{UnrealIdentifiers::LightSourceAnglePropertyName, 	EInterchangePropertyTracks::LightSourceAngle}});
+					return Combined;
+				}();
+				AttributeMapping = &CombinedLightMapping;
+			}
+			else // UInterchangeSkyLightNode
+			{
+				AttributeMapping = &CommonAttributeMapping;
+			}
+		}
+		if (AttributeMapping != nullptr)
+		{
+			AddNodesForAnimatedAttributes(Prim, *AttributeMapping, AccumulatedInfo, UsdContext);
+		}
 
 		UInterchangeBaseNode* SceneNodeBase = AccumulatedInfo.GetOrCreateMainSceneNode(Prim, TraversalInfo, UsdContext);
 		if (UInterchangeSceneNode* SceneNode = Cast<UInterchangeSceneNode>(SceneNodeBase))
diff --git a/Engine/Plugins/Interchange/Extensions/OpenUSD/Source/Import/Private/SchemaHandlers/SchemaHandlerUtils.cpp b/Engine/Plugins/Interchange/Extensions/OpenUSD/Source/Import/Private/SchemaHandlers/SchemaHandlerUtils.cpp
--- a/Engine/Plugins/Interchange/Extensions/OpenUSD/Source/Import/Private/SchemaHandlers/SchemaHandlerUtils.cpp
+++ b/Engine/Plugins/Interchange/Extensions/OpenUSD/Source/Import/Private/SchemaHandlers/SchemaHandlerUtils.cpp
@@ -626,39 +626,9 @@ namespace UE::Interchange::USD
 		UsdContext.CurrentTrackSet->AddCustomAnimationTrackUid(AnimTrackNodeUid);
 	}
 
-	const static TMap<FName, EInterchangePropertyTracks> PropertyNameToTrackType = {
-		// Common properties
-		{UnrealIdentifiers::HiddenInGamePropertyName, 			EInterchangePropertyTracks::ActorHiddenInGame}, // Binding visibility to the actor works better for cameras
-
-		// Camera properties
-		// TODO: Need to add support for SensorHorizontalOffset and SensorVerticalOffset, once Interchange supports those
-		{UnrealIdentifiers::CurrentFocalLengthPropertyName, 		EInterchangePropertyTracks::CameraCurrentFocalLength},
-		{UnrealIdentifiers::ManualFocusDistancePropertyName, 		EInterchangePropertyTracks::CameraFocusSettingsManualFocusDistance},
-		{UnrealIdentifiers::CurrentAperturePropertyName, 			EInterchangePropertyTracks::CameraCurrentAperture},
-		{UnrealIdentifiers::SensorWidthPropertyName, 				EInterchangePropertyTracks::CameraFilmbackSensorWidth},
-		{UnrealIdentifiers::SensorHeightPropertyName, 				EInterchangePropertyTracks::CameraFilmbackSensorHeight},
-		{UnrealIdentifiers::ExposureCompensationPropertyName, 		EInterchangePropertyTracks::CameraPostProcessSettingsAutoExposureBias},
-		{UnrealIdentifiers::ProjectionModePropertyName, 			EInterchangePropertyTracks::CameraProjectionMode},
-		{UnrealIdentifiers::OrthoFarClipPlanePropertyName, 			EInterchangePropertyTracks::CameraOrthoFarClipPlane},
-		{UnrealIdentifiers::OrthoNearClipPlanePropertyName, 		EInterchangePropertyTracks::CameraOrthoNearClipPlane},
-		{UnrealIdentifiers::CustomNearClipppingPlanePropertyName,	EInterchangePropertyTracks::CameraCustomNearClippingPlane},
-
-		// Light properties
-		{UnrealIdentifiers::LightColorPropertyName, 		EInterchangePropertyTracks::LightColor},
-		{UnrealIdentifiers::TemperaturePropertyName, 		EInterchangePropertyTracks::LightTemperature},
-		{UnrealIdentifiers::UseTemperaturePropertyName, 	EInterchangePropertyTracks::LightUseTemperature},
-		{UnrealIdentifiers::SourceHeightPropertyName, 		EInterchangePropertyTracks::LightSourceHeight},
-		{UnrealIdentifiers::SourceWidthPropertyName, 		EInterchangePropertyTracks::LightSourceWidth},
-		{UnrealIdentifiers::SourceRadiusPropertyName, 		EInterchangePropertyTracks::LightSourceRadius},
-		{UnrealIdentifiers::OuterConeAnglePropertyName, 	EInterchangePropertyTracks::LightOuterConeAngle},
-		{UnrealIdentifiers::InnerConeAnglePropertyName, 	EInterchangePropertyTracks::LightInnerConeAngle},
-		{UnrealIdentifiers::LightSourceAnglePropertyName, 	EInterchangePropertyTracks::LightSourceAngle},
-		{UnrealIdentifiers::IntensityPropertyName, 			EInterchangePropertyTracks::LightIntensity},
-	};
-
 	void AddNodesForAnimatedAttributes(
 		const UE::FUsdPrim& Prim,
-		const TArray<FString>& UsdAttributeNames,
+		const TMap<FString, TArray<FInterchangeTrackInfo>>& UsdAttributeNameToTrackInfo,
 		FHandlerAccumulatedInfo& AccumulatedInfo,
 		UInterchangeUsdContext& UsdContext
 	)
@@ -673,34 +643,28 @@ namespace UE::Interchange::USD
 		}
 		const FString PrimPath = Prim.GetPrimPath().GetString();
 
-		for (const FString& UsdAttributeName : UsdAttributeNames)
+		for (const TPair<FString, TArray<FInterchangeTrackInfo>>& Pair : UsdAttributeNameToTrackInfo)
 		{
+			const FString& UsdAttributeName = Pair.Key;
+			const TArray<FInterchangeTrackInfo>& TrackInfos = Pair.Value;
+
 			UE::FUsdAttribute Attr = Prim.GetAttribute(*UsdAttributeName);
 			if (!Attr || !Attr.ValueMightBeTimeVarying() || Attr.GetNumTimeSamples() == 0)
 			{
 				continue;
 			}
 
-			// Emit a STEPCURVE in case of a bool track: CURVE is only for floats/doubles (c.f. FLevelSequenceHelper::PopulateAnimationTrack).
-			// For now we're lucky in that all possible results from GetPropertiesForAttribute() are either all not bool, or either all bool,
-			// so we can reuse this for all the different UEAttrNames we get from the same attribute
-			const FName AttrTypeName = Attr.GetTypeName();
-			const bool bIsBoolTrack = AttrTypeName == TEXT("bool") || AttrTypeName == TEXT("token");	// Visibility is a token track
-
-			TArray<FName> UEPropertyNames = UsdUtils::GetPropertiesForAttribute(Prim, Attr.GetName().ToString());
-			for (const FName& UEPropertyName : UEPropertyNames)
+			for (const FInterchangeTrackInfo& TrackInfo : TrackInfos)
 			{
-				const EInterchangePropertyTracks* FoundTrackType = PropertyNameToTrackType.Find(UEPropertyName);
-				if (!FoundTrackType)
-				{
-					continue;
-				}
-
-				const EInterchangeAnimationPayLoadType PayloadType = bIsBoolTrack ? EInterchangeAnimationPayLoadType::STEPCURVE
-																				  : EInterchangeAnimationPayLoadType::CURVE;
-
 				// Note: Implicit assumption that PrimPath == SceneNodeUid
-				AddPropertyAnimationNode(PrimPath, UEPropertyName, *FoundTrackType, PayloadType, AccumulatedInfo, UsdContext);
+				AddPropertyAnimationNode(
+					PrimPath,
+					TrackInfo.PropertyName,
+					TrackInfo.TrackType,
+					TrackInfo.PayloadType,
+					AccumulatedInfo,
+					UsdContext
+				);
 			}
 		}
 	}
diff --git a/Engine/Plugins/Interchange/Extensions/OpenUSD/Source/Import/Public/SchemaHandlers/SchemaHandlerUtils.h b/Engine/Plugins/Interchange/Extensions/OpenUSD/Source/Import/Public/SchemaHandlers/SchemaHandlerUtils.h
--- a/Engine/Plugins/Interchange/Extensions/OpenUSD/Source/Import/Public/SchemaHandlers/SchemaHandlerUtils.h
+++ b/Engine/Plugins/Interchange/Extensions/OpenUSD/Source/Import/Public/SchemaHandlers/SchemaHandlerUtils.h
@@ -145,13 +145,23 @@ namespace UE::Interchange::USD
 		UInterchangeUsdContext& UsdContext
 	);
 
+	struct FInterchangeTrackInfo
+	{
+		FName PropertyName;
+		EInterchangePropertyTracks TrackType;
+		EInterchangeAnimationPayLoadType PayloadType = EInterchangeAnimationPayLoadType::CURVE;
+	};
+
 	/**
-	 * Iterates through Prim's attributes with the given UsdAttributeNames, and invokes AddPropertyAnimationNode() for all the
-	 * attributes that contain animation
+	 * Iterates through UsdAttributeNameToTrackInfo and gets the Prim's attributes matching the map's keys. Then, for each animated attribute,
+	 * invokes AddPropertyAnimationNode() using the corresponding TrackInfo map value as arguments.
+	 *
+	 * Multiple FInterchangeTrackInfo can be provided because sometimes USD attributes map to multiple UE properties (e.g. if we animate
+	 * sensorWidth on an USD orthographic camera, we need to animate both OrthoWidth and AspectRatio in UE)
 	 */
 	UE_API void AddNodesForAnimatedAttributes(
 		const UE::FUsdPrim& Prim,
-		const TArray<FString>& UsdAttributeNames,
+		const TMap<FString, TArray<FInterchangeTrackInfo>>& UsdAttributeNameToTrackInfo,
 		FHandlerAccumulatedInfo& AccumulatedInfo,
 		UInterchangeUsdContext& UsdContext
 	);
diff --git a/Engine/Plugins/Interchange/Runtime/Source/FactoryNodes/Public/InterchangeCameraFactoryNode.h b/Engine/Plugins/Interchange/Runtime/Source/FactoryNodes/Public/InterchangeCameraFactoryNode.h
--- a/Engine/Plugins/Interchange/Runtime/Source/FactoryNodes/Public/InterchangeCameraFactoryNode.h
+++ b/Engine/Plugins/Interchange/Runtime/Source/FactoryNodes/Public/InterchangeCameraFactoryNode.h
@@ -2,9 +2,10 @@
 
 #pragma once
 
-#include "CoreMinimal.h"
 #include "InterchangeActorFactoryNode.h"
+#include "InterchangeCameraNode.h"
 
+#include "CoreMinimal.h"
 #if WITH_ENGINE
 	#include "CineCameraComponent.h"
 #endif
@@ -12,7 +13,176 @@
 #include "InterchangeCameraFactoryNode.generated.h"
 
 UCLASS(MinimalAPI, BlueprintType)
-class UInterchangePhysicalCameraFactoryNode : public UInterchangeActorFactoryNode
+class UInterchangeBaseCameraFactoryNode : public UInterchangeActorFactoryNode
+{
+	GENERATED_BODY()
+
+public:
+	UFUNCTION(BlueprintCallable, Category = "Interchange | Node | CameraFactory")
+	bool GetCustomNearClipPlane(float& AttributeValue) const
+	{
+		IMPLEMENT_NODE_ATTRIBUTE_GETTER(NearClipPlane, float);
+	}
+
+	UFUNCTION(BlueprintCallable, Category = "Interchange | Node | CameraFactory")
+	bool SetCustomNearClipPlane(const float& AttributeValue, bool bAddApplyDelegate = true)
+	{
+		IMPLEMENT_NODE_ATTRIBUTE_SETTER_WITH_CUSTOM_DELEGATE_WITH_CUSTOM_CLASS(UInterchangeBaseCameraFactoryNode, NearClipPlane, float, UCameraComponent);
+	}
+
+	UFUNCTION(BlueprintCallable, Category = "Interchange | Node | Camera")
+	bool GetCustomAutoExposureMethod(EInterchangeAutoExposureMethod& AttributeValue) const
+	{
+		IMPLEMENT_NODE_ATTRIBUTE_GETTER(AutoExposureMethod, EInterchangeAutoExposureMethod);
+	}
+
+	UFUNCTION(BlueprintCallable, Category = "Interchange | Node | Camera")
+	bool SetCustomAutoExposureMethod(const EInterchangeAutoExposureMethod& AttributeValue, bool bAddApplyDelegate = true)
+	{
+		IMPLEMENT_NODE_ATTRIBUTE_SETTER_WITH_CUSTOM_DELEGATE_WITH_CUSTOM_CLASS(UInterchangeBaseCameraFactoryNode, AutoExposureMethod, EInterchangeAutoExposureMethod, UCameraComponent);
+	}
+
+	UFUNCTION(BlueprintCallable, Category = "Interchange | Node | CameraFactory")
+	bool GetCustomExposureCompensation(float& AttributeValue) const
+	{
+		IMPLEMENT_NODE_ATTRIBUTE_GETTER(ExposureCompensation, float);
+	}
+
+	UFUNCTION(BlueprintCallable, Category = "Interchange | Node | CameraFactory")
+	bool SetCustomExposureCompensation(const float& AttributeValue, bool bAddApplyDelegate = true)
+	{
+		IMPLEMENT_NODE_ATTRIBUTE_SETTER_WITH_CUSTOM_DELEGATE_WITH_CUSTOM_CLASS(UInterchangeBaseCameraFactoryNode, ExposureCompensation, float, UCameraComponent);
+	}
+
+	UFUNCTION(BlueprintCallable, Category = "Interchange | Node | CameraFactory")
+	bool GetCustomExposureFstop(float& AttributeValue) const
+	{
+		IMPLEMENT_NODE_ATTRIBUTE_GETTER(ExposureFstop, float);
+	}
+
+	UFUNCTION(BlueprintCallable, Category = "Interchange | Node | CameraFactory")
+	bool SetCustomExposureFstop(const float& AttributeValue, bool bAddApplyDelegate = true)
+	{
+		IMPLEMENT_NODE_ATTRIBUTE_SETTER_WITH_CUSTOM_DELEGATE_WITH_CUSTOM_CLASS(UInterchangeBaseCameraFactoryNode, ExposureFstop, float, UCameraComponent);
+	}
+
+	UFUNCTION(BlueprintCallable, Category = "Interchange | Node | CameraFactory")
+	bool GetCustomExposureISO(float& AttributeValue) const
+	{
+		IMPLEMENT_NODE_ATTRIBUTE_GETTER(ExposureISO, float);
+	}
+
+	UFUNCTION(BlueprintCallable, Category = "Interchange | Node | CameraFactory")
+	bool SetCustomExposureISO(const float& AttributeValue, bool bAddApplyDelegate = true)
+	{
+		IMPLEMENT_NODE_ATTRIBUTE_SETTER_WITH_CUSTOM_DELEGATE_WITH_CUSTOM_CLASS(UInterchangeBaseCameraFactoryNode, ExposureISO, float, UCameraComponent);
+	}
+
+	UFUNCTION(BlueprintCallable, Category = "Interchange | Node | CameraFactory")
+	bool GetCustomExposureShutterSpeed(float& AttributeValue) const
+	{
+		IMPLEMENT_NODE_ATTRIBUTE_GETTER(ExposureShutterSpeed, float);
+	}
+
+	UFUNCTION(BlueprintCallable, Category = "Interchange | Node | CameraFactory")
+	bool SetCustomExposureShutterSpeed(const float& AttributeValue, bool bAddApplyDelegate = true)
+	{
+		IMPLEMENT_NODE_ATTRIBUTE_SETTER_WITH_CUSTOM_DELEGATE_WITH_CUSTOM_CLASS(UInterchangeBaseCameraFactoryNode, ExposureShutterSpeed, float, UCameraComponent);
+	}
+
+	virtual void CopyWithObject(const UInterchangeFactoryBaseNode* SourceNode, UObject* Object) override
+	{
+		Super::CopyWithObject(SourceNode, Object);
+
+		if (const UInterchangeBaseCameraFactoryNode* ActorFactoryNode = Cast<UInterchangeBaseCameraFactoryNode>(SourceNode))
+		{
+			COPY_NODE_DELEGATES_WITH_CUSTOM_DELEGATE(ActorFactoryNode, UInterchangeBaseCameraFactoryNode, NearClipPlane, float, UCameraComponent::StaticClass())
+			COPY_NODE_DELEGATES_WITH_CUSTOM_DELEGATE(ActorFactoryNode, UInterchangeBaseCameraFactoryNode, AutoExposureMethod, EInterchangeAutoExposureMethod, UCameraComponent::StaticClass())
+			COPY_NODE_DELEGATES_WITH_CUSTOM_DELEGATE(ActorFactoryNode, UInterchangeBaseCameraFactoryNode, ExposureCompensation, float, UCameraComponent::StaticClass())
+			COPY_NODE_DELEGATES_WITH_CUSTOM_DELEGATE(ActorFactoryNode, UInterchangeBaseCameraFactoryNode, ExposureFstop, float, UCameraComponent::StaticClass())
+			COPY_NODE_DELEGATES_WITH_CUSTOM_DELEGATE(ActorFactoryNode, UInterchangeBaseCameraFactoryNode, ExposureISO, float, UCameraComponent::StaticClass())
+			COPY_NODE_DELEGATES_WITH_CUSTOM_DELEGATE(ActorFactoryNode, UInterchangeBaseCameraFactoryNode, ExposureShutterSpeed, float, UCameraComponent::StaticClass())
+		}
+	}
+
+protected:
+	IMPLEMENT_NODE_ATTRIBUTE_KEY(NearClipPlane);
+	IMPLEMENT_NODE_ATTRIBUTE_KEY(AutoExposureMethod);
+	IMPLEMENT_NODE_ATTRIBUTE_KEY(ExposureCompensation);
+	IMPLEMENT_NODE_ATTRIBUTE_KEY(ExposureFstop);
+	IMPLEMENT_NODE_ATTRIBUTE_KEY(ExposureISO);
+	IMPLEMENT_NODE_ATTRIBUTE_KEY(ExposureShutterSpeed);
+
+	// We need to defer these to the derived classes as the implementation is slightly different for each
+	virtual bool ApplyCustomNearClipPlaneToAsset(UObject* Asset) const
+	{
+		// These must always be implemented by derived classes and we should never be calling this base version.
+		// We can't make these pure virtual though because some of these macros will need to instantiate this base class directly
+		return ensure(false);
+	}
+	virtual bool FillCustomNearClipPlaneFromAsset(UObject* Asset)
+	{
+		return ensure(false);
+	}
+
+	// We must implement our other PostProcessSettings delegates because we need to also set the bOverride properties
+	// alongside the property values themselves, as otherwise they will have no effect
+	bool ApplyCustomAutoExposureMethodToAsset(UObject* Asset) const
+	{
+		const bool bValue = true;
+		return ApplyPropertyValueToObject(bValue, Asset, TEXT("PostProcessSettings.bOverride_AutoExposureMethod")) &&
+		       ApplyAttributeToObject<EInterchangeAutoExposureMethod>(Macro_CustomAutoExposureMethodKey.Key, Asset, TEXT("PostProcessSettings.AutoExposureMethod"));
+	}
+	bool FillCustomAutoExposureMethodFromAsset(UObject* Asset)
+	{
+		return FillAttributeFromObject<EInterchangeAutoExposureMethod>(Macro_CustomAutoExposureMethodKey.Key, Asset, TEXT("PostProcessSettings.AutoExposureMethod"));
+	}
+
+	bool ApplyCustomExposureCompensationToAsset(UObject* Asset) const
+	{
+		const bool bValue = true;
+		return ApplyPropertyValueToObject(bValue, Asset, TEXT("PostProcessSettings.bOverride_AutoExposureBias")) &&
+		       ApplyAttributeToObject<float>(Macro_CustomExposureCompensationKey.Key, Asset, TEXT("PostProcessSettings.AutoExposureBias"));
+	}
+	bool FillCustomExposureCompensationFromAsset(UObject* Asset)
+	{
+		return FillAttributeFromObject<float>(Macro_CustomExposureCompensationKey.Key, Asset, TEXT("PostProcessSettings.AutoExposureBias"));
+	}
+
+	virtual bool ApplyCustomExposureFstopToAsset(UObject* Asset) const
+	{
+		return ensure(false);
+	}
+	virtual bool FillCustomExposureFstopFromAsset(UObject* Asset)
+	{
+		return ensure(false);
+	}
+
+	bool ApplyCustomExposureISOToAsset(UObject* Asset) const
+	{
+		const bool bValue = true;
+		return ApplyPropertyValueToObject(bValue, Asset, TEXT("PostProcessSettings.bOverride_CameraISO")) &&
+			   ApplyAttributeToObject<float>(Macro_CustomExposureISOKey.Key, Asset, TEXT("PostProcessSettings.CameraISO"));
+	}
+	bool FillCustomExposureISOFromAsset(UObject* Asset)
+	{
+		return FillAttributeFromObject<float>(Macro_CustomExposureISOKey.Key, Asset, TEXT("PostProcessSettings.CameraISO"));
+	}
+
+	bool ApplyCustomExposureShutterSpeedToAsset(UObject* Asset) const
+	{
+		const bool bValue = true;
+		return ApplyPropertyValueToObject(bValue, Asset, TEXT("PostProcessSettings.bOverride_CameraShutterSpeed")) &&
+			   ApplyAttributeToObject<float>(Macro_CustomExposureShutterSpeedKey.Key, Asset, TEXT("PostProcessSettings.CameraShutterSpeed"));
+	}
+	bool FillCustomExposureShutterSpeedFromAsset(UObject* Asset)
+	{
+		return FillAttributeFromObject<float>(Macro_CustomExposureShutterSpeedKey.Key, Asset, TEXT("PostProcessSettings.CameraShutterSpeed"));
+	}
+};
+
+UCLASS(MinimalAPI, BlueprintType)
+class UInterchangePhysicalCameraFactoryNode : public UInterchangeBaseCameraFactoryNode
 {
 	GENERATED_BODY()
 
@@ -29,6 +199,18 @@ public:
 		IMPLEMENT_NODE_ATTRIBUTE_SETTER_WITH_CUSTOM_DELEGATE_WITH_CUSTOM_CLASS(UInterchangePhysicalCameraFactoryNode, FocalLength, float, UCineCameraComponent);
 	}
 
+	UFUNCTION(BlueprintCallable, Category = "Interchange | Node | CameraFactory")
+	bool GetCustomFocusDistance(float& AttributeValue) const
+	{
+		IMPLEMENT_NODE_ATTRIBUTE_GETTER(FocusDistance, float);
+	}
+
+	UFUNCTION(BlueprintCallable, Category = "Interchange | Node | CameraFactory")
+	bool SetCustomFocusDistance(const float& AttributeValue, bool bAddApplyDelegate = true)
+	{
+		IMPLEMENT_NODE_ATTRIBUTE_SETTER_WITH_CUSTOM_DELEGATE_WITH_CUSTOM_CLASS(UInterchangePhysicalCameraFactoryNode, FocusDistance, float, UCineCameraComponent);
+	}
+
 	UFUNCTION(BlueprintCallable, Category = "Interchange | Node | CameraFactory")
 	bool GetCustomSensorWidth(float& AttributeValue) const
 	{
@@ -53,6 +235,30 @@ public:
 		IMPLEMENT_NODE_ATTRIBUTE_SETTER_WITH_CUSTOM_DELEGATE_WITH_CUSTOM_CLASS(UInterchangePhysicalCameraFactoryNode, SensorHeight, float, UCineCameraComponent);
 	}
 
+	UFUNCTION(BlueprintCallable, Category = "Interchange | Node | CameraFactory")
+	bool GetCustomSensorHorizontalOffset(float& AttributeValue) const
+	{
+		IMPLEMENT_NODE_ATTRIBUTE_GETTER(SensorHorizontalOffset, float);
+	}
+
+	UFUNCTION(BlueprintCallable, Category = "Interchange | Node | CameraFactory")
+	bool SetCustomSensorHorizontalOffset(const float& AttributeValue, bool bAddApplyDelegate = true)
+	{
+		IMPLEMENT_NODE_ATTRIBUTE_SETTER_WITH_CUSTOM_DELEGATE_WITH_CUSTOM_CLASS(UInterchangePhysicalCameraFactoryNode, SensorHorizontalOffset, float, UCineCameraComponent);
+	}
+
+	UFUNCTION(BlueprintCallable, Category = "Interchange | Node | CameraFactory")
+	bool GetCustomSensorVerticalOffset(float& AttributeValue) const
+	{
+		IMPLEMENT_NODE_ATTRIBUTE_GETTER(SensorVerticalOffset, float);
+	}
+
+	UFUNCTION(BlueprintCallable, Category = "Interchange | Node | CameraFactory")
+	bool SetCustomSensorVerticalOffset(const float& AttributeValue, bool bAddApplyDelegate = true)
+	{
+		IMPLEMENT_NODE_ATTRIBUTE_SETTER_WITH_CUSTOM_DELEGATE_WITH_CUSTOM_CLASS(UInterchangePhysicalCameraFactoryNode, SensorVerticalOffset, float, UCineCameraComponent);
+	}
+
 	UFUNCTION(BlueprintCallable, Category = "Interchange | Node | CameraFactory")
 	bool GetCustomFocusMethod(ECameraFocusMethod& AttributeValue) const
 	{
@@ -63,7 +269,19 @@ public:
 	bool SetCustomFocusMethod(const ECameraFocusMethod& AttributeValue, bool bAddApplyDelegate = true)
 	{
 		IMPLEMENT_NODE_ATTRIBUTE_SETTER_WITH_CUSTOM_DELEGATE_WITH_CUSTOM_CLASS(UInterchangePhysicalCameraFactoryNode, FocusMethod, ECameraFocusMethod, UCineCameraComponent);
-	}	
+	}
+
+	UFUNCTION(BlueprintCallable, Category = "Interchange | Node | CameraFactory")
+	bool GetCustomDepthOfFieldFstop(float& AttributeValue) const
+	{
+		IMPLEMENT_NODE_ATTRIBUTE_GETTER(DepthOfFieldFstop, float);
+	}
+
+	UFUNCTION(BlueprintCallable, Category = "Interchange | Node | CameraFactory")
+	bool SetCustomDepthOfFieldFstop(const float& AttributeValue, bool bAddApplyDelegate = true)
+	{
+		IMPLEMENT_NODE_ATTRIBUTE_SETTER_WITH_CUSTOM_DELEGATE_WITH_CUSTOM_CLASS(UInterchangePhysicalCameraFactoryNode, DepthOfFieldFstop, float, UCineCameraComponent);
+	}
 
 	virtual void CopyWithObject(const UInterchangeFactoryBaseNode* SourceNode, UObject* Object) override
 	{
@@ -72,42 +290,83 @@ public:
 		if (const UInterchangePhysicalCameraFactoryNode* ActorFactoryNode = Cast<UInterchangePhysicalCameraFactoryNode>(SourceNode))
 		{
 			COPY_NODE_DELEGATES_WITH_CUSTOM_DELEGATE(ActorFactoryNode, UInterchangePhysicalCameraFactoryNode, FocalLength, float, UCineCameraComponent::StaticClass())
+			COPY_NODE_DELEGATES_WITH_CUSTOM_DELEGATE(ActorFactoryNode, UInterchangePhysicalCameraFactoryNode, FocusDistance, float, UCineCameraComponent::StaticClass())
 			COPY_NODE_DELEGATES_WITH_CUSTOM_DELEGATE(ActorFactoryNode, UInterchangePhysicalCameraFactoryNode, SensorWidth, float, UCineCameraComponent::StaticClass())
 			COPY_NODE_DELEGATES_WITH_CUSTOM_DELEGATE(ActorFactoryNode, UInterchangePhysicalCameraFactoryNode, SensorHeight, float, UCineCameraComponent::StaticClass())
+			COPY_NODE_DELEGATES_WITH_CUSTOM_DELEGATE(ActorFactoryNode, UInterchangePhysicalCameraFactoryNode, SensorHorizontalOffset, float, UCineCameraComponent::StaticClass())
+			COPY_NODE_DELEGATES_WITH_CUSTOM_DELEGATE(ActorFactoryNode, UInterchangePhysicalCameraFactoryNode, SensorVerticalOffset, float, UCineCameraComponent::StaticClass())
 			COPY_NODE_DELEGATES_WITH_CUSTOM_DELEGATE(ActorFactoryNode, UInterchangePhysicalCameraFactoryNode, FocusMethod, ECameraFocusMethod, UCineCameraComponent::StaticClass())
+			COPY_NODE_DELEGATES_WITH_CUSTOM_DELEGATE(ActorFactoryNode, UInterchangePhysicalCameraFactoryNode, DepthOfFieldFstop, float, UCineCameraComponent::StaticClass())
 		}
 	}
 
 private:
-	const UE::Interchange::FAttributeKey Macro_CustomFocalLengthKey = UE::Interchange::FAttributeKey(TEXT("FocalLength"));
-	const UE::Interchange::FAttributeKey Macro_CustomSensorWidthKey = UE::Interchange::FAttributeKey(TEXT("SensorWidth"));
-	const UE::Interchange::FAttributeKey Macro_CustomSensorHeightKey = UE::Interchange::FAttributeKey(TEXT("SensorHeight"));
-	const UE::Interchange::FAttributeKey Macro_CustomFocusMethodKey = UE::Interchange::FAttributeKey(TEXT("FocusMethod"));
+	IMPLEMENT_NODE_ATTRIBUTE_KEY(FocalLength);
+	IMPLEMENT_NODE_ATTRIBUTE_KEY(FocusDistance);
+	IMPLEMENT_NODE_ATTRIBUTE_KEY(SensorWidth);
+	IMPLEMENT_NODE_ATTRIBUTE_KEY(SensorHeight);
+	IMPLEMENT_NODE_ATTRIBUTE_KEY(SensorHorizontalOffset);
+	IMPLEMENT_NODE_ATTRIBUTE_KEY(SensorVerticalOffset);
+	IMPLEMENT_NODE_ATTRIBUTE_KEY(FocusMethod);
+	IMPLEMENT_NODE_ATTRIBUTE_KEY(DepthOfFieldFstop);
 
-private:
 	IMPLEMENT_NODE_ATTRIBUTE_DELEGATE_BY_PROPERTYNAME(FocalLength, float, UCineCameraComponent, TEXT("CurrentFocalLength"));
 	IMPLEMENT_NODE_ATTRIBUTE_DELEGATE_BY_PROPERTYNAME(SensorWidth, float, UCineCameraComponent, TEXT("Filmback.SensorWidth"));
 	IMPLEMENT_NODE_ATTRIBUTE_DELEGATE_BY_PROPERTYNAME(SensorHeight, float, UCineCameraComponent, TEXT("Filmback.SensorHeight"));
+	IMPLEMENT_NODE_ATTRIBUTE_DELEGATE_BY_PROPERTYNAME(SensorHorizontalOffset, float, UCineCameraComponent, TEXT("Filmback.SensorHorizontalOffset"));
+	IMPLEMENT_NODE_ATTRIBUTE_DELEGATE_BY_PROPERTYNAME(SensorVerticalOffset, float, UCineCameraComponent, TEXT("Filmback.SensorVerticalOffset"));
 	IMPLEMENT_NODE_ATTRIBUTE_DELEGATE_BY_PROPERTYNAME(FocusMethod, ECameraFocusMethod, UCineCameraComponent, TEXT("FocusSettings.FocusMethod"));
-};
+	IMPLEMENT_NODE_ATTRIBUTE_DELEGATE_BY_PROPERTYNAME(FocusDistance, float, UCineCameraComponent, TEXT("FocusSettings.ManualFocusDistance"));
+
+	virtual bool ApplyCustomNearClipPlaneToAsset(UObject* Asset) const override
+	{
+		const bool bValue = true;
+		return ApplyPropertyValueToObject(bValue, Asset, TEXT("bOverride_CustomNearClippingPlane")) &&
+			   ApplyAttributeToObject<float>(Macro_CustomNearClipPlaneKey.Key, Asset, TEXT("CustomNearClippingPlane"));
+	}
+	virtual bool FillCustomNearClipPlaneFromAsset(UObject* Asset) override
+	{
+		return FillAttributeFromObject<float>(Macro_CustomNearClipPlaneKey.Key, Asset, TEXT("CustomNearClippingPlane"));
+	}
 
+	virtual bool ApplyCustomExposureFstopToAsset(UObject* Asset) const override
+	{
+		// On the cine camera component, CurrentAperture is used for exposure, and PostProcessSettings.DepthOfFieldFstop for DoF
+		return ApplyAttributeToObject<float>(Macro_CustomExposureFstopKey.Key, Asset, TEXT("CurrentAperture"));
+	}
+	virtual bool FillCustomExposureFstopFromAsset(UObject* Asset) override
+	{
+		return FillAttributeFromObject<float>(Macro_CustomExposureFstopKey.Key, Asset, TEXT("CurrentAperture"));
+	}
+
+	bool ApplyCustomDepthOfFieldFstopToAsset(UObject* Asset) const
+	{
+		const bool bValue = true;
+		return ApplyPropertyValueToObject(bValue, Asset, TEXT("PostProcessSettings.bOverride_DepthOfFieldFstop")) &&
+		ApplyAttributeToObject<float>(Macro_CustomDepthOfFieldFstopKey.Key, Asset, TEXT("PostProcessSettings.DepthOfFieldFstop"));
+	}
+	bool FillCustomDepthOfFieldFstopFromAsset(UObject* Asset)
+	{
+		return FillAttributeFromObject<float>(Macro_CustomDepthOfFieldFstopKey.Key, Asset, TEXT("PostProcessSettings.DepthOfFieldFstop"));
+	}
+};
 
 /*
 * Used for common non-physical camera with orthographic or perspective projection.
 */
 UCLASS(MinimalAPI, BlueprintType)
-class UInterchangeStandardCameraFactoryNode : public UInterchangeActorFactoryNode
+class UInterchangeStandardCameraFactoryNode : public UInterchangeBaseCameraFactoryNode
 {
 	GENERATED_BODY()
 
 public:
 	UFUNCTION(BlueprintCallable, Category = "Interchange | Node | CameraFactory")
-		bool GetCustomProjectionMode(TEnumAsByte<ECameraProjectionMode::Type>& AttributeValue) const
+	bool GetCustomProjectionMode(TEnumAsByte<ECameraProjectionMode::Type>& AttributeValue) const
 	{
 		IMPLEMENT_NODE_ATTRIBUTE_GETTER(ProjectionMode, TEnumAsByte<ECameraProjectionMode::Type>);
 	}
 	UFUNCTION(BlueprintCallable, Category = "Interchange | Node | CameraFactory")
-		bool SetCustomProjectionMode(const TEnumAsByte<ECameraProjectionMode::Type>& AttributeValue, bool bAddApplyDelegate = true)
+	bool SetCustomProjectionMode(const TEnumAsByte<ECameraProjectionMode::Type>& AttributeValue, bool bAddApplyDelegate = true)
 	{
 		IMPLEMENT_NODE_ATTRIBUTE_SETTER_WITH_CUSTOM_DELEGATE_WITH_CUSTOM_CLASS(UInterchangeStandardCameraFactoryNode, ProjectionMode, TEnumAsByte<ECameraProjectionMode::Type>, UCameraComponent);
 	}
@@ -123,17 +382,6 @@ public:
 		IMPLEMENT_NODE_ATTRIBUTE_SETTER_WITH_CUSTOM_DELEGATE_WITH_CUSTOM_CLASS(UInterchangeStandardCameraFactoryNode, Width, float, UCameraComponent);
 	}
 
-	UFUNCTION(BlueprintCallable, Category = "Interchange | Node | CameraFactory")
-	bool GetCustomNearClipPlane(float& AttributeValue) const
-	{
-		IMPLEMENT_NODE_ATTRIBUTE_GETTER(NearClipPlane, float);
-	}
-	UFUNCTION(BlueprintCallable, Category = "Interchange | Node | CameraFactory")
-	bool SetCustomNearClipPlane(const float& AttributeValue, bool bAddApplyDelegate = true)
-	{
-		IMPLEMENT_NODE_ATTRIBUTE_SETTER_WITH_CUSTOM_DELEGATE_WITH_CUSTOM_CLASS(UInterchangeStandardCameraFactoryNode, NearClipPlane, float, UCameraComponent);
-	}
-
 	UFUNCTION(BlueprintCallable, Category = "Interchange | Node | CameraFactory")
 	bool GetCustomFarClipPlane(float& AttributeValue) const
 	{
@@ -157,12 +405,12 @@ public:
 	}
 
 	UFUNCTION(BlueprintCallable, Category = "Interchange | Node | CameraFactory")
-		bool GetCustomFieldOfView(float& AttributeValue) const
+	bool GetCustomFieldOfView(float& AttributeValue) const
 	{
 		IMPLEMENT_NODE_ATTRIBUTE_GETTER(FieldOfView, float);
 	}
 	UFUNCTION(BlueprintCallable, Category = "Interchange | Node | CameraFactory")
-		bool SetCustomFieldOfView(const float& AttributeValue, bool bAddApplyDelegate = true)
+	bool SetCustomFieldOfView(const float& AttributeValue, bool bAddApplyDelegate = true)
 	{
 		IMPLEMENT_NODE_ATTRIBUTE_SETTER_WITH_CUSTOM_DELEGATE_WITH_CUSTOM_CLASS(UInterchangeStandardCameraFactoryNode, FieldOfView, float, UCameraComponent);
 	}
@@ -174,33 +422,56 @@ public:
 		if (const UInterchangeStandardCameraFactoryNode* ActorFactoryNode = Cast<UInterchangeStandardCameraFactoryNode>(SourceNode))
 		{
 			COPY_NODE_DELEGATES_WITH_CUSTOM_DELEGATE(ActorFactoryNode, UInterchangeStandardCameraFactoryNode, ProjectionMode, TEnumAsByte<ECameraProjectionMode::Type>, UCameraComponent::StaticClass())
-
 			COPY_NODE_DELEGATES_WITH_CUSTOM_DELEGATE(ActorFactoryNode, UInterchangeStandardCameraFactoryNode, Width, float, UCameraComponent::StaticClass())
-			COPY_NODE_DELEGATES_WITH_CUSTOM_DELEGATE(ActorFactoryNode, UInterchangeStandardCameraFactoryNode, NearClipPlane, float, UCameraComponent::StaticClass())
 			COPY_NODE_DELEGATES_WITH_CUSTOM_DELEGATE(ActorFactoryNode, UInterchangeStandardCameraFactoryNode, FarClipPlane, float, UCameraComponent::StaticClass())
 			COPY_NODE_DELEGATES_WITH_CUSTOM_DELEGATE(ActorFactoryNode, UInterchangeStandardCameraFactoryNode, AspectRatio, float, UCameraComponent::StaticClass())
-
 			COPY_NODE_DELEGATES_WITH_CUSTOM_DELEGATE(ActorFactoryNode, UInterchangeStandardCameraFactoryNode, FieldOfView, float, UCameraComponent::StaticClass())
 		}
 	}
 
 private:
-	const UE::Interchange::FAttributeKey Macro_CustomProjectionModeKey = UE::Interchange::FAttributeKey(TEXT("ProjectionMode"));
+	IMPLEMENT_NODE_ATTRIBUTE_KEY(ProjectionMode);
+	IMPLEMENT_NODE_ATTRIBUTE_KEY(Width);
+	IMPLEMENT_NODE_ATTRIBUTE_KEY(FarClipPlane);
+	IMPLEMENT_NODE_ATTRIBUTE_KEY(AspectRatio);
+	IMPLEMENT_NODE_ATTRIBUTE_KEY(FieldOfView);
 
-	const UE::Interchange::FAttributeKey Macro_CustomWidthKey = UE::Interchange::FAttributeKey(TEXT("Width"));
-	const UE::Interchange::FAttributeKey Macro_CustomNearClipPlaneKey = UE::Interchange::FAttributeKey(TEXT("NearClip"));
-	const UE::Interchange::FAttributeKey Macro_CustomFarClipPlaneKey = UE::Interchange::FAttributeKey(TEXT("FarClipPlane"));
-	const UE::Interchange::FAttributeKey Macro_CustomAspectRatioKey = UE::Interchange::FAttributeKey(TEXT("AspectRatio"));
-	
-	const UE::Interchange::FAttributeKey Macro_CustomFieldOfViewKey = UE::Interchange::FAttributeKey(TEXT("FieldOfView"));
-
-private:
 	IMPLEMENT_NODE_ATTRIBUTE_DELEGATE_BY_PROPERTYNAME(ProjectionMode, TEnumAsByte<ECameraProjectionMode::Type>, UCameraComponent, TEXT("ProjectionMode"));
-
 	IMPLEMENT_NODE_ATTRIBUTE_DELEGATE_BY_PROPERTYNAME(Width, float, UCameraComponent, TEXT("OrthoWidth"));
-	IMPLEMENT_NODE_ATTRIBUTE_DELEGATE_BY_PROPERTYNAME(NearClipPlane, float, UCameraComponent, TEXT("OrthoNearClipPlane"));
-	IMPLEMENT_NODE_ATTRIBUTE_DELEGATE_BY_PROPERTYNAME(FarClipPlane, float, UCameraComponent, TEXT("OrthoFarClipPlane"));
 	IMPLEMENT_NODE_ATTRIBUTE_DELEGATE_BY_PROPERTYNAME(AspectRatio, float, UCameraComponent, TEXT("AspectRatio"));
-
 	IMPLEMENT_NODE_ATTRIBUTE_DELEGATE_BY_PROPERTYNAME(FieldOfView, float, UCameraComponent, TEXT("FieldOfView"));
+
+	virtual bool ApplyCustomNearClipPlaneToAsset(UObject* Asset) const override
+	{
+		const bool bValue = false;
+		return ApplyPropertyValueToObject(bValue, Asset, TEXT("bAutoCalculateOrthoPlanes")) &&
+			   ApplyAttributeToObject<float>(Macro_CustomNearClipPlaneKey.Key, Asset, TEXT("OrthoNearClipPlane"));
+	}
+	virtual bool FillCustomNearClipPlaneFromAsset(UObject* Asset) override
+	{
+		return FillAttributeFromObject<float>(Macro_CustomNearClipPlaneKey.Key, Asset, TEXT("OrthoNearClipPlane"));
+	}
+
+	virtual bool ApplyCustomExposureFstopToAsset(UObject* Asset) const override
+	{
+		// On the regular camera component, PostProcessSettings.DepthOfFieldFstop is used for exposure, and CurrentAperture doesn't exist
+		const bool bValue = true;
+		return ApplyPropertyValueToObject(bValue, Asset, TEXT("PostProcessSettings.bOverride_DepthOfFieldFstop")) &&
+			   ApplyAttributeToObject<float>(Macro_CustomExposureFstopKey.Key, Asset, TEXT("PostProcessSettings.DepthOfFieldFstop"));
+	}
+	virtual bool FillCustomExposureFstopFromAsset(UObject* Asset) override
+	{
+		return FillAttributeFromObject<float>(Macro_CustomExposureFstopKey.Key, Asset, TEXT("PostProcessSettings.DepthOfFieldFstop"));
+	}
+
+	bool ApplyCustomFarClipPlaneToAsset(UObject* Asset) const
+	{
+		const bool bValue = false;
+		return ApplyPropertyValueToObject(bValue, Asset, TEXT("bAutoCalculateOrthoPlanes")) &&
+			   ApplyAttributeToObject<float>(Macro_CustomFarClipPlaneKey.Key, Asset, TEXT("OrthoFarClipPlane"));
+	}
+	bool FillCustomFarClipPlaneFromAsset(UObject* Asset)
+	{
+		return FillAttributeFromObject<float>(Macro_CustomFarClipPlaneKey.Key, Asset, TEXT("OrthoFarClipPlane"));
+	}
 };
diff --git a/Engine/Plugins/Interchange/Runtime/Source/Import/Private/Animation/InterchangeLevelSequenceHelper.cpp b/Engine/Plugins/Interchange/Runtime/Source/Import/Private/Animation/InterchangeLevelSequenceHelper.cpp
--- a/Engine/Plugins/Interchange/Runtime/Source/Import/Private/Animation/InterchangeLevelSequenceHelper.cpp
+++ b/Engine/Plugins/Interchange/Runtime/Source/Import/Private/Animation/InterchangeLevelSequenceHelper.cpp
@@ -162,6 +162,8 @@ namespace UE::Interchange::Private
 			{EInterchangePropertyTracks::CameraFieldOfView, {UMovieSceneFloatTrack::StaticClass()->GetName(),  GET_MEMBER_NAME_STRING_CHECKED(UCameraComponent, FieldOfView) , TEXT("Field Of View")}},
 			{EInterchangePropertyTracks::CameraFilmbackSensorAspectRatio, {UMovieSceneFloatTrack::StaticClass()->GetName(),  GET_MEMBER_NAME_STRING_CHECKED(UCineCameraComponent, Filmback) + FString{ TEXT(".") } + GET_MEMBER_NAME_STRING_CHECKED(FCameraFilmbackSettings, SensorAspectRatio), TEXT("Sensor Aspect Ratio (Filmback)")}},
 			{EInterchangePropertyTracks::CameraFilmbackSensorHeight, {UMovieSceneFloatTrack::StaticClass()->GetName(),  GET_MEMBER_NAME_STRING_CHECKED(UCineCameraComponent, Filmback) + FString{ TEXT(".") } + GET_MEMBER_NAME_STRING_CHECKED(FCameraFilmbackSettings, SensorHeight), TEXT("Sensor Height (Filmback)")}},
+			{EInterchangePropertyTracks::CameraFilmbackSensorHorizontalOffset, {UMovieSceneFloatTrack::StaticClass()->GetName(),  GET_MEMBER_NAME_STRING_CHECKED(UCineCameraComponent, Filmback) + FString{ TEXT(".") } + GET_MEMBER_NAME_STRING_CHECKED(FCameraFilmbackSettings, SensorHorizontalOffset), TEXT("Sensor Horizontal Offset (Filmback)")}},
+			{EInterchangePropertyTracks::CameraFilmbackSensorVerticalOffset, {UMovieSceneFloatTrack::StaticClass()->GetName(),  GET_MEMBER_NAME_STRING_CHECKED(UCineCameraComponent, Filmback) + FString{ TEXT(".") } + GET_MEMBER_NAME_STRING_CHECKED(FCameraFilmbackSettings, SensorVerticalOffset), TEXT("Sensor Vertical Offset (Filmback)")}},
 			{EInterchangePropertyTracks::CameraFilmbackSensorWidth, {UMovieSceneFloatTrack::StaticClass()->GetName(),  GET_MEMBER_NAME_STRING_CHECKED(UCineCameraComponent, Filmback) + FString{ TEXT(".") } + GET_MEMBER_NAME_STRING_CHECKED(FCameraFilmbackSettings, SensorWidth), TEXT("Sensor Width (Filmback)")}},
 			{EInterchangePropertyTracks::CameraFocusSettingsFocusOffset, {UMovieSceneFloatTrack::StaticClass()->GetName(),  GET_MEMBER_NAME_STRING_CHECKED(UCineCameraComponent, FocusSettings) + FString{ TEXT(".") } + GET_MEMBER_NAME_STRING_CHECKED(FCameraFocusSettings, FocusOffset), TEXT("Focus Offset (Focus Settings)")}},
 			{EInterchangePropertyTracks::CameraFocusSettingsManualFocusDistance, {UMovieSceneFloatTrack::StaticClass()->GetName(),  GET_MEMBER_NAME_STRING_CHECKED(UCineCameraComponent, FocusSettings) + FString{ TEXT(".") } + GET_MEMBER_NAME_STRING_CHECKED(FCameraFocusSettings, ManualFocusDistance), TEXT("Manual Focus Distance (Focus Settings)")}},
@@ -217,6 +219,8 @@ namespace UE::Interchange::Private
 			{EInterchangePropertyTracks::CameraPostProcessSettingsBloomSizeScale, {UMovieSceneFloatTrack::StaticClass()->GetName(),  GET_MEMBER_NAME_STRING_CHECKED(UCameraComponent, PostProcessSettings) + FString{ TEXT(".") } + GET_MEMBER_NAME_STRING_CHECKED(FPostProcessSettings, BloomSizeScale), TEXT("Size scale (Post Process Settings)")}},
 			{EInterchangePropertyTracks::CameraPostProcessSettingsBloomThreshold, {UMovieSceneFloatTrack::StaticClass()->GetName(),  GET_MEMBER_NAME_STRING_CHECKED(UCameraComponent, PostProcessSettings) + FString{ TEXT(".") } + GET_MEMBER_NAME_STRING_CHECKED(FPostProcessSettings, BloomThreshold), TEXT("Threshold (Post Process Settings)")}},
 			{EInterchangePropertyTracks::CameraPostProcessSettingsBlueCorrection, {UMovieSceneFloatTrack::StaticClass()->GetName(),  GET_MEMBER_NAME_STRING_CHECKED(UCameraComponent, PostProcessSettings) + FString{ TEXT(".") } + GET_MEMBER_NAME_STRING_CHECKED(FPostProcessSettings, BlueCorrection), TEXT("Blue Correction (Post Process Settings)")}},
+			{EInterchangePropertyTracks::CameraPostProcessSettingsCameraISO, {UMovieSceneFloatTrack::StaticClass()->GetName(),  GET_MEMBER_NAME_STRING_CHECKED(UCameraComponent, PostProcessSettings) + FString{ TEXT(".") } + GET_MEMBER_NAME_STRING_CHECKED(FPostProcessSettings, CameraISO), TEXT("Camera ISO (Post Process Settings)")}},
+			{EInterchangePropertyTracks::CameraPostProcessSettingsCameraShutterSpeed, {UMovieSceneFloatTrack::StaticClass()->GetName(),  GET_MEMBER_NAME_STRING_CHECKED(UCameraComponent, PostProcessSettings) + FString{ TEXT(".") } + GET_MEMBER_NAME_STRING_CHECKED(FPostProcessSettings, CameraShutterSpeed), TEXT("Camera Shutter Speed (Post Process Settings)")}},
 			{EInterchangePropertyTracks::CameraPostProcessSettingsChromaticAberrationStartOffset, {UMovieSceneFloatTrack::StaticClass()->GetName(),  GET_MEMBER_NAME_STRING_CHECKED(UCameraComponent, PostProcessSettings) + FString{ TEXT(".") } + GET_MEMBER_NAME_STRING_CHECKED(FPostProcessSettings, ChromaticAberrationStartOffset), TEXT("Start Offset (Post Process Settings)")}},
 			{EInterchangePropertyTracks::CameraPostProcessSettingsColorContrast, {UMovieSceneDoubleVectorTrack::StaticClass()->GetName(),  GET_MEMBER_NAME_STRING_CHECKED(UCameraComponent, PostProcessSettings) + FString{ TEXT(".") } + GET_MEMBER_NAME_STRING_CHECKED(FPostProcessSettings, ColorContrast), TEXT("Contrast (Post Process Settings)"), 4}},
 			{EInterchangePropertyTracks::CameraPostProcessSettingsColorContrastHighlights, {UMovieSceneDoubleVectorTrack::StaticClass()->GetName(),  GET_MEMBER_NAME_STRING_CHECKED(UCameraComponent, PostProcessSettings) + FString{ TEXT(".") } + GET_MEMBER_NAME_STRING_CHECKED(FPostProcessSettings, ColorContrastHighlights), TEXT("Contrast (Post Process Settings)"), 4}},
diff --git a/Engine/Plugins/Interchange/Runtime/Source/Nodes/Public/InterchangeAnimationDefinitions.h b/Engine/Plugins/Interchange/Runtime/Source/Nodes/Public/InterchangeAnimationDefinitions.h
--- a/Engine/Plugins/Interchange/Runtime/Source/Nodes/Public/InterchangeAnimationDefinitions.h
+++ b/Engine/Plugins/Interchange/Runtime/Source/Nodes/Public/InterchangeAnimationDefinitions.h
@@ -138,6 +138,8 @@ enum class EInterchangePropertyTracks : int32
 	CameraFieldOfView,
 	CameraFilmbackSensorAspectRatio,
 	CameraFilmbackSensorHeight,
+	CameraFilmbackSensorHorizontalOffset,
+	CameraFilmbackSensorVerticalOffset,
 	CameraFilmbackSensorWidth,
 	CameraFocusSettingsFocusOffset,
 	CameraFocusSettingsManualFocusDistance,
@@ -193,6 +195,8 @@ enum class EInterchangePropertyTracks : int32
 	CameraPostProcessSettingsBloomSizeScale,
 	CameraPostProcessSettingsBloomThreshold,
 	CameraPostProcessSettingsBlueCorrection,
+	CameraPostProcessSettingsCameraISO,
+	CameraPostProcessSettingsCameraShutterSpeed,
 	CameraPostProcessSettingsChromaticAberrationStartOffset,
 	CameraPostProcessSettingsColorContrast,
 	CameraPostProcessSettingsColorContrastHighlights,
diff --git a/Engine/Plugins/Interchange/Runtime/Source/Nodes/Public/InterchangeCameraNode.h b/Engine/Plugins/Interchange/Runtime/Source/Nodes/Public/InterchangeCameraNode.h
--- a/Engine/Plugins/Interchange/Runtime/Source/Nodes/Public/InterchangeCameraNode.h
+++ b/Engine/Plugins/Interchange/Runtime/Source/Nodes/Public/InterchangeCameraNode.h
@@ -7,8 +7,111 @@
 
 #include "InterchangeCameraNode.generated.h"
 
+// Mirrors EAutoExposureMethod from Scene.h
+UENUM(BlueprintType)
+enum class EInterchangeAutoExposureMethod : uint8
+{
+	/** requires compute shader to construct 64 bin histogram */
+	Histogram  UMETA(DisplayName = "Auto Exposure Histogram"),
+
+	/** faster method that computes single value by downsampling */
+	Basic      UMETA(DisplayName = "Auto Exposure Basic"),
+
+	/** Uses camera settings. */
+	Manual   UMETA(DisplayName = "Manual"),
+
+	MAX,
+};
+
+UCLASS(BlueprintType, MinimalAPI)
+class UInterchangeBaseCameraNode : public UInterchangeBaseNode
+{
+	GENERATED_BODY()
+
+public:
+	UFUNCTION(BlueprintCallable, Category = "Interchange | Node | Camera")
+	bool GetCustomNearClipPlane(float& AttributeValue) const
+	{
+		IMPLEMENT_NODE_ATTRIBUTE_GETTER(NearClipPlane, float);
+	}
+
+	UFUNCTION(BlueprintCallable, Category = "Interchange | Node | Camera")
+	bool SetCustomNearClipPlane(const float& AttributeValue)
+	{
+		IMPLEMENT_NODE_ATTRIBUTE_SETTER_NODELEGATE(NearClipPlane, float);
+	}
+
+	UFUNCTION(BlueprintCallable, Category = "Interchange | Node | Camera")
+	bool GetCustomAutoExposureMethod(EInterchangeAutoExposureMethod& AttributeValue) const
+	{
+		IMPLEMENT_NODE_ATTRIBUTE_GETTER(AutoExposureMethod, EInterchangeAutoExposureMethod);
+	}
+
+	UFUNCTION(BlueprintCallable, Category = "Interchange | Node | Camera")
+	bool SetCustomAutoExposureMethod(const EInterchangeAutoExposureMethod& AttributeValue)
+	{
+		IMPLEMENT_NODE_ATTRIBUTE_SETTER_NODELEGATE(AutoExposureMethod, EInterchangeAutoExposureMethod);
+	}
+
+	UFUNCTION(BlueprintCallable, Category = "Interchange | Node | Camera")
+	bool GetCustomExposureCompensation(float& AttributeValue) const
+	{
+		IMPLEMENT_NODE_ATTRIBUTE_GETTER(ExposureCompensation, float);
+	}
+
+	UFUNCTION(BlueprintCallable, Category = "Interchange | Node | Camera")
+	bool SetCustomExposureCompensation(const float& AttributeValue)
+	{
+		IMPLEMENT_NODE_ATTRIBUTE_SETTER_NODELEGATE(ExposureCompensation, float);
+	}
+
+	UFUNCTION(BlueprintCallable, Category = "Interchange | Node | Camera")
+	bool GetCustomExposureFstop(float& AttributeValue) const
+	{
+		IMPLEMENT_NODE_ATTRIBUTE_GETTER(ExposureFstop, float);
+	}
+
+	UFUNCTION(BlueprintCallable, Category = "Interchange | Node | Camera")
+	bool SetCustomExposureFstop(const float& AttributeValue)
+	{
+		IMPLEMENT_NODE_ATTRIBUTE_SETTER_NODELEGATE(ExposureFstop, float);
+	}
+
+	UFUNCTION(BlueprintCallable, Category = "Interchange | Node | Camera")
+	bool GetCustomExposureISO(float& AttributeValue) const
+	{
+		IMPLEMENT_NODE_ATTRIBUTE_GETTER(ExposureISO, float);
+	}
+
+	UFUNCTION(BlueprintCallable, Category = "Interchange | Node | Camera")
+	bool SetCustomExposureISO(const float& AttributeValue)
+	{
+		IMPLEMENT_NODE_ATTRIBUTE_SETTER_NODELEGATE(ExposureISO, float);
+	}
+
+	UFUNCTION(BlueprintCallable, Category = "Interchange | Node | Camera")
+	bool GetCustomExposureShutterSpeed(float& AttributeValue) const
+	{
+		IMPLEMENT_NODE_ATTRIBUTE_GETTER(ExposureShutterSpeed, float);
+	}
+
+	UFUNCTION(BlueprintCallable, Category = "Interchange | Node | Camera")
+	bool SetCustomExposureShutterSpeed(const float& AttributeValue)
+	{
+		IMPLEMENT_NODE_ATTRIBUTE_SETTER_NODELEGATE(ExposureShutterSpeed, float);
+	}
+
+private:
+	IMPLEMENT_NODE_ATTRIBUTE_KEY(NearClipPlane);
+	IMPLEMENT_NODE_ATTRIBUTE_KEY(AutoExposureMethod);
+	IMPLEMENT_NODE_ATTRIBUTE_KEY(ExposureCompensation);
+	IMPLEMENT_NODE_ATTRIBUTE_KEY(ExposureFstop);
+	IMPLEMENT_NODE_ATTRIBUTE_KEY(ExposureISO);
+	IMPLEMENT_NODE_ATTRIBUTE_KEY(ExposureShutterSpeed);
+};
+
 UCLASS(BlueprintType, MinimalAPI)
-class UInterchangePhysicalCameraNode : public UInterchangeBaseNode
+class UInterchangePhysicalCameraNode : public UInterchangeBaseCameraNode
 {
 	GENERATED_BODY()
 
@@ -40,6 +143,18 @@ public:
 		IMPLEMENT_NODE_ATTRIBUTE_SETTER_NODELEGATE(FocalLength, float);
 	}
 
+	UFUNCTION(BlueprintCallable, Category = "Interchange | Node | Camera")
+	bool GetCustomFocusDistance(float& AttributeValue) const
+	{
+		IMPLEMENT_NODE_ATTRIBUTE_GETTER(FocusDistance, float);
+	}
+
+	UFUNCTION(BlueprintCallable, Category = "Interchange | Node | Camera")
+	bool SetCustomFocusDistance(const float& AttributeValue)
+	{
+		IMPLEMENT_NODE_ATTRIBUTE_SETTER_NODELEGATE(FocusDistance, float);
+	}
+
 	UFUNCTION(BlueprintCallable, Category = "Interchange | Node | Camera")
 	bool GetCustomSensorWidth(float& AttributeValue) const
 	{
@@ -53,17 +168,41 @@ public:
 	}
 
 	UFUNCTION(BlueprintCallable, Category = "Interchange | Node | Camera")
-		bool GetCustomSensorHeight(float& AttributeValue) const
+	bool GetCustomSensorHeight(float& AttributeValue) const
 	{
 		IMPLEMENT_NODE_ATTRIBUTE_GETTER(SensorHeight, float);
 	}
 
 	UFUNCTION(BlueprintCallable, Category = "Interchange | Node | Camera")
-		bool SetCustomSensorHeight(const float& AttributeValue)
+	bool SetCustomSensorHeight(const float& AttributeValue)
 	{
 		IMPLEMENT_NODE_ATTRIBUTE_SETTER_NODELEGATE(SensorHeight, float);
 	}
 
+	UFUNCTION(BlueprintCallable, Category = "Interchange | Node | Camera")
+	bool GetCustomSensorHorizontalOffset(float& AttributeValue) const
+	{
+		IMPLEMENT_NODE_ATTRIBUTE_GETTER(SensorHorizontalOffset, float);
+	}
+
+	UFUNCTION(BlueprintCallable, Category = "Interchange | Node | Camera")
+	bool SetCustomSensorHorizontalOffset(const float& AttributeValue)
+	{
+		IMPLEMENT_NODE_ATTRIBUTE_SETTER_NODELEGATE(SensorHorizontalOffset, float);
+	}
+
+	UFUNCTION(BlueprintCallable, Category = "Interchange | Node | Camera")
+	bool GetCustomSensorVerticalOffset(float& AttributeValue) const
+	{
+		IMPLEMENT_NODE_ATTRIBUTE_GETTER(SensorVerticalOffset, float);
+	}
+
+	UFUNCTION(BlueprintCallable, Category = "Interchange | Node | Camera")
+	bool SetCustomSensorVerticalOffset(const float& AttributeValue)
+	{
+		IMPLEMENT_NODE_ATTRIBUTE_SETTER_NODELEGATE(SensorVerticalOffset, float);
+	}
+
 	UFUNCTION(BlueprintCallable, Category = "Interchange | Node | Camera")
 	bool GetCustomEnableDepthOfField(bool& AttributeValue) const
 	{
@@ -75,11 +214,28 @@ public:
 	{
 		IMPLEMENT_NODE_ATTRIBUTE_SETTER_NODELEGATE(EnableDepthOfField, bool);
 	}
+
+	UFUNCTION(BlueprintCallable, Category = "Interchange | Node | Camera")
+	bool GetCustomDepthOfFieldFstop(float& AttributeValue) const
+	{
+		IMPLEMENT_NODE_ATTRIBUTE_GETTER(DepthOfFieldFstop, float);
+	}
+
+	UFUNCTION(BlueprintCallable, Category = "Interchange | Node | Camera")
+	bool SetCustomDepthOfFieldFstop(const float& AttributeValue)
+	{
+		IMPLEMENT_NODE_ATTRIBUTE_SETTER_NODELEGATE(DepthOfFieldFstop, float);
+	}
+
 private:
-	const UE::Interchange::FAttributeKey Macro_CustomFocalLengthKey = UE::Interchange::FAttributeKey(TEXT("FocalLength"));
-	const UE::Interchange::FAttributeKey Macro_CustomSensorWidthKey = UE::Interchange::FAttributeKey(TEXT("SensorWidth"));
-	const UE::Interchange::FAttributeKey Macro_CustomSensorHeightKey = UE::Interchange::FAttributeKey(TEXT("SensorHeight"));
-	const UE::Interchange::FAttributeKey Macro_CustomEnableDepthOfFieldKey = UE::Interchange::FAttributeKey(TEXT("EnableDepthOfField"));
+	IMPLEMENT_NODE_ATTRIBUTE_KEY(FocalLength);
+	IMPLEMENT_NODE_ATTRIBUTE_KEY(FocusDistance);
+	IMPLEMENT_NODE_ATTRIBUTE_KEY(SensorWidth);
+	IMPLEMENT_NODE_ATTRIBUTE_KEY(SensorHeight);
+	IMPLEMENT_NODE_ATTRIBUTE_KEY(SensorHorizontalOffset);
+	IMPLEMENT_NODE_ATTRIBUTE_KEY(SensorVerticalOffset);
+	IMPLEMENT_NODE_ATTRIBUTE_KEY(EnableDepthOfField);
+	IMPLEMENT_NODE_ATTRIBUTE_KEY(DepthOfFieldFstop);
 };
 
 UENUM(BlueprintType)
@@ -91,7 +247,7 @@ enum class EInterchangeCameraProjectionType : uint8
 
 // Primarily used for Ortho Camera
 UCLASS(BlueprintType, MinimalAPI)
-class UInterchangeStandardCameraNode : public UInterchangeBaseNode
+class UInterchangeStandardCameraNode : public UInterchangeBaseCameraNode
 {
 	GENERATED_BODY()
 
@@ -111,7 +267,6 @@ public:
 	}
 
 public:
-
 	UFUNCTION(BlueprintCallable, Category = "Interchange | Node | Camera")
 	bool GetCustomProjectionMode(EInterchangeCameraProjectionType& AttributeValue) const
 	{
@@ -134,17 +289,6 @@ public:
 		IMPLEMENT_NODE_ATTRIBUTE_SETTER_NODELEGATE(Width, float);
 	}
 
-	UFUNCTION(BlueprintCallable, Category = "Interchange | Node | Camera")
-	bool GetCustomNearClipPlane(float& AttributeValue) const
-	{
-	IMPLEMENT_NODE_ATTRIBUTE_GETTER(NearClipPlane, float);
-	}
-	UFUNCTION(BlueprintCallable, Category = "Interchange | Node | Camera")
-	bool SetCustomNearClipPlane(const float& AttributeValue)
-	{
-		IMPLEMENT_NODE_ATTRIBUTE_SETTER_NODELEGATE(NearClipPlane, float);
-	}
-
 	UFUNCTION(BlueprintCallable, Category = "Interchange | Node | Camera")
 	bool GetCustomFarClipPlane(float& AttributeValue) const
 	{
@@ -181,12 +325,9 @@ public:
 	}
 
 private:
-	const UE::Interchange::FAttributeKey Macro_CustomProjectionModeKey = UE::Interchange::FAttributeKey(TEXT("ProjectionMode"));
-	
-	const UE::Interchange::FAttributeKey Macro_CustomWidthKey = UE::Interchange::FAttributeKey(TEXT("Width"));
-	const UE::Interchange::FAttributeKey Macro_CustomNearClipPlaneKey = UE::Interchange::FAttributeKey(TEXT("NearClipPlane"));
-	const UE::Interchange::FAttributeKey Macro_CustomFarClipPlaneKey = UE::Interchange::FAttributeKey(TEXT("FarClipPlane"));
-	const UE::Interchange::FAttributeKey Macro_CustomAspectRatioKey = UE::Interchange::FAttributeKey(TEXT("AspectRatio"));
-
-	const UE::Interchange::FAttributeKey Macro_CustomFieldOfViewKey = UE::Interchange::FAttributeKey(TEXT("FieldOfView"));
+	IMPLEMENT_NODE_ATTRIBUTE_KEY(ProjectionMode);
+	IMPLEMENT_NODE_ATTRIBUTE_KEY(Width);
+	IMPLEMENT_NODE_ATTRIBUTE_KEY(FarClipPlane);
+	IMPLEMENT_NODE_ATTRIBUTE_KEY(AspectRatio);
+	IMPLEMENT_NODE_ATTRIBUTE_KEY(FieldOfView);
 };
\ No newline at end of file
diff --git a/Engine/Plugins/Interchange/Runtime/Source/Pipelines/Private/InterchangeGenericScenesPipeline.cpp b/Engine/Plugins/Interchange/Runtime/Source/Pipelines/Private/InterchangeGenericScenesPipeline.cpp
--- a/Engine/Plugins/Interchange/Runtime/Source/Pipelines/Private/InterchangeGenericScenesPipeline.cpp
+++ b/Engine/Plugins/Interchange/Runtime/Source/Pipelines/Private/InterchangeGenericScenesPipeline.cpp
@@ -1181,106 +1181,169 @@ void UInterchangeGenericLevelPipeline::SetUpFactoryNode(UInterchangeActorFactory
 			ActorFactoryNode->SetCustomActorClassName(APointLight::StaticClass()->GetPathName());
 		}
 	}
-	else if (const UInterchangePhysicalCameraNode* PhysicalCameraNode = Cast<UInterchangePhysicalCameraNode>(TranslatedAssetNode))
+	else if (const UInterchangeBaseCameraNode* CameraNode = Cast<UInterchangeBaseCameraNode>(TranslatedAssetNode))
 	{
-		ActorFactoryNode->SetCustomActorClassName(ACineCameraActor::StaticClass()->GetPathName());
-		ActorFactoryNode->SetCustomMobility(EComponentMobility::Movable);
-
-		if (UInterchangePhysicalCameraFactoryNode* PhysicalCameraFactoryNode = Cast<UInterchangePhysicalCameraFactoryNode>(ActorFactoryNode))
+		if (UInterchangeBaseCameraFactoryNode* CameraFactoryNode = Cast<UInterchangeBaseCameraFactoryNode>(ActorFactoryNode))
 		{
-			float FocalLength;
-			if (PhysicalCameraNode->GetCustomFocalLength(FocalLength))
+			ActorFactoryNode->SetCustomMobility(EComponentMobility::Movable);
+
+			static_assert(int(EInterchangeAutoExposureMethod::Histogram) == int(EAutoExposureMethod::AEM_Histogram), "EInterchangeAutoExposureMethod::Histogram differs from EAutoExposureMethod::AEM_Histogram");
+			static_assert(int(EInterchangeAutoExposureMethod::Basic) == int(EAutoExposureMethod::AEM_Basic), "EInterchangeAutoExposureMethod::Basic differs from EAutoExposureMethod::AEM_Basic");
+			static_assert(int(EInterchangeAutoExposureMethod::Manual) == int(EAutoExposureMethod::AEM_Manual), "EInterchangeAutoExposureMethod::Manual differs from EAutoExposureMethod::AEM_Manual");
+			static_assert(int(EInterchangeAutoExposureMethod::MAX) == int(EAutoExposureMethod::AEM_MAX), "EInterchangeAutoExposureMethod::MAX differs from EAutoExposureMethod::AEM_MAX");
+
+			float NearClipPlane;
+			if (CameraNode->GetCustomNearClipPlane(NearClipPlane))
 			{
-				PhysicalCameraFactoryNode->SetCustomFocalLength(FocalLength);
+				CameraFactoryNode->SetCustomNearClipPlane(NearClipPlane);
 			}
 
-			float SensorHeight;
-			if (PhysicalCameraNode->GetCustomSensorHeight(SensorHeight))
+			EInterchangeAutoExposureMethod AutoExposureMethod = EInterchangeAutoExposureMethod::Histogram;
+			if (CameraNode->GetCustomAutoExposureMethod(AutoExposureMethod))
 			{
-				PhysicalCameraFactoryNode->SetCustomSensorHeight(SensorHeight);
+				CameraFactoryNode->SetCustomAutoExposureMethod(AutoExposureMethod);
 			}
 
-			float SensorWidth;
-			if (PhysicalCameraNode->GetCustomSensorWidth(SensorWidth))
+			float ExposureCompensation;
+			if (CameraNode->GetCustomExposureCompensation(ExposureCompensation))
 			{
-				PhysicalCameraFactoryNode->SetCustomSensorWidth(SensorWidth);
+				CameraFactoryNode->SetCustomExposureCompensation(ExposureCompensation);
 			}
 
-			bool bEnableDepthOfField;
-			if (PhysicalCameraNode->GetCustomEnableDepthOfField(bEnableDepthOfField))
+			float ExposureFstop;
+			if (CameraNode->GetCustomExposureFstop(ExposureFstop))
 			{
-				PhysicalCameraFactoryNode->SetCustomFocusMethod(bEnableDepthOfField ? ECameraFocusMethod::Manual : ECameraFocusMethod::DoNotOverride);
+				CameraFactoryNode->SetCustomExposureFstop(ExposureFstop);
 			}
 
-		}
-	}
-	else if (const UInterchangeStandardCameraNode* CameraNode = Cast<UInterchangeStandardCameraNode>(TranslatedAssetNode))
-	{
-		EInterchangeCameraProjectionType ProjectionType = EInterchangeCameraProjectionType::Perspective;
-		if (CameraNode->GetCustomProjectionMode(ProjectionType) && bUsePhysicalInsteadOfStandardPerspectiveCamera && ProjectionType == EInterchangeCameraProjectionType::Perspective)
-		{
-			float AspectRatio = 1.0f;
-			CameraNode->GetCustomAspectRatio(AspectRatio);
-
-			const float SensorWidth = 36.f;  // mm
-			float SensorHeight = SensorWidth / AspectRatio;
-
-			float FieldOfView = 90;
-			CameraNode->GetCustomFieldOfView(FieldOfView); //Degrees
-
-			float FocalLength = (SensorHeight) / (2.0 * tan(FMath::DegreesToRadians(FieldOfView) / 2.0));
-
-			ActorFactoryNode->SetCustomActorClassName(ACineCameraActor::StaticClass()->GetPathName());
-			ActorFactoryNode->SetCustomMobility(EComponentMobility::Movable);
+			float ExposureISO;
+			if (CameraNode->GetCustomExposureISO(ExposureISO))
+			{
+				CameraFactoryNode->SetCustomExposureISO(ExposureISO);
+			}
 
-			if (UInterchangePhysicalCameraFactoryNode* PhysicalCameraFactoryNode = Cast<UInterchangePhysicalCameraFactoryNode>(ActorFactoryNode))
+			float ExposureShutterSpeed;
+			if (CameraNode->GetCustomExposureShutterSpeed(ExposureShutterSpeed))
 			{
-				PhysicalCameraFactoryNode->SetCustomFocalLength(FocalLength);
-				PhysicalCameraFactoryNode->SetCustomSensorHeight(SensorHeight);
-				PhysicalCameraFactoryNode->SetCustomSensorWidth(SensorWidth);
-				PhysicalCameraFactoryNode->SetCustomFocusMethod(ECameraFocusMethod::DoNotOverride);
+				CameraFactoryNode->SetCustomExposureShutterSpeed(ExposureShutterSpeed);
 			}
 		}
-		else
+
+		if (const UInterchangePhysicalCameraNode* PhysicalCameraNode = Cast<UInterchangePhysicalCameraNode>(CameraNode))
 		{
-			ActorFactoryNode->SetCustomActorClassName(ACameraActor::StaticClass()->GetPathName());
-			ActorFactoryNode->SetCustomMobility(EComponentMobility::Movable);
+			ActorFactoryNode->SetCustomActorClassName(ACineCameraActor::StaticClass()->GetPathName());
 
-			if (UInterchangeStandardCameraFactoryNode* CameraFactoryNode = Cast<UInterchangeStandardCameraFactoryNode>(ActorFactoryNode))
+			if (UInterchangePhysicalCameraFactoryNode* PhysicalCameraFactoryNode = Cast<UInterchangePhysicalCameraFactoryNode>(ActorFactoryNode))
 			{
-				if (CameraNode->GetCustomProjectionMode(ProjectionType))
+				float FocalLength;
+				if (PhysicalCameraNode->GetCustomFocalLength(FocalLength))
 				{
-					CameraFactoryNode->SetCustomProjectionMode((ECameraProjectionMode::Type)ProjectionType);
+					PhysicalCameraFactoryNode->SetCustomFocalLength(FocalLength);
 				}
 
-				float OrthoWidth;
-				if (CameraNode->GetCustomWidth(OrthoWidth))
+				float FocusDistance;
+				if (PhysicalCameraNode->GetCustomFocusDistance(FocusDistance))
 				{
-					CameraFactoryNode->SetCustomWidth(OrthoWidth);
+					PhysicalCameraFactoryNode->SetCustomFocusDistance(FocusDistance);
 				}
 
-				float OrthoNearClipPlane;
-				if (CameraNode->GetCustomNearClipPlane(OrthoNearClipPlane))
+				float SensorHeight;
+				if (PhysicalCameraNode->GetCustomSensorHeight(SensorHeight))
 				{
-					CameraFactoryNode->SetCustomNearClipPlane(OrthoNearClipPlane);
+					PhysicalCameraFactoryNode->SetCustomSensorHeight(SensorHeight);
 				}
 
-				float OrthoFarClipPlane;
-				if (CameraNode->GetCustomFarClipPlane(OrthoFarClipPlane))
+				float SensorWidth;
+				if (PhysicalCameraNode->GetCustomSensorWidth(SensorWidth))
 				{
-					CameraFactoryNode->SetCustomFarClipPlane(OrthoFarClipPlane);
+					PhysicalCameraFactoryNode->SetCustomSensorWidth(SensorWidth);
 				}
 
-				float AspectRatio;
-				if (CameraNode->GetCustomAspectRatio(AspectRatio))
+				float SensorVerticalOffset;
+				if (PhysicalCameraNode->GetCustomSensorVerticalOffset(SensorVerticalOffset))
 				{
-					CameraFactoryNode->SetCustomAspectRatio(AspectRatio);
+					PhysicalCameraFactoryNode->SetCustomSensorVerticalOffset(SensorVerticalOffset);
 				}
 
-				float FieldOfView;
-				if (CameraNode->GetCustomFieldOfView(FieldOfView))
+				float SensorHorizontalOffset;
+				if (PhysicalCameraNode->GetCustomSensorHorizontalOffset(SensorHorizontalOffset))
 				{
-					CameraFactoryNode->SetCustomFieldOfView(FieldOfView);
+					PhysicalCameraFactoryNode->SetCustomSensorHorizontalOffset(SensorHorizontalOffset);
+				}
+
+				bool bEnableDepthOfField;
+				if (PhysicalCameraNode->GetCustomEnableDepthOfField(bEnableDepthOfField))
+				{
+					PhysicalCameraFactoryNode->SetCustomFocusMethod(bEnableDepthOfField ? ECameraFocusMethod::Manual : ECameraFocusMethod::DoNotOverride);
+				}
+
+				float DepthOfFieldFstop;
+				if (PhysicalCameraNode->GetCustomDepthOfFieldFstop(DepthOfFieldFstop))
+				{
+					PhysicalCameraFactoryNode->SetCustomDepthOfFieldFstop(DepthOfFieldFstop);
+				}
+			}
+		}
+		else if (const UInterchangeStandardCameraNode* StandardCameraNode = Cast<UInterchangeStandardCameraNode>(TranslatedAssetNode))
+		{
+			EInterchangeCameraProjectionType ProjectionType = EInterchangeCameraProjectionType::Perspective;
+			if (StandardCameraNode->GetCustomProjectionMode(ProjectionType) && bUsePhysicalInsteadOfStandardPerspectiveCamera && ProjectionType == EInterchangeCameraProjectionType::Perspective)
+			{
+				float AspectRatio = 1.0f;
+				StandardCameraNode->GetCustomAspectRatio(AspectRatio);
+
+				const float SensorWidth = 36.f;  // mm
+				float SensorHeight = SensorWidth / AspectRatio;
+
+				float FieldOfView = 90;
+				StandardCameraNode->GetCustomFieldOfView(FieldOfView); //Degrees
+
+				float FocalLength = (SensorHeight) / (2.0 * tan(FMath::DegreesToRadians(FieldOfView) / 2.0));
+
+				ActorFactoryNode->SetCustomActorClassName(ACineCameraActor::StaticClass()->GetPathName());
+
+				if (UInterchangePhysicalCameraFactoryNode* PhysicalCameraFactoryNode = Cast<UInterchangePhysicalCameraFactoryNode>(ActorFactoryNode))
+				{
+					PhysicalCameraFactoryNode->SetCustomFocalLength(FocalLength);
+					PhysicalCameraFactoryNode->SetCustomSensorHeight(SensorHeight);
+					PhysicalCameraFactoryNode->SetCustomSensorWidth(SensorWidth);
+					PhysicalCameraFactoryNode->SetCustomFocusMethod(ECameraFocusMethod::DoNotOverride);
+				}
+			}
+			else
+			{
+				ActorFactoryNode->SetCustomActorClassName(ACameraActor::StaticClass()->GetPathName());
+
+				if (UInterchangeStandardCameraFactoryNode* CameraFactoryNode = Cast<UInterchangeStandardCameraFactoryNode>(ActorFactoryNode))
+				{
+					if (StandardCameraNode->GetCustomProjectionMode(ProjectionType))
+					{
+						CameraFactoryNode->SetCustomProjectionMode((ECameraProjectionMode::Type)ProjectionType);
+					}
+
+					float OrthoWidth;
+					if (StandardCameraNode->GetCustomWidth(OrthoWidth))
+					{
+						CameraFactoryNode->SetCustomWidth(OrthoWidth);
+					}
+
+					float OrthoFarClipPlane;
+					if (StandardCameraNode->GetCustomFarClipPlane(OrthoFarClipPlane))
+					{
+						CameraFactoryNode->SetCustomFarClipPlane(OrthoFarClipPlane);
+					}
+
+					float AspectRatio;
+					if (StandardCameraNode->GetCustomAspectRatio(AspectRatio))
+					{
+						CameraFactoryNode->SetCustomAspectRatio(AspectRatio);
+					}
+
+					float FieldOfView;
+					if (StandardCameraNode->GetCustomFieldOfView(FieldOfView))
+					{
+						CameraFactoryNode->SetCustomFieldOfView(FieldOfView);
+					}
 				}
 			}
 		}
diff --git a/Engine/Plugins/Runtime/USDCore/Source/USDUtilities/Private/USDPrimConversion.cpp b/Engine/Plugins/Runtime/USDCore/Source/USDUtilities/Private/USDPrimConversion.cpp
--- a/Engine/Plugins/Runtime/USDCore/Source/USDUtilities/Private/USDPrimConversion.cpp
+++ b/Engine/Plugins/Runtime/USDCore/Source/USDUtilities/Private/USDPrimConversion.cpp
@@ -573,6 +573,10 @@ bool UsdToUnreal::ConvertGeomCamera(const UE::FUsdPrim& Prim, UCineCameraCompone
 	UE::FUsdStage Stage = Prim.GetStage();
 	FUsdStageInfo StageInfo(Stage);
 
+	bool bNeedsManualFocus = false;
+	bool bNeedsManualExposure = false;
+	bool bDisableFocus = false;
+
 	if (pxr::UsdAttribute Attr = GeomCamera.GetFocalLengthAttr(); Attr.HasAuthoredValue())
 	{
 		CameraComponent.SetCurrentFocalLength(UsdToUnreal::ConvertDistance(StageInfo, UsdUtils::GetUsdValue<float>(Attr, UsdTimeCode)));
@@ -580,32 +584,41 @@ bool UsdToUnreal::ConvertGeomCamera(const UE::FUsdPrim& Prim, UCineCameraCompone
 
 	if (pxr::UsdAttribute Attr = GeomCamera.GetFocusDistanceAttr(); Attr.HasAuthoredValue())
 	{
-		CameraComponent.FocusSettings.ManualFocusDistance = UsdToUnreal::ConvertDistance(
+		float UsdFocusDistance = UsdToUnreal::ConvertDistance(
 			StageInfo,
 			UsdUtils::GetUsdValue<float>(GeomCamera.GetFocusDistanceAttr(), UsdTimeCode)
 		);
 
-		if (FMath::IsNearlyZero(CameraComponent.FocusSettings.ManualFocusDistance))
-		{
-			CameraComponent.FocusSettings.FocusMethod = ECameraFocusMethod::DoNotOverride;
-		}
-		else
+		CameraComponent.FocusSettings.ManualFocusDistance = UsdFocusDistance;
+
+		if (!FMath::IsNearlyZero(UsdFocusDistance))
 		{
-			CameraComponent.FocusSettings.FocusMethod = ECameraFocusMethod::Manual;
+			bNeedsManualFocus = true;
 		}
 	}
 
+	// Regular 'fStop' is now meant to be used for DoF effects according to the latest schema
 	if (pxr::UsdAttribute Attr = GeomCamera.GetFStopAttr(); Attr.HasAuthoredValue())
 	{
 		const float UsdFStop = UsdUtils::GetUsdValue<float>(Attr, UsdTimeCode);
-		CameraComponent.SetCurrentAperture(UsdFStop);
+		CameraComponent.PostProcessSettings.bOverride_DepthOfFieldFstop = true;
+		CameraComponent.PostProcessSettings.DepthOfFieldFstop = UsdFStop;
+		bNeedsManualFocus = true;
 
 		if (FMath::IsNearlyZero(UsdFStop))
 		{
-			CameraComponent.FocusSettings.FocusMethod = ECameraFocusMethod::Disable;
+			bDisableFocus = true;
 		}
 	}
 
+	// 'exposure:fStop' is the the value that is meant to affect exposure
+	if (pxr::UsdAttribute Attr = GeomCamera.GetExposureFStopAttr(); Attr.HasAuthoredValue())
+	{
+		const float UsdFStop = UsdUtils::GetUsdValue<float>(Attr, UsdTimeCode);
+		CameraComponent.SetCurrentAperture(UsdFStop);
+		bNeedsManualExposure = true;
+	}
+
 	if (pxr::UsdAttribute Attr = GeomCamera.GetHorizontalApertureAttr(); Attr.HasAuthoredValue())
 	{
 		CameraComponent.Filmback.SensorWidth = UsdToUnreal::ConvertDistance(StageInfo, UsdUtils::GetUsdValue<float>(Attr, UsdTimeCode));
@@ -630,6 +643,8 @@ bool UsdToUnreal::ConvertGeomCamera(const UE::FUsdPrim& Prim, UCineCameraCompone
 	{
 		CameraComponent.PostProcessSettings.bOverride_AutoExposureBias = true;
 		CameraComponent.PostProcessSettings.AutoExposureBias = UsdUtils::GetUsdValue<float>(Attr, UsdTimeCode);
+
+		// We don't necessarily want manual exposure just because we have this: It still works with auto exposure
 	}
 
 	if (pxr::UsdAttribute Attr = GeomCamera.GetProjectionAttr(); Attr.HasAuthoredValue())
@@ -653,6 +668,43 @@ bool UsdToUnreal::ConvertGeomCamera(const UE::FUsdPrim& Prim, UCineCameraCompone
 		CameraComponent.bOverride_CustomNearClippingPlane = true;
 	}
 
+	if (pxr::UsdAttribute Attr = GeomCamera.GetExposureIsoAttr(); Attr.HasAuthoredValue())
+	{
+		CameraComponent.PostProcessSettings.bOverride_CameraISO = true;
+		CameraComponent.PostProcessSettings.CameraISO = UsdUtils::GetUsdValue<float>(Attr, UsdTimeCode);
+		bNeedsManualExposure = true;
+	}
+
+	if (pxr::UsdAttribute Attr = GeomCamera.GetExposureTimeAttr(); Attr.HasAuthoredValue())
+	{
+		float UsdTimeSeconds = UsdUtils::GetUsdValue<float>(Attr, UsdTimeCode);
+		if (!FMath::IsNearlyZero(UsdTimeSeconds))
+		{
+			CameraComponent.PostProcessSettings.bOverride_CameraShutterSpeed = true;
+			CameraComponent.PostProcessSettings.CameraShutterSpeed = 1.0f / UsdTimeSeconds;
+		}
+		bNeedsManualExposure = true;
+	}
+
+	if (bDisableFocus)
+	{
+		CameraComponent.FocusSettings.FocusMethod = ECameraFocusMethod::Disable;
+	}
+	else if (bNeedsManualFocus)
+	{
+		CameraComponent.FocusSettings.FocusMethod = ECameraFocusMethod::Manual;
+	}
+	else
+	{
+		CameraComponent.FocusSettings.FocusMethod = ECameraFocusMethod::DoNotOverride;
+	}
+
+	if (bNeedsManualExposure)
+	{
+		CameraComponent.PostProcessSettings.bOverride_AutoExposureMethod = true;
+		CameraComponent.PostProcessSettings.AutoExposureMethod = EAutoExposureMethod::AEM_Manual;
+	}
+
 	return true;
 }
 
@@ -1373,6 +1425,7 @@ UsdToUnreal::FPropertyTrackReader UsdToUnreal::CreatePropertyTrackReader(
 	{
 		TOptional<int32> ClippingPlaneIndex;
 		bool bConvertDistance = true;
+		bool bInvertValue = false;
 		pxr::UsdAttribute Attr;
 
 		if (PropertyPath == UnrealIdentifiers::CurrentFocalLengthPropertyName)
@@ -1386,7 +1439,7 @@ UsdToUnreal::FPropertyTrackReader UsdToUnreal::CreatePropertyTrackReader(
 		else if (PropertyPath == UnrealIdentifiers::CurrentAperturePropertyName)
 		{
 			bConvertDistance = false;
-			Attr = Camera.GetFStopAttr();
+			Attr = Camera.GetExposureFStopAttr();
 		}
 		else if (PropertyPath == UnrealIdentifiers::SensorWidthPropertyName)
 		{
@@ -1424,6 +1477,22 @@ UsdToUnreal::FPropertyTrackReader UsdToUnreal::CreatePropertyTrackReader(
 			Attr = Camera.GetClippingRangeAttr();
 			ClippingPlaneIndex = 0;	   // Near plane
 		}
+		else if (PropertyPath == UnrealIdentifiers::DepthOfFieldFstopPropertyName)
+		{
+			bConvertDistance = false;
+			Attr = Camera.GetFStopAttr();
+		}
+		else if (PropertyPath == UnrealIdentifiers::CameraISOPropertyName)
+		{
+			bConvertDistance = false;
+			Attr = Camera.GetExposureIsoAttr();
+		}
+		else if (PropertyPath == UnrealIdentifiers::CameraShutterSpeedPropertyName)
+		{
+			bInvertValue = true;
+			bConvertDistance = false;
+			Attr = Camera.GetExposureTimeAttr();
+		}
 
 		// TODO: Projection mode track reader? Would need new enum flag reader types
 
@@ -1472,12 +1541,27 @@ UsdToUnreal::FPropertyTrackReader UsdToUnreal::CreatePropertyTrackReader(
 				float Default = 0.0;
 				Attr.Get<float>(&Default);
 
-				Reader.FloatReader = [Attr, Default](double UsdTimeCode)
+				if (bInvertValue)
 				{
-					float Result = Default;
-					Attr.Get<float>(&Result, UsdTimeCode);
-					return Result;
-				};
+					Reader.FloatReader = [Attr, Default](double UsdTimeCode)
+					{
+						float Result = Default;
+						Attr.Get<float>(&Result, UsdTimeCode);
+
+						return Result != 0.0f ? (1.0f / Result) : 0.0f;
+					};
+				}
+				else
+				{
+					Reader.FloatReader = [Attr, Default](double UsdTimeCode)
+					{
+						float Result = Default;
+						Attr.Get<float>(&Result, UsdTimeCode);
+
+						return Result;
+					};
+				}
+
 				return Reader;
 			}
 		}
@@ -2344,16 +2428,22 @@ bool UnrealToUsd::ConvertCameraComponent(const UCineCameraComponent& CameraCompo
 		UsdUtils::NotifyIfOverriddenOpinion(Attr);
 	}
 
-	if (pxr::UsdAttribute Attr = GeomCamera.CreateFocusDistanceAttr())
+	if (CameraComponent.FocusSettings.FocusMethod == ECameraFocusMethod::Manual)
 	{
-		Attr.Set<float>(UnrealToUsd::ConvertDistance(StageInfo, CameraComponent.FocusSettings.ManualFocusDistance), UsdTimeCode);
-		UsdUtils::NotifyIfOverriddenOpinion(Attr);
-	}
+		if (CameraComponent.PostProcessSettings.bOverride_DepthOfFieldFstop)
+		{
+			if (pxr::UsdAttribute Attr = GeomCamera.CreateFStopAttr())
+			{
+				Attr.Set<float>(CameraComponent.PostProcessSettings.DepthOfFieldFstop, UsdTimeCode);
+				UsdUtils::NotifyIfOverriddenOpinion(Attr);
+			}
+		}
 
-	if (pxr::UsdAttribute Attr = GeomCamera.CreateFStopAttr())
-	{
-		Attr.Set<float>(CameraComponent.CurrentAperture, UsdTimeCode);
-		UsdUtils::NotifyIfOverriddenOpinion(Attr);
+		if (pxr::UsdAttribute Attr = GeomCamera.CreateFocusDistanceAttr())
+		{
+			Attr.Set<float>(UnrealToUsd::ConvertDistance(StageInfo, CameraComponent.FocusSettings.ManualFocusDistance), UsdTimeCode);
+			UsdUtils::NotifyIfOverriddenOpinion(Attr);
+		}
 	}
 
 	if (pxr::UsdAttribute Attr = GeomCamera.CreateHorizontalApertureAttr())
@@ -2380,10 +2470,13 @@ bool UnrealToUsd::ConvertCameraComponent(const UCineCameraComponent& CameraCompo
 		UsdUtils::NotifyIfOverriddenOpinion(Attr);
 	}
 
-	if (pxr::UsdAttribute Attr = GeomCamera.CreateExposureAttr(); CameraComponent.PostProcessSettings.bOverride_AutoExposureBias)
+	if (CameraComponent.PostProcessSettings.bOverride_AutoExposureBias)
 	{
-		Attr.Set<float>(CameraComponent.PostProcessSettings.AutoExposureBias, UsdTimeCode);
-		UsdUtils::NotifyIfOverriddenOpinion(Attr);
+		if (pxr::UsdAttribute Attr = GeomCamera.CreateExposureAttr())
+		{
+			Attr.Set<float>(CameraComponent.PostProcessSettings.AutoExposureBias, UsdTimeCode);
+			UsdUtils::NotifyIfOverriddenOpinion(Attr);
+		}
 	}
 
 	if (pxr::UsdAttribute Attr = GeomCamera.CreateProjectionAttr())
@@ -2427,6 +2520,34 @@ bool UnrealToUsd::ConvertCameraComponent(const UCineCameraComponent& CameraCompo
 		}
 	}
 
+	if (CameraComponent.PostProcessSettings.bOverride_AutoExposureMethod && CameraComponent.PostProcessSettings.AutoExposureMethod == EAutoExposureMethod::AEM_Manual)
+	{
+		if (pxr::UsdAttribute Attr = GeomCamera.CreateExposureFStopAttr())
+		{
+			Attr.Set<float>(CameraComponent.CurrentAperture, UsdTimeCode);
+			UsdUtils::NotifyIfOverriddenOpinion(Attr);
+		}
+
+		if (CameraComponent.PostProcessSettings.bOverride_CameraISO)
+		{
+			if (pxr::UsdAttribute Attr = GeomCamera.CreateExposureIsoAttr())
+			{
+				Attr.Set<float>(CameraComponent.PostProcessSettings.CameraISO, UsdTimeCode);
+				UsdUtils::NotifyIfOverriddenOpinion(Attr);
+			}
+		}
+
+		if (CameraComponent.PostProcessSettings.bOverride_CameraShutterSpeed)
+		{
+			if (pxr::UsdAttribute Attr = GeomCamera.CreateExposureTimeAttr())
+			{
+				float UsdValue = CameraComponent.PostProcessSettings.CameraShutterSpeed != 0.0f ? (1.0f / CameraComponent.PostProcessSettings.CameraShutterSpeed) : 0.0f;
+				Attr.Set<float>(UsdValue, UsdTimeCode);
+				UsdUtils::NotifyIfOverriddenOpinion(Attr);
+			}
+		}
+	}
+
 	return true;
 }
 
@@ -4576,6 +4697,11 @@ bool UnrealToUsd::CreateComponentPropertyBaker(
 		return UEValue;
 	};
 
+	TFunction<float(float)> InvertConversion = [StageInfo](float UEValue)
+	{
+		return UEValue != 0.0f ? (1.0f / UEValue) : 0.0f;
+	};
+
 	TFunction<bool(bool)> NoBoolConversion = [](bool Value)
 	{
 		return Value;
@@ -4819,6 +4945,9 @@ bool UnrealToUsd::CreateComponentPropertyBaker(
 			pxr::UsdAttribute ExposureAttr = GeomCamera.CreateExposureAttr();
 			pxr::UsdAttribute ProjectionAttr = GeomCamera.CreateProjectionAttr();
 			pxr::UsdAttribute ClippingRangeAttr = GeomCamera.CreateClippingRangeAttr();
+			pxr::UsdAttribute ExposureFStopAttr = GeomCamera.CreateExposureFStopAttr();
+			pxr::UsdAttribute ExposureISOAttr = GeomCamera.CreateExposureIsoAttr();
+			pxr::UsdAttribute ExposureTimeAttr = GeomCamera.CreateExposureTimeAttr();
 
 			TFunction<void(float, double)> FocalLengthSetter = CreateCachedAttrSetter(FocalLengthAttr, DistanceConversion);
 			TFunction<void(float, double)> FocusDistanceSetter = CreateCachedAttrSetter(FocusDistanceAttr, DistanceConversion);
@@ -4830,6 +4959,9 @@ bool UnrealToUsd::CreateComponentPropertyBaker(
 			TFunction<void(float, double)> ExposureSetter = CreateCachedAttrSetter(ExposureAttr, NoFloatConversion);
 			TFunction<void(ECameraProjectionMode::Type, double)> ProjectionSetter = CreateCachedAttrSetter(ProjectionAttr, ProjectionConversion);
 			TFunction<void(const FVector2D&, double)> ClippingRangeSetter = CreateCachedAttrSetter(ClippingRangeAttr, ClippingRangeConversion);
+			TFunction<void(float, double)> ExposureFStopSetter = CreateCachedAttrSetter(ExposureFStopAttr, NoFloatConversion);
+			TFunction<void(float, double)> ExposureISOSetter = CreateCachedAttrSetter(ExposureISOAttr, NoFloatConversion);
+			TFunction<void(float, double)> ExposureTimeSetter = CreateCachedAttrSetter(ExposureTimeAttr, InvertConversion);
 
 			float DefaultUEFarPlane = FLT_MAX;
 			if (ClippingRangeAttr)
@@ -4852,7 +4984,10 @@ bool UnrealToUsd::CreateComponentPropertyBaker(
 							 VerticalApertureOffsetSetter,
 							 ExposureSetter,
 							 ProjectionSetter,
-							 ClippingRangeSetter](double UsdTimeCode) mutable
+							 ClippingRangeSetter,
+							 ExposureFStopSetter,
+							 ExposureISOSetter,
+							 ExposureTimeSetter](double UsdTimeCode) mutable
 			{
 				TRACE_CPUPROFILER_EVENT_SCOPE(CameraBaker);
 
@@ -4864,7 +4999,12 @@ bool UnrealToUsd::CreateComponentPropertyBaker(
 
 				FocalLengthSetter(CameraComponent->CurrentFocalLength, UsdTimeCode);
 				FocusDistanceSetter(CameraComponent->FocusSettings.ManualFocusDistance, UsdTimeCode);
-				FStopSetter(CameraComponent->CurrentAperture, UsdTimeCode);
+
+				if (CameraComponent->PostProcessSettings.bOverride_DepthOfFieldFstop)
+				{
+					FStopSetter(CameraComponent->PostProcessSettings.DepthOfFieldFstop, UsdTimeCode);
+				}
+
 				HorizontalApertureSetter(CameraComponent->Filmback.SensorWidth, UsdTimeCode);
 				VerticalApertureSetter(CameraComponent->Filmback.SensorHeight, UsdTimeCode);
 				HorizontalApertureOffsetSetter(CameraComponent->Filmback.SensorHorizontalOffset, UsdTimeCode);
@@ -4897,6 +5037,18 @@ bool UnrealToUsd::CreateComponentPropertyBaker(
 				{
 					ClippingRangeSetter(UEClippingRange, UsdTimeCode);
 				}
+
+				ExposureFStopSetter(CameraComponent->CurrentAperture, UsdTimeCode);
+
+				if (CameraComponent->PostProcessSettings.bOverride_CameraISO)
+				{
+					ExposureISOSetter(CameraComponent->PostProcessSettings.CameraISO, UsdTimeCode);
+				}
+
+				if (CameraComponent->PostProcessSettings.bOverride_CameraShutterSpeed)
+				{
+					ExposureTimeSetter(CameraComponent->PostProcessSettings.CameraShutterSpeed, UsdTimeCode);
+				}
 			};
 		}
 	}
@@ -5628,6 +5780,7 @@ UnrealToUsd::FPropertyTrackWriter UnrealToUsd::CreatePropertyTrackWriter(
 
 		if (pxr::UsdGeomCamera Camera{UsdPrim})
 		{
+			bool bInvertValue = false;
 			bool bConvertDistance = true;
 			TOptional<int32> ClippingPlaneIndex;
 
@@ -5642,7 +5795,7 @@ UnrealToUsd::FPropertyTrackWriter UnrealToUsd::CreatePropertyTrackWriter(
 			else if (PropertyPath == UnrealIdentifiers::CurrentAperturePropertyName)
 			{
 				bConvertDistance = false;
-				Attr = Camera.CreateFStopAttr();
+				Attr = Camera.CreateExposureFStopAttr();
 			}
 			else if (PropertyPath == UnrealIdentifiers::SensorWidthPropertyName)
 			{
@@ -5680,6 +5833,22 @@ UnrealToUsd::FPropertyTrackWriter UnrealToUsd::CreatePropertyTrackWriter(
 				Attr = Camera.CreateClippingRangeAttr();
 				ClippingPlaneIndex = 0;
 			}
+			else if (PropertyPath == UnrealIdentifiers::DepthOfFieldFstopPropertyName)
+			{
+				bConvertDistance = false;
+				Attr = Camera.CreateFStopAttr();
+			}
+			else if (PropertyPath == UnrealIdentifiers::CameraISOPropertyName)
+			{
+				bConvertDistance = false;
+				Attr = Camera.CreateExposureIsoAttr();
+			}
+			else if (PropertyPath == UnrealIdentifiers::CameraShutterSpeedPropertyName)
+			{
+				bConvertDistance = false;
+				bInvertValue = true;
+				Attr = Camera.CreateExposureTimeAttr();
+			}
 
 			// TODO: Write out animation track for camera projection? Would need a new custom Enum writer type
 
@@ -5706,6 +5875,13 @@ UnrealToUsd::FPropertyTrackWriter UnrealToUsd::CreatePropertyTrackWriter(
 						Attr.Set(UnrealToUsd::ConvertDistance(StageInfo, UEValue), UsdTimeCode);
 					};
 				}
+				else if (bInvertValue)
+				{
+					Result.FloatWriter = [Attr](float UEValue, double UsdTimeCode) mutable
+					{
+						Attr.Set(UEValue != 0.0f ? (1.0f / UEValue) : 0.0f, UsdTimeCode);
+					};
+				}
 				else
 				{
 					Result.FloatWriter = [Attr](float UEValue, double UsdTimeCode) mutable
@@ -6473,7 +6649,7 @@ TArray<UE::FUsdAttribute> UsdUtils::GetAttributesForProperty(const UE::FUsdPrim&
 	}
 	else if (PropertyPath == CurrentAperturePropertyName)
 	{
-		return {UE::FUsdAttribute{UsdPrim.GetAttribute(pxr::UsdGeomTokens->fStop)}};
+		return {UE::FUsdAttribute{UsdPrim.GetAttribute(pxr::UsdGeomTokens->exposureFStop)}};
 	}
 	else if (PropertyPath == SensorWidthPropertyName)
 	{
@@ -6505,6 +6681,18 @@ TArray<UE::FUsdAttribute> UsdUtils::GetAttributesForProperty(const UE::FUsdPrim&
 	{
 		return {UE::FUsdAttribute{UsdPrim.GetAttribute(pxr::UsdGeomTokens->clippingRange)}};
 	}
+	else if (PropertyPath == DepthOfFieldFstopPropertyName)
+	{
+		return {UE::FUsdAttribute{UsdPrim.GetAttribute(pxr::UsdGeomTokens->fStop)}};
+	}
+	else if (PropertyPath == CameraISOPropertyName)
+	{
+		return {UE::FUsdAttribute{UsdPrim.GetAttribute(pxr::UsdGeomTokens->exposureIso)}};
+	}
+	else if (PropertyPath == CameraShutterSpeedPropertyName)
+	{
+		return {UE::FUsdAttribute{UsdPrim.GetAttribute(pxr::UsdGeomTokens->exposureTime)}};
+	}
 
 	// Light attributes
 	else if (PropertyPath == IntensityPropertyName)
@@ -6683,7 +6871,7 @@ TArray<FName> UsdUtils::GetPropertiesForAttribute(const UE::FUsdPrim& Prim, cons
 
 		{pxr::UsdGeomTokens->focalLength, 					{CurrentFocalLengthPropertyName}},
 		{pxr::UsdGeomTokens->focusDistance, 				{ManualFocusDistancePropertyName}},
-		{pxr::UsdGeomTokens->fStop, 						{CurrentAperturePropertyName}},
+		{pxr::UsdGeomTokens->fStop, 						{DepthOfFieldFstopPropertyName}},
 		{pxr::UsdGeomTokens->horizontalAperture, 			{SensorWidthPropertyName}},
 		{pxr::UsdGeomTokens->verticalAperture, 				{SensorHeightPropertyName}},
 		{pxr::UsdGeomTokens->horizontalApertureOffset, 		{SensorHorizontalOffsetPropertyName}},
@@ -6691,6 +6879,9 @@ TArray<FName> UsdUtils::GetPropertiesForAttribute(const UE::FUsdPrim& Prim, cons
 		{pxr::UsdGeomTokens->exposure, 						{ExposureCompensationPropertyName}},
 		{pxr::UsdGeomTokens->projection, 					{ProjectionModePropertyName}},
 		{pxr::UsdGeomTokens->clippingRange, 				{OrthoFarClipPlanePropertyName, OrthoNearClipPlanePropertyName, CustomNearClipppingPlanePropertyName}},
+		{pxr::UsdGeomTokens->exposureFStop, 				{CurrentAperturePropertyName}},
+		{pxr::UsdGeomTokens->exposureIso, 					{CameraISOPropertyName}},
+		{pxr::UsdGeomTokens->exposureTime, 					{CameraShutterSpeedPropertyName}},
 
 		{pxr::UsdLuxTokens->inputsIntensity, 				{IntensityPropertyName}},
 		{pxr::UsdLuxTokens->inputsExposure, 				{IntensityPropertyName}}, // In USD true intensity is just a function of "intensity" and "exposure"
diff --git a/Engine/Plugins/Runtime/USDCore/Source/UnrealUSDWrapper/Private/UnrealUSDWrapper.cpp b/Engine/Plugins/Runtime/USDCore/Source/UnrealUSDWrapper/Private/UnrealUSDWrapper.cpp
--- a/Engine/Plugins/Runtime/USDCore/Source/UnrealUSDWrapper/Private/UnrealUSDWrapper.cpp
+++ b/Engine/Plugins/Runtime/USDCore/Source/UnrealUSDWrapper/Private/UnrealUSDWrapper.cpp
@@ -457,6 +457,11 @@ FName UnrealIdentifiers::TransformPropertyName = TEXT("Transform");	   // Fake F
 FName UnrealIdentifiers::HiddenInGamePropertyName = GET_MEMBER_NAME_CHECKED(USceneComponent, bHiddenInGame);
 FName UnrealIdentifiers::HiddenPropertyName = AActor::GetHiddenPropertyName();
 
+FName UnrealIdentifiers::AspectRatioPropertyName = GET_MEMBER_NAME_CHECKED(UCameraComponent, AspectRatio);
+FName UnrealIdentifiers::OrthoWidthPropertyName = GET_MEMBER_NAME_CHECKED(UCameraComponent, OrthoWidth);
+FName UnrealIdentifiers::OrthoFarClipPlanePropertyName = GET_MEMBER_NAME_CHECKED(UCameraComponent, OrthoFarClipPlane);
+FName UnrealIdentifiers::OrthoNearClipPlanePropertyName = GET_MEMBER_NAME_CHECKED(UCameraComponent, OrthoNearClipPlane);
+
 FName UnrealIdentifiers::CurrentFocalLengthPropertyName = GET_MEMBER_NAME_CHECKED(UCineCameraComponent, CurrentFocalLength);
 FName UnrealIdentifiers::ManualFocusDistancePropertyName = GET_MEMBER_NAME_CHECKED(UCineCameraComponent, FocusSettings.ManualFocusDistance);
 FName UnrealIdentifiers::CurrentAperturePropertyName = GET_MEMBER_NAME_CHECKED(UCineCameraComponent, CurrentAperture);
@@ -466,9 +471,10 @@ FName UnrealIdentifiers::SensorHorizontalOffsetPropertyName = GET_MEMBER_NAME_CH
 FName UnrealIdentifiers::SensorVerticalOffsetPropertyName = GET_MEMBER_NAME_CHECKED(UCineCameraComponent, Filmback.SensorVerticalOffset);
 FName UnrealIdentifiers::ExposureCompensationPropertyName = GET_MEMBER_NAME_CHECKED(UCineCameraComponent, PostProcessSettings.AutoExposureBias);
 FName UnrealIdentifiers::ProjectionModePropertyName = GET_MEMBER_NAME_CHECKED(UCineCameraComponent, ProjectionMode);
-FName UnrealIdentifiers::OrthoFarClipPlanePropertyName = GET_MEMBER_NAME_CHECKED(UCineCameraComponent, OrthoFarClipPlane);
-FName UnrealIdentifiers::OrthoNearClipPlanePropertyName = GET_MEMBER_NAME_CHECKED(UCineCameraComponent, OrthoNearClipPlane);
 FName UnrealIdentifiers::CustomNearClipppingPlanePropertyName = GET_MEMBER_NAME_CHECKED(UCineCameraComponent, CustomNearClippingPlane);
+FName UnrealIdentifiers::DepthOfFieldFstopPropertyName = GET_MEMBER_NAME_CHECKED(UCineCameraComponent, PostProcessSettings.DepthOfFieldFstop);
+FName UnrealIdentifiers::CameraISOPropertyName = GET_MEMBER_NAME_CHECKED(UCineCameraComponent, PostProcessSettings.CameraISO);
+FName UnrealIdentifiers::CameraShutterSpeedPropertyName = GET_MEMBER_NAME_CHECKED(UCineCameraComponent, PostProcessSettings.CameraShutterSpeed);
 
 FName UnrealIdentifiers::IntensityPropertyName = GET_MEMBER_NAME_CHECKED(ULightComponentBase, Intensity);
 FName UnrealIdentifiers::LightColorPropertyName = GET_MEMBER_NAME_CHECKED(ULightComponentBase, LightColor);
diff --git a/Engine/Plugins/Runtime/USDCore/Source/UnrealUSDWrapper/Public/UnrealUSDWrapper.h b/Engine/Plugins/Runtime/USDCore/Source/UnrealUSDWrapper/Public/UnrealUSDWrapper.h
--- a/Engine/Plugins/Runtime/USDCore/Source/UnrealUSDWrapper/Public/UnrealUSDWrapper.h
+++ b/Engine/Plugins/Runtime/USDCore/Source/UnrealUSDWrapper/Public/UnrealUSDWrapper.h
@@ -616,6 +616,12 @@ namespace UnrealIdentifiers
 	extern UNREALUSDWRAPPER_API FName HiddenInGamePropertyName;
 	extern UNREALUSDWRAPPER_API FName HiddenPropertyName;
 
+	// UCameraComponent properties
+	extern UNREALUSDWRAPPER_API FName AspectRatioPropertyName;
+	extern UNREALUSDWRAPPER_API FName OrthoWidthPropertyName;
+	extern UNREALUSDWRAPPER_API FName OrthoFarClipPlanePropertyName;
+	extern UNREALUSDWRAPPER_API FName OrthoNearClipPlanePropertyName;
+
 	// UCineCameraComponent properties
 	extern UNREALUSDWRAPPER_API FName CurrentFocalLengthPropertyName;
 	extern UNREALUSDWRAPPER_API FName ManualFocusDistancePropertyName;
@@ -626,9 +632,10 @@ namespace UnrealIdentifiers
 	extern UNREALUSDWRAPPER_API FName SensorVerticalOffsetPropertyName;
 	extern UNREALUSDWRAPPER_API FName ExposureCompensationPropertyName;
 	extern UNREALUSDWRAPPER_API FName ProjectionModePropertyName;
-	extern UNREALUSDWRAPPER_API FName OrthoFarClipPlanePropertyName;
-	extern UNREALUSDWRAPPER_API FName OrthoNearClipPlanePropertyName;
 	extern UNREALUSDWRAPPER_API FName CustomNearClipppingPlanePropertyName;
+	extern UNREALUSDWRAPPER_API FName DepthOfFieldFstopPropertyName;
+	extern UNREALUSDWRAPPER_API FName CameraISOPropertyName;
+	extern UNREALUSDWRAPPER_API FName CameraShutterSpeedPropertyName;
 
 	// LightComponentBase properties
 	extern UNREALUSDWRAPPER_API FName IntensityPropertyName;
diff --git a/Engine/Source/Runtime/Interchange/Core/Public/Nodes/InterchangeFactoryBaseNode.h b/Engine/Source/Runtime/Interchange/Core/Public/Nodes/InterchangeFactoryBaseNode.h
--- a/Engine/Source/Runtime/Interchange/Core/Public/Nodes/InterchangeFactoryBaseNode.h
+++ b/Engine/Source/Runtime/Interchange/Core/Public/Nodes/InterchangeFactoryBaseNode.h
@@ -417,6 +417,16 @@ protected:
 	INTERCHANGECORE_API virtual void OnRestoreAllCustomAttributeDelegates();
 
 protected:
+	/**
+	 * Writes a value to an UObject's property.
+	 * @param Value				The value to write.
+	 * @param Object			The object to write to.
+	 * @param PropertyName		The name of the property to write to on the Object.
+	 * @result					True if the operation succeeded.
+	 */
+	template<typename PropertyType>
+	inline bool ApplyPropertyValueToObject(const PropertyType& Value, UObject* Object, const FName PropertyName) const;
+
 	/**
 	 * Those dependencies are used by the Interchange parsing task to make sure the assets are created in the correct order.
 	 * Examples: The Mesh factory node will have dependencies on the material factory node, and the Material factory node will have dependencies on the texture factory node.
@@ -473,131 +483,98 @@ inline bool UInterchangeFactoryBaseNode::ApplyAttributeToObject(const FString& N
 	AttributeType ValueData;
 	if (GetAttribute<AttributeType>(NodeAttributeKey, ValueData))
 	{
-		TVariant<UObject*, uint8*> Container;
-		Container.Set<UObject*>(Object);
-		if (FProperty* Property = InterchangePrivateNodeBase::FindPropertyByPathChecked(Container, Object->GetClass(), PropertyName.ToString()))
-		{
-			AttributeType* PropertyValueAddress;
-			if (Container.IsType<UObject*>())
-			{
-				PropertyValueAddress = Property->ContainerPtrToValuePtr<AttributeType>(Container.Get<UObject*>());
-			}
-			else
-			{
-				PropertyValueAddress = Property->ContainerPtrToValuePtr<AttributeType>(Container.Get<uint8*>());
-			}
-
-			*PropertyValueAddress = ValueData;
-		}
-		return true;
+		return ApplyPropertyValueToObject(ValueData, Object, PropertyName);
 	}
+
 	return false;
 }
 
-/**
- * Specialized version of ApplyAttributeToObject for strings.
- * If the target property is an FObjectPropertyBase, treat the string as an object path.
- */
-template<>
-inline bool UInterchangeFactoryBaseNode::ApplyAttributeToObject<FString>(const FString& NodeAttributeKey, UObject* Object, const FName PropertyName) const
+template<typename AttributeType>
+inline bool UInterchangeFactoryBaseNode::FillAttributeFromObject(const FString& NodeAttributeKey, UObject* Object, const FName PropertyName)
 {
-	if (!Object)
-	{
-		return false;
-	}
-
-	FString ValueData;
-	if (GetAttribute<FString>(NodeAttributeKey, ValueData))
+	TVariant<UObject*, uint8*> Container;
+	Container.Set<UObject*>(Object);
+	if (FProperty* Property = InterchangePrivateNodeBase::FindPropertyByPathChecked(Container, Object->GetClass(), PropertyName.ToString()))
 	{
-		TVariant<UObject*, uint8*> Container;
-		Container.Set<UObject*>(Object);
-		if (FProperty* Property = InterchangePrivateNodeBase::FindPropertyByPathChecked(Container, Object->GetClass(), PropertyName.ToString()))
+		AttributeType* PropertyValueAddress;
+		if (Container.IsType<UObject*>())
 		{
-			FString* PropertyValueAddress;
-			if (Container.IsType<UObject*>())
-			{
-				PropertyValueAddress = Property->ContainerPtrToValuePtr<FString>(Container.Get<UObject*>());
-			}
-			else
-			{
-				PropertyValueAddress = Property->ContainerPtrToValuePtr<FString>(Container.Get<uint8*>());
-			}
-
-			if (FObjectPropertyBase* ObjectProperty = CastField<FObjectPropertyBase>(Property))
-			{
-				ObjectProperty->SetObjectPropertyValue(PropertyValueAddress, FSoftObjectPath(ValueData).TryLoad());
-			}
-			else
-			{
-				*PropertyValueAddress = ValueData;
-			}
+			PropertyValueAddress = Property->ContainerPtrToValuePtr<AttributeType>(Container.Get<UObject*>());
+		}
+		else
+		{
+			PropertyValueAddress = Property->ContainerPtrToValuePtr<AttributeType>(Container.Get<uint8*>());
+		}
 
-			return true;
+		// Support for bitfields
+		if (FBoolProperty* BoolProperty = CastField<FBoolProperty>(Property))
+		{
+			const bool bPropertyValue = BoolProperty->GetPropertyValue(PropertyValueAddress);
+			return SetAttribute(NodeAttributeKey, bPropertyValue);
+		}
+		else
+		{
+			return SetAttribute(NodeAttributeKey, *PropertyValueAddress);
 		}
 	}
+
 	return false;
 }
 
 /**
- * Specialized version of ApplyAttributeToObject for bools.
- * If the target property is an FBoolProperty, treat the property as a bitfield.
+ * Specialized version of FillAttributeFromObject for strings.
+ * If the target property is an FObjectPropertyBase, treat the string as an object path.
  */
 template<>
-inline bool UInterchangeFactoryBaseNode::ApplyAttributeToObject<bool>(const FString& NodeAttributeKey, UObject* Object, const FName PropertyName) const
+inline bool UInterchangeFactoryBaseNode::FillAttributeFromObject<FString>(const FString& NodeAttributeKey, UObject* Object, const FName PropertyName)
 {
-	if (!Object)
-	{
-		return false;
-	}
-
-	bool bValueData;
-	if (GetAttribute<bool>(NodeAttributeKey, bValueData))
+	TVariant<UObject*, uint8*> Container;
+	Container.Set<UObject*>(Object);
+	if (FProperty* Property = InterchangePrivateNodeBase::FindPropertyByPathChecked(Container, Object->GetClass(), PropertyName.ToString()))
 	{
-		TVariant<UObject*, uint8*> Container;
-		Container.Set<UObject*>(Object);
-		if (FProperty* Property = InterchangePrivateNodeBase::FindPropertyByPathChecked(Container, Object->GetClass(), PropertyName.ToString()))
+		FString* PropertyValueAddress;
+		if (Container.IsType<UObject*>())
 		{
-			bool* PropertyValueAddress;
-			if (Container.IsType<UObject*>())
-			{
-				PropertyValueAddress = Property->ContainerPtrToValuePtr<bool>(Container.Get<UObject*>());
-			}
-			else
-			{
-				PropertyValueAddress = Property->ContainerPtrToValuePtr<bool>(Container.Get<uint8*>());
-			}
+			PropertyValueAddress = Property->ContainerPtrToValuePtr<FString>(Container.Get<UObject*>());
+		}
+		else
+		{
+			PropertyValueAddress = Property->ContainerPtrToValuePtr<FString>(Container.Get<uint8*>());
+		}
 
-			// Support for bitfields
-			if (FBoolProperty* BoolProperty = CastField<FBoolProperty>(Property))
-			{
-				BoolProperty->SetPropertyValue(PropertyValueAddress, bValueData);
-			}
-			else
-			{
-				*PropertyValueAddress = bValueData;
-			}
-			
-			return true;
+		if (FObjectPropertyBase* ObjectProperty = CastField<FObjectPropertyBase>(Property))
+		{
+			UObject* ObjectValue = ObjectProperty->GetObjectPropertyValue(PropertyValueAddress);
+			return SetAttribute(NodeAttributeKey, ObjectValue->GetPathName());
+		}
+		else
+		{
+			return SetAttribute(NodeAttributeKey, *PropertyValueAddress);
 		}
 	}
+
 	return false;
 }
 
-template<typename AttributeType>
-inline bool UInterchangeFactoryBaseNode::FillAttributeFromObject(const FString& NodeAttributeKey, UObject* Object, const FName PropertyName)
+/**
+ * Specialized version of FillAttributeFromObject for bools.
+ * If the target property is an FBoolProperty, treat the property as a bitfield.
+ */
+template<>
+inline bool UInterchangeFactoryBaseNode::FillAttributeFromObject<bool>(const FString& NodeAttributeKey, UObject* Object, const FName PropertyName)
 {
 	TVariant<UObject*, uint8*> Container;
 	Container.Set<UObject*>(Object);
 	if (FProperty* Property = InterchangePrivateNodeBase::FindPropertyByPathChecked(Container, Object->GetClass(), PropertyName.ToString()))
 	{
-		AttributeType* PropertyValueAddress;
+		bool* PropertyValueAddress;
 		if (Container.IsType<UObject*>())
 		{
-			PropertyValueAddress = Property->ContainerPtrToValuePtr<AttributeType>(Container.Get<UObject*>());
+			PropertyValueAddress = Property->ContainerPtrToValuePtr<bool>(Container.Get<UObject*>());
 		}
 		else
 		{
-			PropertyValueAddress = Property->ContainerPtrToValuePtr<AttributeType>(Container.Get<uint8*>());
+			PropertyValueAddress = Property->ContainerPtrToValuePtr<bool>(Container.Get<uint8*>());
 		}
 
 		// Support for bitfields
@@ -615,13 +592,48 @@ inline bool UInterchangeFactoryBaseNode::FillAttributeFromObject(const FString&
 	return false;
 }
 
+template<typename PropertyType>
+inline bool UInterchangeFactoryBaseNode::ApplyPropertyValueToObject(const PropertyType& Value, UObject* Object, const FName PropertyName) const
+{
+	if (!Object)
+	{
+		return false;
+	}
+
+	TVariant<UObject*, uint8*> Container;
+	Container.Set<UObject*>(Object);
+	if (FProperty* Property = InterchangePrivateNodeBase::FindPropertyByPathChecked(Container, Object->GetClass(), PropertyName.ToString()))
+	{
+		PropertyType* PropertyValueAddress;
+		if (Container.IsType<UObject*>())
+		{
+			PropertyValueAddress = Property->ContainerPtrToValuePtr<PropertyType>(Container.Get<UObject*>());
+		}
+		else
+		{
+			PropertyValueAddress = Property->ContainerPtrToValuePtr<PropertyType>(Container.Get<uint8*>());
+		}
+
+		*PropertyValueAddress = Value;
+
+		return true;
+	}
+
+	return false;
+}
+
 /**
- * Specialized version of FillAttributeFromObject for strings.
+ * Specialized version of ApplyPropertyValueToObject for strings.
  * If the target property is an FObjectPropertyBase, treat the string as an object path.
  */
 template<>
-inline bool UInterchangeFactoryBaseNode::FillAttributeFromObject<FString>(const FString& NodeAttributeKey, UObject* Object, const FName PropertyName)
+inline bool UInterchangeFactoryBaseNode::ApplyPropertyValueToObject<FString>(const FString& Value, UObject* Object, const FName PropertyName) const
 {
+	if (!Object)
+	{
+		return false;
+	}
+
 	TVariant<UObject*, uint8*> Container;
 	Container.Set<UObject*>(Object);
 	if (FProperty* Property = InterchangePrivateNodeBase::FindPropertyByPathChecked(Container, Object->GetClass(), PropertyName.ToString()))
@@ -638,25 +650,30 @@ inline bool UInterchangeFactoryBaseNode::FillAttributeFromObject<FString>(const
 
 		if (FObjectPropertyBase* ObjectProperty = CastField<FObjectPropertyBase>(Property))
 		{
-			UObject* ObjectValue = ObjectProperty->GetObjectPropertyValue(PropertyValueAddress);
-			return SetAttribute(NodeAttributeKey, ObjectValue->GetPathName());
+			ObjectProperty->SetObjectPropertyValue(PropertyValueAddress, FSoftObjectPath(Value).TryLoad());
 		}
 		else
 		{
-			return SetAttribute(NodeAttributeKey, *PropertyValueAddress);
+			*PropertyValueAddress = Value;
 		}
-	}
 
+		return true;
+	}
 	return false;
 }
 
 /**
- * Specialized version of FillAttributeFromObject for bools.
+ * Specialized version of ApplyPropertyValueToObject for bools.
  * If the target property is an FBoolProperty, treat the property as a bitfield.
  */
 template<>
-inline bool UInterchangeFactoryBaseNode::FillAttributeFromObject<bool>(const FString& NodeAttributeKey, UObject* Object, const FName PropertyName)
+inline bool UInterchangeFactoryBaseNode::ApplyPropertyValueToObject<bool>(const bool& Value, UObject* Object, const FName PropertyName) const
 {
+	if (!Object)
+	{
+		return false;
+	}
+
 	TVariant<UObject*, uint8*> Container;
 	Container.Set<UObject*>(Object);
 	if (FProperty* Property = InterchangePrivateNodeBase::FindPropertyByPathChecked(Container, Object->GetClass(), PropertyName.ToString()))
@@ -674,14 +691,14 @@ inline bool UInterchangeFactoryBaseNode::FillAttributeFromObject<bool>(const FSt
 		// Support for bitfields
 		if (FBoolProperty* BoolProperty = CastField<FBoolProperty>(Property))
 		{
-			const bool bPropertyValue = BoolProperty->GetPropertyValue(PropertyValueAddress);
-			return SetAttribute(NodeAttributeKey, bPropertyValue);
+			BoolProperty->SetPropertyValue(PropertyValueAddress, Value);
 		}
 		else
 		{
-			return SetAttribute(NodeAttributeKey, *PropertyValueAddress);
+			*PropertyValueAddress = Value;
 		}
-	}
 
+		return true;
+	}
 	return false;
-}
\ No newline at end of file
+}
