[Interchange] ufbx fixed opacity

#Synced-CL 47926590

diff --git a/Engine/Plugins/Interchange/Runtime/Source/Parsers/Fbx/Private/Ufbx/UfbxMaterial.cpp b/Engine/Plugins/Interchange/Runtime/Source/Parsers/Fbx/Private/Ufbx/UfbxMaterial.cpp
--- a/Engine/Plugins/Interchange/Runtime/Source/Parsers/Fbx/Private/Ufbx/UfbxMaterial.cpp
+++ b/Engine/Plugins/Interchange/Runtime/Source/Parsers/Fbx/Private/Ufbx/UfbxMaterial.cpp
@@ -370,7 +370,7 @@ namespace UE::Interchange::Private
 					MaterialInstanceNode->AddStaticSwitchParameterValue(HasOpacity, true);
 					MaterialInstanceNode->SetCustomBlendMode(BlendMode);
 					
-					bHasInput |= SetTextureFromMap(InputName, ColorMap, Factor);
+					bHasInput |= SetTextureFromMap(InputName, ColorMap, 1.0f);
 				}
 				else
 				{
@@ -378,6 +378,7 @@ namespace UE::Interchange::Private
 					const float OpacityScalar = 1.f - FMath::Clamp(TransparencyScalar, 0.f, 1.f);
 					if (OpacityScalar < 1.f)
 					{
+						MaterialInstanceNode->SetCustomBlendMode(EBlendMode::BLEND_Translucent);
 						bHasInput |= MaterialInstanceNode->AddScalarParameterValue(InputName.ToString(), OpacityScalar);
 					}
 				}
