Interchange MaterialX:
- Removed matching <saturate> with Desaturation Material Expression as it was wrong
- Add support of lumacoeffs for <luminance>
- Fixed an error in MX_ColorCorrect where it used a Desaturation instead of a saturation

#virtualized
#Synced-CL 42983915

diff --git a/Engine/Plugins/Interchange/Runtime/Source/Import/Private/MaterialX/MaterialXUtils/MaterialXManager.cpp b/Engine/Plugins/Interchange/Runtime/Source/Import/Private/MaterialX/MaterialXUtils/MaterialXManager.cpp
--- a/Engine/Plugins/Interchange/Runtime/Source/Import/Private/MaterialX/MaterialXUtils/MaterialXManager.cpp
+++ b/Engine/Plugins/Interchange/Runtime/Source/Import/Private/MaterialX/MaterialXUtils/MaterialXManager.cpp
@@ -109,7 +109,6 @@ FMaterialXManager::FMaterialXManager()
 	   {{mx::Category::Rotate3D,              TEXT("amount")},      ExpressionInput(UE::Expressions::Inputs::RotationAngle)},
 	   {{mx::Category::Rotate3D,              TEXT("axis")},		ExpressionInput(UE::Expressions::Inputs::NormalizedRotationAxis)},
 	   {{mx::Category::Rotate3D,              TEXT("in")},          ExpressionInput(UE::Expressions::Inputs::Position)},
-	   {{mx::Category::Saturate,              TEXT("amount")},      ExpressionInput(UE::Expressions::Inputs::Fraction)},
 	   {{mx::Category::Smoothstep,            TEXT("in")},          ExpressionInput(UE::Expressions::Inputs::Value)},
 	   {{mx::Category::Separate2,             TEXT("in")},          ExpressionInput(UE::Expressions::Inputs::Float2)},
 	   {{mx::Category::Separate3,             TEXT("in")},          ExpressionInput(UE::Expressions::Inputs::Float3)},
@@ -199,11 +198,9 @@ FMaterialXManager::FMaterialXManager()
 		// Adjustment nodes,
 		{mx::Category::Contrast,     UE::Expressions::Names::Contrast},
 		{mx::Category::HsvToRgb,     UE::Expressions::Names::HsvToRgb},
-		{mx::Category::Luminance,    UE::Expressions::Names::MaterialXLuminance},
 		{mx::Category::Range,        UE::Expressions::Names::MaterialXRange},
 		{mx::Category::Remap,        UE::Expressions::Names::MaterialXRemap},
 		{mx::Category::RgbToHsv,     UE::Expressions::Names::RgbToHsv},
-		{mx::Category::Saturate,     UE::Expressions::Names::Desaturation},
 		{mx::Category::Smoothstep,   UE::Expressions::Names::SmoothStep}
 	}
 	, MaterialXContainerDelegates {
@@ -317,7 +314,6 @@ FMaterialXManager::FMaterialXManager()
 		mx::Category::RandomColor,
 		mx::Category::Range,
 		mx::Category::Refract,
-		mx::Category::Saturate,
 		mx::Category::Separate2,
 		mx::Category::Separate3,
 		mx::Category::Separate4,
diff --git a/Engine/Plugins/Interchange/Runtime/Source/Import/Private/MaterialX/MaterialXUtils/MaterialXSurfaceShaderAbstract.cpp b/Engine/Plugins/Interchange/Runtime/Source/Import/Private/MaterialX/MaterialXUtils/MaterialXSurfaceShaderAbstract.cpp
--- a/Engine/Plugins/Interchange/Runtime/Source/Import/Private/MaterialX/MaterialXUtils/MaterialXSurfaceShaderAbstract.cpp
+++ b/Engine/Plugins/Interchange/Runtime/Source/Import/Private/MaterialX/MaterialXUtils/MaterialXSurfaceShaderAbstract.cpp
@@ -1440,6 +1440,21 @@ void FMaterialXSurfaceShaderAbstract::ConnectSparseVolumeInputToOutput(const FCo
 	UInterchangeShaderPortsAPI::ConnectOuputToInputByIndex(Connect.ParentShaderNode, Connect.InputChannelName, MaskNode->GetUniqueID(), OutputIndex);
 }
 
+void FMaterialXSurfaceShaderAbstract::ConnectLuminanceInputToOutput(const FConnectNode& Connect)
+{
+	using namespace UE::Interchange::Materials::Standard::Nodes;
+	UInterchangeShaderNode* LuminanceNode = CreateShaderNode(Connect.UpstreamNode, Connect.UpstreamNode->getName().c_str(), Luminance::Name.ToString());
+	AddAttributeFromValueOrInterface(Connect.UpstreamNode->getInput("in"), Luminance::Inputs::Input.ToString(), LuminanceNode);
+
+	if (mx::InputPtr InputCoeffs = Connect.UpstreamNode->getInput("lumacoeffs"); InputCoeffs && InputCoeffs->hasValue())
+	{
+		mx::Color3 LumaCoeffs = InputCoeffs->getValue()->asA<mx::Color3>();
+		LuminanceNode->AddLinearColorAttribute(Luminance::Attributes::LuminanceFactors.ToString(), FLinearColor{ LumaCoeffs[0], LumaCoeffs[1] ,LumaCoeffs[2] });
+	}
+
+	UInterchangeShaderPortsAPI::ConnectDefaultOuputToInput(Connect.ParentShaderNode, Connect.InputChannelName, LuminanceNode->GetUniqueID());
+}
+
 UInterchangeShaderNode* FMaterialXSurfaceShaderAbstract::CreateMaskShaderNode(uint8 RGBA, mx::ElementPtr Element, const FString& NodeName, const FString& OutputName)
 {
 	bool bR = (0b1000 & RGBA) >> 3;
@@ -1667,6 +1682,7 @@ void FMaterialXSurfaceShaderAbstract::RegisterConnectNodeOutputToInputDelegates(
 	MatchingConnectNodeDelegates.Add(mx::Category::Xor,				FOnConnectNodeOutputToInput::CreateSP(this, &FMaterialXSurfaceShaderAbstract::ConnectXorInputToOutput));
 	MatchingConnectNodeDelegates.Add(mx::Category::Not,				FOnConnectNodeOutputToInput::CreateSP(this, &FMaterialXSurfaceShaderAbstract::ConnectNotInputToOutput));
 	MatchingConnectNodeDelegates.Add(mx::Category::SparseVolume,	FOnConnectNodeOutputToInput::CreateSP(this, &FMaterialXSurfaceShaderAbstract::ConnectSparseVolumeInputToOutput));
+	MatchingConnectNodeDelegates.Add(mx::Category::Luminance,		FOnConnectNodeOutputToInput::CreateSP(this, &FMaterialXSurfaceShaderAbstract::ConnectLuminanceInputToOutput));
 }
 
 void FMaterialXSurfaceShaderAbstract::SetMatchingInputsNames(MaterialX::NodePtr Node) const
diff --git a/Engine/Plugins/Interchange/Runtime/Source/Import/Public/MaterialX/MaterialXUtils/MaterialXSurfaceShaderAbstract.h b/Engine/Plugins/Interchange/Runtime/Source/Import/Public/MaterialX/MaterialXUtils/MaterialXSurfaceShaderAbstract.h
--- a/Engine/Plugins/Interchange/Runtime/Source/Import/Public/MaterialX/MaterialXUtils/MaterialXSurfaceShaderAbstract.h
+++ b/Engine/Plugins/Interchange/Runtime/Source/Import/Public/MaterialX/MaterialXUtils/MaterialXSurfaceShaderAbstract.h
@@ -339,6 +339,9 @@ protected:
 	/** <sparse_volume> is define in Engine/Binaries/ThirdParty/MaterialX/libraries/Interchange */
 	void ConnectSparseVolumeInputToOutput(const FConnectNode& Connect);
 
+	/** <luminance>*/
+	void ConnectLuminanceInputToOutput(const FConnectNode& Connect);
+
 	/** End Connect MaterialX nodes*/
 
 	/**
diff --git a/Engine/Plugins/Interchange/Runtime/Source/Nodes/Public/InterchangeMaterialDefinitions.h b/Engine/Plugins/Interchange/Runtime/Source/Nodes/Public/InterchangeMaterialDefinitions.h
--- a/Engine/Plugins/Interchange/Runtime/Source/Nodes/Public/InterchangeMaterialDefinitions.h
+++ b/Engine/Plugins/Interchange/Runtime/Source/Nodes/Public/InterchangeMaterialDefinitions.h
@@ -48,6 +48,21 @@ namespace Interchange
 					}
 				}
 
+				namespace Luminance
+				{
+					const FName Name = TEXT("Luminance");
+
+					namespace Inputs
+					{
+						const FName Input = TEXT("Input");
+					}
+
+					namespace Attributes
+					{
+						const FName LuminanceFactors = TEXT("LuminanceFactors");
+					}
+				}
+
 				namespace Mask
 				{
 					const FName Name = TEXT("ComponentMask");
diff --git a/Engine/Plugins/Interchange/Runtime/Source/Pipelines/Private/InterchangeGenericMaterialPipeline.cpp b/Engine/Plugins/Interchange/Runtime/Source/Pipelines/Private/InterchangeGenericMaterialPipeline.cpp
--- a/Engine/Plugins/Interchange/Runtime/Source/Pipelines/Private/InterchangeGenericMaterialPipeline.cpp
+++ b/Engine/Plugins/Interchange/Runtime/Source/Pipelines/Private/InterchangeGenericMaterialPipeline.cpp
@@ -4,6 +4,7 @@
 
 #include "AssetRegistry/AssetData.h"
 #include "AssetRegistry/AssetRegistryModule.h"
+#include "Import/Private/MaterialX/MaterialExpressions/MaterialExpressionLuminance.h"
 #include "InterchangeAssetImportData.h"
 #include "InterchangeGenericTexturePipeline.h"
 #include "InterchangeManager.h"
@@ -3051,6 +3052,32 @@ void UInterchangeGenericMaterialPipeline::HandleSlabBSDFNode(const UInterchangeS
 	AddSpecularProfileToFactoryNode(ShaderNode, SlabBSDFFactoryNode, BaseNodeContainer);
 }
 
+void UInterchangeGenericMaterialPipeline::HandleLuminanceNode(const UInterchangeShaderNode* ShaderNode, UInterchangeBaseMaterialFactoryNode* MaterialFactoryNode, UInterchangeMaterialExpressionFactoryNode* LuminanceFactoryNode)
+{
+	using namespace UE::Interchange::Materials::Standard::Nodes;
+	LuminanceFactoryNode->SetCustomExpressionClassName(UMaterialExpressionMaterialXLuminance::StaticClass()->GetName());
+
+	// Input
+	{
+		TTuple<UInterchangeMaterialExpressionFactoryNode*, FString> InputExpression =
+			CreateMaterialExpressionForInput(MaterialFactoryNode, ShaderNode, Luminance::Inputs::Input.ToString(), LuminanceFactoryNode->GetUniqueID());
+
+		if (InputExpression.Get<0>())
+		{
+			UInterchangeShaderPortsAPI::ConnectOuputToInputByName(LuminanceFactoryNode, GET_MEMBER_NAME_CHECKED(UMaterialExpressionMaterialXLuminance, Input).ToString(),
+																  InputExpression.Get<0>()->GetUniqueID(), InputExpression.Get<1>());
+		}
+	}
+
+	// LuminanceFactors
+	if (FLinearColor LuminanceFactors; ShaderNode->GetLinearColorAttribute(Luminance::Attributes::LuminanceFactors.ToString(), LuminanceFactors))
+	{
+		const FName LuminanceFactorsMemberName = GET_MEMBER_NAME_CHECKED(UMaterialExpressionMaterialXLuminance, LuminanceFactors);
+		LuminanceFactoryNode->AddLinearColorAttribute(LuminanceFactorsMemberName.ToString(), LuminanceFactors);
+		LuminanceFactoryNode->AddApplyAndFillDelegates<FLinearColor>(LuminanceFactorsMemberName.ToString(), UMaterialExpressionMaterialXLuminance::StaticClass(), LuminanceFactorsMemberName);
+	}
+}
+
 /** Trigonometry expressions are unitless and sets the period to 1 by default to let the user sets their period*/
 void UInterchangeGenericMaterialPipeline::HandleTrigonometryNode(const UInterchangeShaderNode* ShaderNode, UClass* StaticClass, UInterchangeBaseMaterialFactoryNode* MaterialFactoryNode, UInterchangeMaterialExpressionFactoryNode* TrigonometryFactoryNode)
 {
@@ -3187,6 +3214,10 @@ UInterchangeMaterialExpressionFactoryNode* UInterchangeGenericMaterialPipeline::
 	{
 		HandleLerpNode(ShaderNode, MaterialFactoryNode, MaterialExpression);
 	}
+	else if (*ShaderType == Nodes::Luminance::Name)
+	{
+		HandleLuminanceNode(ShaderNode, MaterialFactoryNode, MaterialExpression);
+	}
 	else if(*ShaderType == Nodes::Mask::Name)
 	{
 		HandleMaskNode(ShaderNode, MaterialFactoryNode, MaterialExpression);
diff --git a/Engine/Plugins/Interchange/Runtime/Source/Pipelines/Public/InterchangeGenericMaterialPipeline.h b/Engine/Plugins/Interchange/Runtime/Source/Pipelines/Public/InterchangeGenericMaterialPipeline.h
--- a/Engine/Plugins/Interchange/Runtime/Source/Pipelines/Public/InterchangeGenericMaterialPipeline.h
+++ b/Engine/Plugins/Interchange/Runtime/Source/Pipelines/Public/InterchangeGenericMaterialPipeline.h
@@ -230,6 +230,7 @@ private:
 	UE_API void HandleVectorNoiseNode(const UInterchangeShaderNode* ShaderNode, UInterchangeBaseMaterialFactoryNode* MaterialFactoryNode, UInterchangeMaterialExpressionFactoryNode* NoiseFactoryNode);
 	UE_API void HandleSwitchNode(const UInterchangeShaderNode* ShaderNode, UInterchangeBaseMaterialFactoryNode* MaterialFactoryNode, UInterchangeMaterialExpressionFactoryNode* SwitchFactoryNode);
 	UE_API void HandleSlabBSDFNode(const UInterchangeShaderNode* ShaderNode, UInterchangeBaseMaterialFactoryNode* MaterialFactoryNode, UInterchangeMaterialExpressionFactoryNode* SlabBSDFFactoryNode);
+    UE_API void HandleLuminanceNode(const UInterchangeShaderNode* ShaderNode, UInterchangeBaseMaterialFactoryNode* MaterialFactoryNode, UInterchangeMaterialExpressionFactoryNode* LuminanceFactoryNode);
 	UE_API void HandleTrigonometryNode(const UInterchangeShaderNode* ShaderNode, UClass* StaticClass, UInterchangeBaseMaterialFactoryNode* MaterialFactoryNode, UInterchangeMaterialExpressionFactoryNode* TrigonometryFactoryNode);
 
 	UE_API void HandleScalarParameterNode(const UInterchangeShaderNode* ShaderNode, UInterchangeMaterialExpressionFactoryNode* ScalarParameterFactoryNode);
