[Interchange-FBX] BInd pose validation tracking fixes for bind pose recreation use cases.

[Interchange-FBX] Introduced FbxCluster vs FbxPose BindPose acquisition/generation.
Also providing warning in case the 2 does not match.
Introduced new FBX Translator option to prioritize FbxCluster above FbxPoses for BindPose acquisitions.

#rnx

#Synced-CL 48034546

diff --git a/Engine/Plugins/Interchange/Runtime/Source/Dispatcher/Private/FBX/InterchangeDispatcherFBXTasks.cpp b/Engine/Plugins/Interchange/Runtime/Source/Dispatcher/Private/FBX/InterchangeDispatcherFBXTasks.cpp
--- a/Engine/Plugins/Interchange/Runtime/Source/Dispatcher/Private/FBX/InterchangeDispatcherFBXTasks.cpp
+++ b/Engine/Plugins/Interchange/Runtime/Source/Dispatcher/Private/FBX/InterchangeDispatcherFBXTasks.cpp
@@ -11,6 +11,7 @@ namespace UE::Interchange
 	const TCHAR* FJsonFBXLoadSourceCmd::ForceFrontXAxisJsonKey = TEXT("ForceFrontXAxis");
 	const TCHAR* FJsonFBXLoadSourceCmd::ConvertSceneUnitJsonKey = TEXT("ConvertSceneUnit");
 	const TCHAR* FJsonFBXLoadSourceCmd::KeepFbxNamespaceJsonKey = TEXT("KeepFbxNamespace");
+	const TCHAR* FJsonFBXLoadSourceCmd::UseFbxClustersForBindPoseMatricesJsonKey = TEXT("UseFbxClustersForBindPoseMatrices");
 
 	TSharedPtr<FJsonObject> FJsonFBXLoadSourceCmd::GetActionDataObject() const
 	{
@@ -19,6 +20,7 @@ namespace UE::Interchange
 		ActionDataObject->SetBoolField(ForceFrontXAxisJsonKey, bForceFrontXAxis);
 		ActionDataObject->SetBoolField(ConvertSceneUnitJsonKey, bConvertSceneUnit);
 		ActionDataObject->SetBoolField(KeepFbxNamespaceJsonKey, bKeepFbxNamespace);
+		ActionDataObject->SetBoolField(UseFbxClustersForBindPoseMatricesJsonKey, bConsiderClusterBeforePoseForMeshBindPose);
 
 		return ActionDataObject;
 	}
@@ -42,6 +44,11 @@ namespace UE::Interchange
 			return false;
 		}
 
+		if (!(ActionDataObject.TryGetBoolField(UseFbxClustersForBindPoseMatricesJsonKey, bConsiderClusterBeforePoseForMeshBindPose)))
+		{
+			return false;
+		}
+
 		return true;
 	}
 }//ns UE
diff --git a/Engine/Plugins/Interchange/Runtime/Source/Dispatcher/Public/FBX/InterchangeDispatcherFBXTasks.h b/Engine/Plugins/Interchange/Runtime/Source/Dispatcher/Public/FBX/InterchangeDispatcherFBXTasks.h
--- a/Engine/Plugins/Interchange/Runtime/Source/Dispatcher/Public/FBX/InterchangeDispatcherFBXTasks.h
+++ b/Engine/Plugins/Interchange/Runtime/Source/Dispatcher/Public/FBX/InterchangeDispatcherFBXTasks.h
@@ -17,12 +17,14 @@ namespace UE::Interchange
 			, const bool InbConvertScene
 			, const bool InbForceFrontXAxis
 			, const bool InbConvertSceneUnit
-			, const bool InbKeepFbxNamespace)
+			, const bool InbKeepFbxNamespace
+			, const bool InbConsiderClusterBeforePoseForMeshBindPose)
 			: FJsonLoadSourceCmd(InTranslatorID, InSourceFilename)
 			, bConvertScene(InbConvertScene)
 			, bForceFrontXAxis(InbForceFrontXAxis)
 			, bConvertSceneUnit(InbConvertSceneUnit)
 			, bKeepFbxNamespace(InbKeepFbxNamespace)
+			, bConsiderClusterBeforePoseForMeshBindPose(InbConsiderClusterBeforePoseForMeshBindPose)
 		{
 		}
 
@@ -57,6 +59,13 @@ namespace UE::Interchange
 			return bKeepFbxNamespace;
 		}
 
+		bool GetDoesUseFbxClustersForBindPoseMatrices() const
+		{
+			//Code should not do query data if the data was not set before
+			ensure(bIsDataInitialize);
+			return bConsiderClusterBeforePoseForMeshBindPose;
+		}
+
 		INTERCHANGEDISPATCHER_API static const TCHAR* TaskName;
 
 	private:
@@ -64,10 +73,12 @@ namespace UE::Interchange
 		bool bForceFrontXAxis = false;
 		bool bConvertSceneUnit = true;
 		bool bKeepFbxNamespace = false;
+		bool bConsiderClusterBeforePoseForMeshBindPose = false;
 
 		static const TCHAR* ConvertSceneJsonKey;
 		static const TCHAR* ForceFrontXAxisJsonKey;
 		static const TCHAR* ConvertSceneUnitJsonKey;
 		static const TCHAR* KeepFbxNamespaceJsonKey;
+		static const TCHAR* UseFbxClustersForBindPoseMatricesJsonKey;
 	};
 }//ns UE
diff --git a/Engine/Plugins/Interchange/Runtime/Source/Import/Private/Fbx/InterchangeFbxTranslator.cpp b/Engine/Plugins/Interchange/Runtime/Source/Import/Private/Fbx/InterchangeFbxTranslator.cpp
--- a/Engine/Plugins/Interchange/Runtime/Source/Import/Private/Fbx/InterchangeFbxTranslator.cpp
+++ b/Engine/Plugins/Interchange/Runtime/Source/Import/Private/Fbx/InterchangeFbxTranslator.cpp
@@ -309,6 +309,7 @@ bool UInterchangeFbxTranslator::Translate(UInterchangeBaseNodeContainer& BaseNod
 	const bool bForceFrontXAxis = CacheFbxTranslatorSettings ? CacheFbxTranslatorSettings->bForceFrontXAxis : false;
 	const bool bConvertSceneUnit = CacheFbxTranslatorSettings ? CacheFbxTranslatorSettings->bConvertSceneUnit : true;
 	const bool bKeepFbxNamespace = CacheFbxTranslatorSettings ? CacheFbxTranslatorSettings->bKeepFbxNamespace : false;
+	const bool bConsiderClusterBeforePoseForMeshBindPose = CacheFbxTranslatorSettings ? CacheFbxTranslatorSettings->bConsiderClusterBeforePoseForMeshBindPose : false;
 	const bool bUseUfbxParser = CacheFbxTranslatorSettings ? CacheFbxTranslatorSettings->bUseUfbxParser : false;
 
 	if(CacheFbxTranslatorSettings)
@@ -330,7 +331,7 @@ bool UInterchangeFbxTranslator::Translate(UInterchangeBaseNodeContainer& BaseNod
 			return false;
 		}
 		//Create a json command to read the fbx file
-		FString JsonCommand = CreateLoadFbxFileCommand(FilePath, bConvertScene, bForceFrontXAxis, bConvertSceneUnit, bKeepFbxNamespace);
+		FString JsonCommand = CreateLoadFbxFileCommand(FilePath, bConvertScene, bForceFrontXAxis, bConvertSceneUnit, bKeepFbxNamespace, bConsiderClusterBeforePoseForMeshBindPose);
 		int32 TaskIndex = Dispatcher->AddTask(JsonCommand);
 
 		//Blocking call until all tasks are executed
@@ -389,7 +390,7 @@ bool UInterchangeFbxTranslator::Translate(UInterchangeBaseNodeContainer& BaseNod
 #if WITH_EDITOR
 		FbxParser.Reset(bUseUfbxParser);
 		FbxParser.SetResultContainer(Results);
-		FbxParser.SetConvertSettings(bConvertScene, bForceFrontXAxis, bConvertSceneUnit, bKeepFbxNamespace);
+		FbxParser.SetConvertSettings(bConvertScene, bForceFrontXAxis, bConvertSceneUnit, bKeepFbxNamespace, bConsiderClusterBeforePoseForMeshBindPose);
 		FbxParser.LoadFbxFile(FilePath, BaseNodeContainer);
 
 		TArray<UInterchangeResult*> TranslatorResults = Results->GetResults();
@@ -944,10 +945,10 @@ TArray<UE::Interchange::FAnimationPayloadData> UInterchangeFbxTranslator::GetAni
 }
 
 
-FString UInterchangeFbxTranslator::CreateLoadFbxFileCommand(const FString& FbxFilePath, const bool bConvertScene, const bool bForceFrontXAxis, const bool bConvertSceneUnit, const bool bKeepFbxNamespace) const
+FString UInterchangeFbxTranslator::CreateLoadFbxFileCommand(const FString& FbxFilePath, const bool bConvertScene, const bool bForceFrontXAxis, const bool bConvertSceneUnit, const bool bKeepFbxNamespace, const bool bConsiderClusterBeforePoseForMeshBindPose) const
 {
 	using namespace UE::Interchange;
-	FJsonFBXLoadSourceCmd LoadSourceCommand(FJsonFBXLoadSourceCmd::TaskName, FbxFilePath, bConvertScene, bForceFrontXAxis, bConvertSceneUnit, bKeepFbxNamespace);
+	FJsonFBXLoadSourceCmd LoadSourceCommand(FJsonFBXLoadSourceCmd::TaskName, FbxFilePath, bConvertScene, bForceFrontXAxis, bConvertSceneUnit, bKeepFbxNamespace, bConsiderClusterBeforePoseForMeshBindPose);
 	return LoadSourceCommand.ToJson();
 }
 
diff --git a/Engine/Plugins/Interchange/Runtime/Source/Import/Public/Fbx/InterchangeFbxTranslator.h b/Engine/Plugins/Interchange/Runtime/Source/Import/Public/Fbx/InterchangeFbxTranslator.h
--- a/Engine/Plugins/Interchange/Runtime/Source/Import/Public/Fbx/InterchangeFbxTranslator.h
+++ b/Engine/Plugins/Interchange/Runtime/Source/Import/Public/Fbx/InterchangeFbxTranslator.h
@@ -55,6 +55,10 @@ public:
 	UPROPERTY(EditAnywhere, Category = "Fbx Translator")
 	bool bKeepFbxNamespace = false;
 
+	/** Whether to first consider Fbx Clusters instead of FbxPoses for Bind Pose Calculations. */
+	UPROPERTY(EditAnywhere, Category = "Fbx Translator")
+	bool bConsiderClusterBeforePoseForMeshBindPose = false;
+
 	/** Whether to use the uFBX parser. */
 	UPROPERTY(EditAnywhere, Category = "Fbx Translator", meta = (EditCondition = "bDisplayUfbxParser", EditConditionHides, DisplayName="Use ufbx translator (Experimental)", Tooltip = "If enabled, the ufbx SDK is used instead of the FBX SDK. This feature is experimental"))
 	bool bUseUfbxParser = false;
@@ -148,7 +152,7 @@ public:
 	
 	/* IInterchangeAnimationPayloadInterface End */
 private:
-	INTERCHANGEIMPORT_API FString CreateLoadFbxFileCommand(const FString& FbxFilePath, const bool bConvertScene, const bool bForceFrontXAxis, const bool bConvertSceneUnit, const bool bKeepFbxNamespace) const;
+	INTERCHANGEIMPORT_API FString CreateLoadFbxFileCommand(const FString& FbxFilePath, const bool bConvertScene, const bool bForceFrontXAxis, const bool bConvertSceneUnit, const bool bKeepFbxNamespace, const bool bConsiderClusterBeforePoseForMeshBindPose) const;
 
 	INTERCHANGEIMPORT_API FString CreateFetchMeshPayloadFbxCommand(const FString& FbxPayloadKey, const FTransform& MeshGlobalTransform) const;
 
diff --git a/Engine/Plugins/Interchange/Runtime/Source/Parsers/Fbx/Private/FbxAPI.h b/Engine/Plugins/Interchange/Runtime/Source/Parsers/Fbx/Private/FbxAPI.h
--- a/Engine/Plugins/Interchange/Runtime/Source/Parsers/Fbx/Private/FbxAPI.h
+++ b/Engine/Plugins/Interchange/Runtime/Source/Parsers/Fbx/Private/FbxAPI.h
@@ -46,7 +46,7 @@ namespace UE
 				virtual bool IsThreadSafe() { return false; }
 
 				virtual void SetResultContainer(UInterchangeResultsContainer* Result) = 0;
-				virtual void SetConvertSettings(const bool InbConvertScene, const bool InbForceFrontXAxis, const bool InbConvertSceneUnit, const bool InbKeepFbxNamespace) = 0;
+				virtual void SetConvertSettings(const bool InbConvertScene, const bool InbForceFrontXAxis, const bool InbConvertSceneUnit, const bool InbKeepFbxNamespace, const bool InbConsiderClusterBeforePoseForMeshBindPose) = 0;
 
 				/* Load an fbx file into the fbx sdk, return false if the file could not be load. */
 				virtual bool LoadFbxFile(const FString& Filename, UInterchangeBaseNodeContainer& NodeContainer) = 0;
@@ -92,12 +92,13 @@ namespace UE
 					ResultsContainer = Result;
 				}
 
-				virtual void SetConvertSettings(const bool InbConvertScene, const bool InbForceFrontXAxis, const bool InbConvertSceneUnit, const bool InbKeepFbxNamespace) override
+				virtual void SetConvertSettings(const bool InbConvertScene, const bool InbForceFrontXAxis, const bool InbConvertSceneUnit, const bool InbKeepFbxNamespace, const bool InbConsiderClusterBeforePoseForMeshBindPose) override
 				{
 					bConvertScene = InbConvertScene;
 					bForceFrontXAxis = InbForceFrontXAxis;
 					bConvertSceneUnit = InbConvertSceneUnit;
 					bKeepFbxNamespace = InbKeepFbxNamespace;
+					bConsiderClusterBeforePoseForMeshBindPose = InbConsiderClusterBeforePoseForMeshBindPose;
 				}
 
 				//return the fbx helper for this parser
@@ -151,6 +152,8 @@ namespace UE
 
 				const FString& GetSourceFilename() const { return SourceFilename; }
 
+				bool GetUseFbxClustersForBindPoseMatrices() const { return bConsiderClusterBeforePoseForMeshBindPose; }
+
 				FbxAMatrix JointOrientationMatrix;
 
 				/**
@@ -184,6 +187,7 @@ namespace UE
 				bool bConvertSceneUnit = true;
 				bool bKeepFbxNamespace = false;
 				bool bCreatorIsBlender = false;
+				bool bConsiderClusterBeforePoseForMeshBindPose = false;
 
 				struct FileDetails
 				{
diff --git a/Engine/Plugins/Interchange/Runtime/Source/Parsers/Fbx/Private/FbxMesh.cpp b/Engine/Plugins/Interchange/Runtime/Source/Parsers/Fbx/Private/FbxMesh.cpp
--- a/Engine/Plugins/Interchange/Runtime/Source/Parsers/Fbx/Private/FbxMesh.cpp
+++ b/Engine/Plugins/Interchange/Runtime/Source/Parsers/Fbx/Private/FbxMesh.cpp
@@ -2245,10 +2245,15 @@ UInterchangeMeshNode* FFbxMesh::CreateMeshNode(UInterchangeBaseNodeContainer& No
 	return MeshNode;
 }
 
-FFbxJointMeshBindPoseGenerator::FFbxJointMeshBindPoseGenerator(FbxScene* SDKScene, FFbxParser& Parser)
+FFbxJointMeshBindPoseGenerator::FFbxJointMeshBindPoseGenerator(FbxScene* SDKScene, FFbxParser& Parser, const bool bConsiderClusterBeforePoseForMeshBindPose)
 {
 	//Search all skeletalmesh(FbxGeometry with valid deformer) using this joint and see if there is a valid FbxCluster
 	//containing a TransformLinkMatrix for this joint
+
+	TSet<const FbxNode*> ValidationTracker; //Gets reset everytime when a BindPose gets recreated.
+	TMap<const FbxNode*, FString> JointToGeometryMap; //Maps Joint to the Geometry that the Bind Pose is being used from.
+	bool bClusterVsPoseMissmatch = false;
+
 	const int32 GeometryCount = SDKScene->GetGeometryCount();
 	for (int32 GeometryIndex = 0; GeometryIndex < GeometryCount; ++GeometryIndex)
 	{
@@ -2257,12 +2262,14 @@ FFbxJointMeshBindPoseGenerator::FFbxJointMeshBindPoseGenerator(FbxScene* SDKScen
 		{
 			continue;
 		}
+		FString GeometryName = Parser.GetFbxHelper()->GetFbxObjectName(Geometry);
 		FString MeshUniqueID = Parser.GetFbxHelper()->GetMeshUniqueID(Geometry);
 
 		const int32 GeometryDeformerCount = Geometry->GetDeformerCount(FbxDeformer::eSkin);
 		for (int32 GeometryDeformerIndex = 0; GeometryDeformerIndex < GeometryDeformerCount; ++GeometryDeformerIndex)
 		{
 			FbxSkin* Skin = (FbxSkin*)Geometry->GetDeformer(GeometryDeformerIndex, FbxDeformer::eSkin);
+
 			if (!ensure(Skin))
 			{
 				continue;
@@ -2281,15 +2288,25 @@ FFbxJointMeshBindPoseGenerator::FFbxJointMeshBindPoseGenerator(FbxScene* SDKScen
 
 				FbxNode* Joint = Cluster->GetLink();
 
-				uint64 ClusterId = Cluster->GetUniqueID();
-				uint64 MeshId = Geometry->GetUniqueID();
+				if (!Joint)
+				{
+					continue;
+				}
+
+				//Acquire Cluster transforms:
+				FbxAMatrix BoneBindMatrix_Cluster;
+				FbxAMatrix MeshBindMatrix_Cluster;
+				Cluster->GetTransformMatrix(MeshBindMatrix_Cluster);
+				Cluster->GetTransformLinkMatrix(BoneBindMatrix_Cluster);
 
+				//Acquire BindPose transforms, reasoning:
 				// Here we try fetching the bind pose only from actual FbxPose objects. Interchange traditionally preferred reading the
 				// TransformLink directly from the cluster, but it seems that in some edge cases (see UE-273136) that is not the correct thing to do?
 				// By checking the pose first, we mirror the FbxLegacy behavior shown at UnFbx::FFbxImporter::ImportBones, with how it first 
 				// checks the PoseArray when filling in GlobalsPerLink
-				FbxAMatrix GlobalBindPoseJointMatrixPerGeo;
-				bool bRetrieveBindPoseResult = RetrievePoseFromBindPose(Joint, GlobalBindPoseJointMatrixPerGeo);
+				FbxAMatrix BoneBindMatrix_Pose;
+				FbxAMatrix MeshBindMatrix_Pose;
+				bool bRetrieveBindPoseResult = RetrievePoseFromBindPose(Joint, BoneBindMatrix_Pose);
 				if (!bRetrieveBindPoseResult)
 				{
 					// if failed, delete bind pose, and retry.
@@ -2304,41 +2321,114 @@ FFbxJointMeshBindPoseGenerator::FFbxJointMeshBindPoseGenerator(FbxScene* SDKScen
 						{
 							SDKScene->RemovePose(PoseIndex);
 							CurrentPose->Destroy();
+
+							ValidationTracker.Reset();
 						}
 					}
 
 					SDKScene->GetFbxManager()->CreateMissingBindPoses(SDKScene);
-					bRetrieveBindPoseResult = RetrievePoseFromBindPose(Joint, GlobalBindPoseJointMatrixPerGeo);
+					bRetrieveBindPoseResult = RetrievePoseFromBindPose(Joint, BoneBindMatrix_Pose);
+
+					if (!bRetrieveBindPoseResult)
+					{
+						Cluster->GetTransformLinkMatrix(BoneBindMatrix_Pose);
+					}
+				}
+				if (!RetrievePoseFromBindPose(Geometry->GetNode(), MeshBindMatrix_Pose))
+				{
+					Cluster->GetTransformMatrix(MeshBindMatrix_Pose);
 				}
 
-				if (!bRetrieveBindPoseResult)
+				if (bRetrieveBindPoseResult && !bClusterVsPoseMissmatch)
 				{
-					Cluster->GetTransformLinkMatrix(GlobalBindPoseJointMatrixPerGeo);
+					// Convert to FMatrix to check equality with tolerance. 
+					const FMatrix ClusterMatrix = FFbxConvert::ConvertMatrix<FMatrix>(BoneBindMatrix_Cluster);
+					const FMatrix PoseMatrix = FFbxConvert::ConvertMatrix<FMatrix>(BoneBindMatrix_Pose);
+
+					if (!ClusterMatrix.Equals(PoseMatrix, KINDA_SMALL_NUMBER))
+					{
+						bClusterVsPoseMissmatch = true;
+					}
+				}
+
+				FbxAMatrix MeshBindMatrix;
+				FbxAMatrix BoneBindMatrix;
+				if (bConsiderClusterBeforePoseForMeshBindPose)
+				{
+					MeshBindMatrix = MeshBindMatrix_Cluster;
+					BoneBindMatrix = BoneBindMatrix_Cluster;
+				}
+				else
+				{
+					MeshBindMatrix = MeshBindMatrix_Pose;
+					BoneBindMatrix = BoneBindMatrix_Pose;
 				}
 
 				TMap<FString, FMatrix>& MeshIdToGlobalBindPoseJointMap = JointToMeshIdToGlobalBindPoseJointMap.FindOrAdd(Joint);
-				MeshIdToGlobalBindPoseJointMap.Add(MeshUniqueID, FFbxConvert::ConvertMatrix<FMatrix>(GlobalBindPoseJointMatrixPerGeo));
+				MeshIdToGlobalBindPoseJointMap.Add(MeshUniqueID, FFbxConvert::ConvertMatrix<FMatrix>(BoneBindMatrix));
 
+				const FbxAMatrix BindPoseFbxAMatrix = BoneBindMatrix * Parser.JointOrientationMatrix;
 				if (!JointToBindPoseMap.Contains(Joint))
 				{
-					JointToBindPoseMap.Add(Joint, GlobalBindPoseJointMatrixPerGeo * Parser.JointOrientationMatrix);
+					JointToBindPoseMap.Add(Joint, BindPoseFbxAMatrix);
+					ValidationTracker.Add(Joint);
+					JointToGeometryMap.Add(Joint, GeometryName);
 				}
 				else
 				{
-					// Convert to FMatrix to check equality with tolerance. 
-					const FMatrix ExistingMatrix = FFbxConvert::ConvertMatrix<FMatrix>(JointToBindPoseMap[Joint]);
-					const FMatrix NewMatrix = FFbxConvert::ConvertMatrix<FMatrix>(GlobalBindPoseJointMatrixPerGeo * Parser.JointOrientationMatrix);
-					ensureMsgf(NewMatrix.Equals(ExistingMatrix, KINDA_SMALL_NUMBER), TEXT("Expected to have same joint node in different clusters to have same Bind Pose matrix!"));
-				}
+					if (ValidationTracker.Contains(Joint))
+					{
+						// Convert to FMatrix to check equality with tolerance. 
+						const FMatrix ExistingMatrix = FFbxConvert::ConvertMatrix<FMatrix>(JointToBindPoseMap[Joint]);
+						const FMatrix NewMatrix = FFbxConvert::ConvertMatrix<FMatrix>(BindPoseFbxAMatrix);
 
-				FbxAMatrix GlobalBindPoseReferenceMatrix;
-				Cluster->GetTransformMatrix(GlobalBindPoseReferenceMatrix);
+						if (!NewMatrix.Equals(ExistingMatrix, KINDA_SMALL_NUMBER))
+						{
+							FText JointName = FText::FromString(Parser.GetFbxHelper()->GetFbxObjectName(Joint));
+							bool bWarningAdded = false;
+							if (JointToGeometryMap.Contains(Joint))
+							{
+								if (GeometryName != JointToGeometryMap[Joint])
+								{
+									//Conflicting bind Poses.
+									//#interchange_todo: Perhaps in the future we could introduce a Pose selector for Bind Pose calculations. Currently out of scope.
+
+									UInterchangeResultWarning_Generic* Message = Parser.AddMessage<UInterchangeResultWarning_Generic>();
+									Message->Text = FText::Format(LOCTEXT("FFbxJointMeshBindPoseGenerator_ConflictingBindPoseMatrices", 
+										"Imported SkeletalMesh and/or Skeleton might be incorrect.\nConflicting Bind Pose Matrices found for Joint [{0}].\n   Currently using Geometry [{1}] for Bind Pose Matrix [{2}].\n   However Geometry [{3}] stores a differing Bind Pose Matrix [{4}]."), 
+										JointName, 
+										FText::FromString(JointToGeometryMap[Joint]), FText::FromString(ExistingMatrix.ToString()),
+										FText::FromString(GeometryName), FText::FromString(NewMatrix.ToString()));
+
+									bWarningAdded = true;
+								}
+							}
+							
+							if (!bWarningAdded)
+							{
+								UInterchangeResultWarning_Generic* Message = Parser.AddMessage<UInterchangeResultWarning_Generic>();
+								Message->Text = FText::Format(LOCTEXT("FFbxJointMeshBindPoseGenerator_UnexpectedBindPoseMatrices", "Imported SkeletalMesh and/or Skeleton might be incorrect.\nExpected to have same joint [{0}] node in different clusters to have same Bind Pose matrix!"), FText::FromString(Parser.GetFbxHelper()->GetFbxObjectName(Joint)));
+							}
+						}
+					}
+					else
+					{
+						JointToBindPoseMap[Joint] = BindPoseFbxAMatrix;
+						JointToGeometryMap.Add(Joint, GeometryName);
+					}
+				}
 
 				TMap<FString, FMatrix>& MeshIdToGlobalBindPoseReferenceMap = JointToMeshIdToGlobalBindPoseReferenceMap.FindOrAdd(Joint);
-				MeshIdToGlobalBindPoseReferenceMap.Add(MeshUniqueID, FFbxConvert::ConvertMatrix<FMatrix>(GlobalBindPoseReferenceMatrix));
+				MeshIdToGlobalBindPoseReferenceMap.Add(MeshUniqueID, FFbxConvert::ConvertMatrix<FMatrix>(MeshBindMatrix));
 			}
 		}
 	}
+
+	if (bClusterVsPoseMissmatch)
+	{
+		UInterchangeResultWarning_Generic* Message = Parser.AddMessage<UInterchangeResultWarning_Generic>();
+		Message->Text = LOCTEXT("FFbxJointMeshBindPoseGenerator_NonMatchingBindPoses", "BindPose Matrix generation acquired 2 different Matrices from FbxCluster vs FbxPose.\n If import is incorrect, consider using a different setting for FBX option 'bConsiderClusterBeforePoseForMeshBindPose'.");
+	}
 }
 
 TMap<FString, FMatrix>& FFbxJointMeshBindPoseGenerator::GetMeshIdToGlobalBindPoseReferenceMap(const FbxNode* Joint)
diff --git a/Engine/Plugins/Interchange/Runtime/Source/Parsers/Fbx/Private/FbxMesh.h b/Engine/Plugins/Interchange/Runtime/Source/Parsers/Fbx/Private/FbxMesh.h
--- a/Engine/Plugins/Interchange/Runtime/Source/Parsers/Fbx/Private/FbxMesh.h
+++ b/Engine/Plugins/Interchange/Runtime/Source/Parsers/Fbx/Private/FbxMesh.h
@@ -29,7 +29,7 @@ namespace UE
 				TMap<const FbxNode*, TMap<FString, FMatrix>> JointToMeshIdToGlobalBindPoseJointMap;
 				TMap<const FbxNode*, FbxAMatrix> JointToBindPoseMap;
 
-				FFbxJointMeshBindPoseGenerator(FbxScene* SDKScene, FFbxParser& Parser);
+				FFbxJointMeshBindPoseGenerator(FbxScene* SDKScene, FFbxParser& Parser, const bool bConsiderClusterBeforePoseForMeshBindPose);
 
 				TMap<FString, FMatrix>& GetMeshIdToGlobalBindPoseReferenceMap(const FbxNode* Joint);
 				TMap<FString, FMatrix>& GetMeshIdToGlobalBindPoseJointMap(const FbxNode* Joint);
diff --git a/Engine/Plugins/Interchange/Runtime/Source/Parsers/Fbx/Private/FbxScene.cpp b/Engine/Plugins/Interchange/Runtime/Source/Parsers/Fbx/Private/FbxScene.cpp
--- a/Engine/Plugins/Interchange/Runtime/Source/Parsers/Fbx/Private/FbxScene.cpp
+++ b/Engine/Plugins/Interchange/Runtime/Source/Parsers/Fbx/Private/FbxScene.cpp
@@ -661,7 +661,7 @@ namespace UE
 					RootJointInfo.Value.bValidBindPose = IsValidBindPose(SDKScene, RootJointInfo.Key);
 				}
 
-				FFbxJointMeshBindPoseGenerator FbxJointMeshBindPoseGenerator(SDKScene, Parser);
+				FFbxJointMeshBindPoseGenerator FbxJointMeshBindPoseGenerator(SDKScene, Parser, Parser.GetUseFbxClustersForBindPoseMatrices());
 
 				bool bBadBindPoseMessageDisplay = false;
 				AddHierarchyRecursively(nullptr, RootNode, SDKScene, NodeContainer, PayloadContexts, ForceJointNodes, bBadBindPoseMessageDisplay, FbxJointMeshBindPoseGenerator);
@@ -956,7 +956,7 @@ namespace UE
 
 					int32 AttributeCount = Node->GetNodeAttributeCount();
 
-					bool bNewSkeltalAnimationStarted = false;
+					bool bNewSkeletalAnimationStarted = false;
 
 					for (int32 AttributeIndex = 0; AttributeIndex < AttributeCount && !HasSkeletonAttribute; ++AttributeIndex)
 					{
@@ -971,7 +971,7 @@ namespace UE
 							}
 						case FbxNodeAttribute::eSkeleton:
 							bIsNodeContainJointAttribute = true;
-							bNewSkeltalAnimationStarted = ApplySkeletonAttribute() || bNewSkeltalAnimationStarted;
+							bNewSkeletalAnimationStarted = ApplySkeletonAttribute() || bNewSkeletalAnimationStarted;
 							break;
 						default:
 							break;
@@ -983,15 +983,15 @@ namespace UE
 						//Make sure to treat the node like a joint if it's in the ForcejointNodes array
 						if (ForceJointNodes.Contains(Node))
 						{
-							bNewSkeltalAnimationStarted = ApplySkeletonAttribute() || bNewSkeltalAnimationStarted;
+							bNewSkeletalAnimationStarted = ApplySkeletonAttribute() || bNewSkeletalAnimationStarted;
 						}
 						else if (!bIsRootNode && IsNodeUnderCommonJointRootNode(Node, CommonJointRootNodes))
 						{
-							bNewSkeltalAnimationStarted = ApplySkeletonAttribute() || bNewSkeltalAnimationStarted;
+							bNewSkeletalAnimationStarted = ApplySkeletonAttribute() || bNewSkeletalAnimationStarted;
 						}
 					}
 
-					if (bNewSkeltalAnimationStarted)
+					if (bNewSkeletalAnimationStarted)
 					{
 						ProcessCustomAttributes(Parser, Node, SkeletalAnimationTrackNode);
 					}
@@ -1007,7 +1007,7 @@ namespace UE
 					else if (SkeletalAnimationTrackNode)
 					{
 						//Scene node transform can be animated, add the transform animation payload key.
-						if (FFbxAnimation::AddSkeletalTransformAnimation(NodeContainer, SDKScene, Parser, Node, UnrealNode, PayloadContexts, SkeletalAnimationTrackNode, AnimationIndex, bNewSkeltalAnimationStarted)
+						if (FFbxAnimation::AddSkeletalTransformAnimation(NodeContainer, SDKScene, Parser, Node, UnrealNode, PayloadContexts, SkeletalAnimationTrackNode, AnimationIndex, bNewSkeletalAnimationStarted)
 							&& !SkeletalAnimationAddedToContainer)
 						{
 							SkeletalAnimationAddedToContainer = true;
diff --git a/Engine/Plugins/Interchange/Runtime/Source/Parsers/Fbx/Private/InterchangeFbxParser.cpp b/Engine/Plugins/Interchange/Runtime/Source/Parsers/Fbx/Private/InterchangeFbxParser.cpp
--- a/Engine/Plugins/Interchange/Runtime/Source/Parsers/Fbx/Private/InterchangeFbxParser.cpp
+++ b/Engine/Plugins/Interchange/Runtime/Source/Parsers/Fbx/Private/InterchangeFbxParser.cpp
@@ -70,9 +70,9 @@ namespace UE::Interchange
 		FbxParserPrivate->SetResultContainer(Result);
 	}
 
-	void FInterchangeFbxParser::SetConvertSettings(const bool InbConvertScene, const bool InbForceFrontXAxis, const bool InbConvertSceneUnit, const bool InbKeepFbxNamespace)
+	void FInterchangeFbxParser::SetConvertSettings(const bool InbConvertScene, const bool InbForceFrontXAxis, const bool InbConvertSceneUnit, const bool InbKeepFbxNamespace, const bool InbConsiderClusterBeforePoseForMeshBindPose)
 	{
-		GetFbxParser().SetConvertSettings(InbConvertScene, InbForceFrontXAxis, InbConvertSceneUnit, InbKeepFbxNamespace);
+		GetFbxParser().SetConvertSettings(InbConvertScene, InbForceFrontXAxis, InbConvertSceneUnit, InbKeepFbxNamespace, InbConsiderClusterBeforePoseForMeshBindPose);
 	}
 
 	void FInterchangeFbxParser::LoadFbxFile(const FString& Filename, const FString& ResultFolder)
diff --git a/Engine/Plugins/Interchange/Runtime/Source/Parsers/Fbx/Private/Ufbx/UfbxParser.h b/Engine/Plugins/Interchange/Runtime/Source/Parsers/Fbx/Private/Ufbx/UfbxParser.h
--- a/Engine/Plugins/Interchange/Runtime/Source/Parsers/Fbx/Private/Ufbx/UfbxParser.h
+++ b/Engine/Plugins/Interchange/Runtime/Source/Parsers/Fbx/Private/Ufbx/UfbxParser.h
@@ -94,7 +94,7 @@ namespace UE::Interchange::Private
 		virtual bool IsThreadSafe() override { return true; }
 
 		virtual void SetResultContainer(UInterchangeResultsContainer* Result) override;
-		virtual void SetConvertSettings(const bool InbConvertScene, const bool InbForceFrontXAxis, const bool InbConvertSceneUnit, const bool InbKeepFbxNamespace) override
+		virtual void SetConvertSettings(const bool InbConvertScene, const bool InbForceFrontXAxis, const bool InbConvertSceneUnit, const bool InbKeepFbxNamespace, const bool InbConsiderClusterBeforePoseForMeshBindPose)  override
 		{
 			bKeepFbxNamespace = InbKeepFbxNamespace;
 		}
diff --git a/Engine/Plugins/Interchange/Runtime/Source/Parsers/Fbx/Public/InterchangeFbxParser.h b/Engine/Plugins/Interchange/Runtime/Source/Parsers/Fbx/Public/InterchangeFbxParser.h
--- a/Engine/Plugins/Interchange/Runtime/Source/Parsers/Fbx/Public/InterchangeFbxParser.h
+++ b/Engine/Plugins/Interchange/Runtime/Source/Parsers/Fbx/Public/InterchangeFbxParser.h
@@ -40,7 +40,7 @@ namespace UE
 
 			INTERCHANGEFBXPARSER_API void SetResultContainer(UInterchangeResultsContainer* Result);
 
-			INTERCHANGEFBXPARSER_API void SetConvertSettings(const bool InbConvertScene, const bool InbForceFrontXAxis, const bool InbConvertSceneUnit, const bool InbKeepFbxNamespace);
+			INTERCHANGEFBXPARSER_API void SetConvertSettings(const bool InbConvertScene, const bool InbForceFrontXAxis, const bool InbConvertSceneUnit, const bool InbKeepFbxNamespace, const bool InbConsiderClusterBeforePoseForMeshBindPose);
 			/**
 			 * Parse a file support by the fbx sdk. It just extract all the fbx node and create a FBaseNodeContainer and dump it in a json file inside the ResultFolder
 			 * @param - Filename is the file that the fbx sdk will read (.fbx or .obj)
diff --git a/Engine/Source/Programs/InterchangeWorker/Private/InterchangeWorkerImpl.cpp b/Engine/Source/Programs/InterchangeWorker/Private/InterchangeWorkerImpl.cpp
--- a/Engine/Source/Programs/InterchangeWorker/Private/InterchangeWorkerImpl.cpp
+++ b/Engine/Source/Programs/InterchangeWorker/Private/InterchangeWorkerImpl.cpp
@@ -298,7 +298,7 @@ ETaskState FInterchangeWorkerImpl::LoadFbxFile(const FJsonFBXLoadSourceCmd& Load
 	FString SourceFilename = LoadSourceCommand.GetSourceFilename();
 	UE_LOG(LogInterchangeWorker, Verbose, TEXT("LoadFbxFile : %s"), *SourceFilename);
 	FbxParser.Reset();
-	FbxParser.SetConvertSettings(LoadSourceCommand.GetDoesConvertScene(), LoadSourceCommand.GetDoesForceFrontXAxis(), LoadSourceCommand.GetDoesConvertSceneUnit(), LoadSourceCommand.GetDoesKeepFbxNamespace());
+	FbxParser.SetConvertSettings(LoadSourceCommand.GetDoesConvertScene(), LoadSourceCommand.GetDoesForceFrontXAxis(), LoadSourceCommand.GetDoesConvertSceneUnit(), LoadSourceCommand.GetDoesKeepFbxNamespace(), LoadSourceCommand.GetDoesUseFbxClustersForBindPoseMatrices());
 	FbxParser.LoadFbxFile(SourceFilename, ResultFolder);
 	FJsonLoadSourceCmd::JsonResultParser ResultParser;
 	ResultParser.SetResultFilename(FbxParser.GetResultFilepath());
