[Interchange] ufbx - Exported file/extra info for imported file

#Synced-CL 48545231

diff --git a/Engine/Plugins/Interchange/Runtime/Source/Parsers/Fbx/Private/Ufbx/UfbxParser.cpp b/Engine/Plugins/Interchange/Runtime/Source/Parsers/Fbx/Private/Ufbx/UfbxParser.cpp
--- a/Engine/Plugins/Interchange/Runtime/Source/Parsers/Fbx/Private/Ufbx/UfbxParser.cpp
+++ b/Engine/Plugins/Interchange/Runtime/Source/Parsers/Fbx/Private/Ufbx/UfbxParser.cpp
@@ -117,9 +117,59 @@ bool FUfbxParser::LoadFbxFile(const FString& Filename, UInterchangeBaseNodeConta
 	// #ufbx_todo: read scene converter to UE's
 
 	UInterchangeSourceNode* SourceNode = UInterchangeSourceNode::FindOrCreateUniqueInstance(&NodeContainer);
-	// #ufbx_todo:
-	// SourceNode->SetCustomSourceFrameRateNumerator(FrameRate);
-	// SourceNode->SetCustomSourceFrameRateDenominator(Denominator);
+	
+	SourceNode->SetCustomSourceFrameRateNumerator(Scene->settings.frames_per_second);
+	SourceNode->SetCustomSourceFrameRateDenominator(1.0);
+
+	
+	SourceNode->SetExtraInformation(TEXT("File Version"), Convert::ToUnrealString(Scene->metadata.original_application.version));
+	SourceNode->SetExtraInformation(TEXT("File Creator"), Convert::ToUnrealString(Scene->metadata.creator));
+
+	// FBX SDK extrapolates unit name(e.g. "centimeters") from numerical scaling value that is contained in the fbx
+	SourceNode->SetExtraInformation(TEXT("File Units"),  FString::Printf(TEXT("UniteMeters: %.3f"), Scene->settings.original_unit_meters));
+
+	FString AxisDirection;
+
+	// Negative axes are odd
+	int32 Negative = Scene->settings.original_axis_up % 2;
+	switch (Scene->settings.original_axis_up / 2)
+	{
+	case UFBX_COORDINATE_AXIS_POSITIVE_X:
+		AxisDirection += TEXT("X");
+		break;
+	case UFBX_COORDINATE_AXIS_POSITIVE_Y:
+		AxisDirection += TEXT("Y");
+		break;
+	case UFBX_COORDINATE_AXIS_POSITIVE_Z:
+		AxisDirection += TEXT("Z");
+		break;
+	case UFBX_COORDINATE_AXIS_UNKNOWN:
+		break;
+	default: ;
+	}
+
+	AxisDirection += Negative ? TEXT("-DOWN") : TEXT("-UP");
+
+	SourceNode->SetExtraInformation(TEXT("File Axis Direction"), AxisDirection);
+	SourceNode->SetExtraInformation(TEXT("File Frame Rate"), FString::Printf(TEXT("%.2f"), Scene->settings.frames_per_second));
+
+
+	FString ApplicationVendor = Convert::ToUnrealString(Scene->metadata.latest_application.vendor);
+	if (!ApplicationVendor.IsEmpty())
+	{
+		SourceNode->SetExtraInformation(FSourceNodeExtraInfoStaticData::GetApplicationVendorExtraInfoKey(), ApplicationVendor);
+	}
+	FString ApplicationName = Convert::ToUnrealString(Scene->metadata.latest_application.name);
+	if (!ApplicationName.IsEmpty())
+	{
+		SourceNode->SetExtraInformation(FSourceNodeExtraInfoStaticData::GetApplicationNameExtraInfoKey(), ApplicationName);
+	}
+	FString ApplicationVersion = Convert::ToUnrealString(Scene->metadata.latest_application.version);
+	if (!ApplicationVersion.IsEmpty())
+	{
+		SourceNode->SetExtraInformation(FSourceNodeExtraInfoStaticData::GetApplicationVersionExtraInfoKey(), ApplicationVersion);
+	}
+
 
 	// Fbx legacy has a special way to bake the skeletal mesh that do not fit the interchange standard
 	// The interchange skeletal mesh factory will read this to use the proper bake transform so it match legacy behavior.
