Removing FBX specific Transform specialization.
Also clean-up FbxScene to improve readability.

This is in preparation for UE-354035

Some notes about Transform specialization:

Joint Specialization = actual Bone coming from the DCC.
Transform Specialization = a node that is not a Bone but has a Bone ancestor.

- Transform are only used with eNull FBX.
- A specialized Transform is always a specialized Joint.
- Other format like GLB can import same hierarchies that include empty/null nodes without the usage of Transform specialization.

-----------

About the modified "TP_HierarC_Level_Actor" test:
- The tests does a level import of `AllinOne4.fbx`
- The only modification in the test is the number of expected actors (from 13 to 12)
- The extra actor seems to be a bug in previous code since it's an empty actor outside of the RootNode.
- That actor is not present in Legacy either

-----------

While testing for this CL, I noticed some inconsistencies in how roots are considered either inside or outside the skeleton. The inconsistency is with both the content in FBX/Blender and what Legacy used to import. To Be discussed outside this CL.

[RN] Removed FSceneNodeStaticData::GetTransformSpecializeTypeString()
#tests All Interchange tests are passing except the one that was fixed in this CL.
#virtualized
#Synced-CL 48518362

diff --git a/Engine/Plugins/Interchange/Runtime/Source/Nodes/Private/InterchangeSceneNode.cpp b/Engine/Plugins/Interchange/Runtime/Source/Nodes/Private/InterchangeSceneNode.cpp
--- a/Engine/Plugins/Interchange/Runtime/Source/Nodes/Private/InterchangeSceneNode.cpp
+++ b/Engine/Plugins/Interchange/Runtime/Source/Nodes/Private/InterchangeSceneNode.cpp
@@ -22,12 +22,6 @@ namespace UE
 			static FAttributeKey MaterialDependencyUids_BaseKey(TEXT("__MaterialDependencyUidsBaseKey__"));
 			return MaterialDependencyUids_BaseKey;
 		}
-		
-		const FString& FSceneNodeStaticData::GetTransformSpecializeTypeString()
-		{
-			static FString TransformSpecializeTypeString(TEXT("Transform"));
-			return TransformSpecializeTypeString;
-		}
 
 		const FString& FSceneNodeStaticData::GetJointSpecializeTypeString()
 		{
diff --git a/Engine/Plugins/Interchange/Runtime/Source/Nodes/Public/InterchangeSceneNode.h b/Engine/Plugins/Interchange/Runtime/Source/Nodes/Public/InterchangeSceneNode.h
--- a/Engine/Plugins/Interchange/Runtime/Source/Nodes/Public/InterchangeSceneNode.h
+++ b/Engine/Plugins/Interchange/Runtime/Source/Nodes/Public/InterchangeSceneNode.h
@@ -17,12 +17,10 @@ namespace UE
 {
 	namespace Interchange
 	{
-
 		struct FSceneNodeStaticData : public FBaseNodeStaticData
 		{
 			static UE_API const FAttributeKey& GetNodeSpecializeTypeBaseKey();
 			static UE_API const FAttributeKey& GetMaterialDependencyUidsBaseKey();
-			static UE_API const FString& GetTransformSpecializeTypeString();
 			static UE_API const FString& GetJointSpecializeTypeString();
 			static UE_API const FString& GetLodGroupSpecializeTypeString();
 			static UE_API const FString& GetMeshToGlobalBindPoseReferencesString();
diff --git a/Engine/Plugins/Interchange/Runtime/Source/Parsers/Fbx/Private/FbxScene.cpp b/Engine/Plugins/Interchange/Runtime/Source/Parsers/Fbx/Private/FbxScene.cpp
--- a/Engine/Plugins/Interchange/Runtime/Source/Parsers/Fbx/Private/FbxScene.cpp
+++ b/Engine/Plugins/Interchange/Runtime/Source/Parsers/Fbx/Private/FbxScene.cpp
@@ -140,7 +140,7 @@ namespace UE
 				, FbxScene* SDKScene
 				, UInterchangeBaseNodeContainer& NodeContainer
 				, TMap<FString, TSharedPtr<FPayloadContextBase, ESPMode::ThreadSafe>>& PayloadContexts
-				, TArray<FbxNode*>& ForceJointNodes
+				, const TArray<FbxNode*>& ForceJointNodes
 				, bool& bBadBindPoseMessageDisplay
 				, FFbxJointMeshBindPoseGenerator& FbxJointMeshBindPoseGenerator)
 			{
@@ -330,7 +330,8 @@ namespace UE
 							//Unsupported attribute
 							break;
 
-						case FbxNodeAttribute::eShape: //We do not add a dependency for shape on the scene node since shapes are a MeshNode dependency.
+						case FbxNodeAttribute::eShape:
+							//We do not add a dependency for shape on the scene node since shapes are a MeshNode dependency.
 							break;
 
 						case FbxNodeAttribute::eNull:
@@ -340,7 +341,6 @@ namespace UE
 								//eNull node not in a hierarchy containing any joint will not be set has joint
 								break;
 							}
-							UnrealNode->AddSpecializedType(FSceneNodeStaticData::GetTransformSpecializeTypeString());
 						}
 						//No break since the eNull act has a skeleton if possible
 						case FbxNodeAttribute::eSkeleton:
@@ -405,15 +405,9 @@ namespace UE
 				if (!bIsNodeContainJointAttribute)
 				{
 					//Make sure to treat the node like a joint if it's in the ForcejointNodes array
-					if (ForceJointNodes.Contains(Node))
+					//Or if it's a regular scene node with a joint ancestor
+					if (ForceJointNodes.Contains(Node) || (!bIsRootNode && IsNodeUnderCommonJointRootNode(Node, CommonJointRootNodes)))
 					{
-						UnrealNode->AddSpecializedType(FSceneNodeStaticData::GetTransformSpecializeTypeString());
-						ApplySkeletonAttribute(false);
-					}
-					else if (!bIsRootNode && IsNodeUnderCommonJointRootNode(Node, CommonJointRootNodes))
-					{
-						FString Name = Node->GetName();
-						UnrealNode->AddSpecializedType(FSceneNodeStaticData::GetTransformSpecializeTypeString());
 						ApplySkeletonAttribute(false);
 					}
 					else
@@ -439,7 +433,6 @@ namespace UE
 						}
 					}
 				}
-
 				
 				auto AddAnimationTrackNode = [&](EInterchangePropertyTracks PropertyTrack, const FString& CurveNodeName, const FString& PayloadKey, EInterchangeAnimationPayLoadType PayloadType)
 				{
@@ -551,13 +544,8 @@ namespace UE
 				return TransformNode;
 			}
 			
-			bool FFbxScene::IsValidBindPose(FbxScene* SDKScene, FbxNode* RootJoint) const
+			bool FFbxScene::IsValidBindPose(FbxScene* SDKScene, FbxNode* RootJoint)
 			{
-				if (CommonJointRootNodes.IsEmpty())
-				{
-					return false;
-				}
-
 				int32 PoseCount = SDKScene->GetPoseCount();
 				if (PoseCount == 0)
 				{
@@ -650,16 +638,10 @@ namespace UE
 
 				//Some fbx file have node without attribute that are link in cluster,
 				//We must consider those node has joint
-				TArray<FbxNode*> ForceJointNodes;
-				FindForceJointNode(SDKScene, ForceJointNodes);
+				TArray<FbxNode*> ForceJointNodes = FindForceJointNode(SDKScene);
 
 				//Cache the common root joint
-				FindCommonJointRootNode(SDKScene, ForceJointNodes);
-
-				for (TPair<FbxNode*, FRootJointInfo>& RootJointInfo : CommonJointRootNodes)
-				{
-					RootJointInfo.Value.bValidBindPose = IsValidBindPose(SDKScene, RootJointInfo.Key);
-				}
+				FillCommonJointRootNode(SDKScene, ForceJointNodes);
 
 				FFbxJointMeshBindPoseGenerator FbxJointMeshBindPoseGenerator(SDKScene, Parser, Parser.GetUseFbxClustersForBindPoseMatrices());
 
@@ -683,7 +665,7 @@ namespace UE
 			}
 
 			void FFbxScene::AddRigidAnimation(FbxNode* Node
-				, UInterchangeSceneNode* UnrealNode
+				, const UInterchangeSceneNode* UnrealNode
 				, UInterchangeBaseNodeContainer& NodeContainer
 				, TMap<FString, TSharedPtr<FPayloadContextBase, ESPMode::ThreadSafe>>& PayloadContexts)
 			{
@@ -765,7 +747,7 @@ namespace UE
 				}
 			}
 
-			FbxNode* FFbxScene::Internal_GetRootSkeleton(FbxScene* SDKScene, FbxNode* Link)
+			FbxNode* FFbxScene::Internal_GetRootSkeleton(FbxScene* SDKScene, FbxNode* Link, bool bIsCreatorBlender)
 			{
 				FbxNode* RootBone = Link;
 
@@ -774,7 +756,7 @@ namespace UE
 				while (RootBone && RootBone->GetParent())
 				{
 					bool bIsBlenderArmatureBone = false;
-					if (Parser.IsCreatorBlender())
+					if (bIsCreatorBlender)
 					{
 						//Hack to support armature dummy node from blender
 						//Users do not want the null attribute node named armature which is the parent of the real root bone in blender fbx file
@@ -813,7 +795,7 @@ namespace UE
 				return RootBone;
 			}
 
-			void FFbxScene::FindCommonJointRootNode(FbxScene* SDKScene, const TArray<FbxNode*>& ForceJointNodes)
+			void FFbxScene::FillCommonJointRootNode(FbxScene* SDKScene, const TArray<FbxNode*>& ForceJointNodes)
 			{
 				//Process the ForceJointNodes and any skeleton joint node
 				int32 NodeCount = SDKScene->GetNodeCount();
@@ -836,17 +818,22 @@ namespace UE
 						}
 						if (bProcessNode)
 						{
-							if (FbxNode* Root = Internal_GetRootSkeleton(SDKScene, Node))
+							if (FbxNode* Root = Internal_GetRootSkeleton(SDKScene, Node, Parser.IsCreatorBlender()))
 							{
-								CommonJointRootNodes.FindOrAdd(Root);
+								if (!CommonJointRootNodes.Contains(Root))
+								{
+									CommonJointRootNodes.Add(Root, FFbxScene::FRootJointInfo(IsValidBindPose(SDKScene, Root)));
+								}
 							}
 						}
 					}
 				}
 			}
 
-			void FFbxScene::FindForceJointNode(FbxScene* SDKScene, TArray<FbxNode*>& ForceJointNodes)
+			TArray<FbxNode*> FFbxScene::FindForceJointNode(FbxScene* SDKScene)
 			{
+				TArray<FbxNode*> ForceJointNodes;
+
 				int32 GeometryCount = SDKScene->GetGeometryCount();
 				for (int32 GeometryIndex = 0; GeometryIndex < GeometryCount; ++GeometryIndex)
 				{
@@ -855,7 +842,7 @@ namespace UE
 					{
 						continue;
 					}
-					FbxMesh* Mesh = static_cast<FbxMesh*>(Geometry);
+					const FbxMesh* Mesh = static_cast<FbxMesh*>(Geometry);
 					if (!Mesh)
 					{
 						continue;
@@ -883,6 +870,8 @@ namespace UE
 						}
 					}
 				}
+
+				return ForceJointNodes;
 			}
 
 			void FFbxScene::AddAnimationRecursively(FbxNode* Node
@@ -892,7 +881,7 @@ namespace UE
 				, UInterchangeSkeletalAnimationTrackNode* SkeletalAnimationTrackNode, bool SkeletalAnimationAddedToContainer
 				, const FString& RootSceneNodeUid, const TSet<FString>& SkeletonRootNodeUids
 				, const int32& AnimationIndex
-				, TArray<FbxNode*>& ForceJointNodes)
+				, const TArray<FbxNode*>& ForceJointNodes)
 			{
 				FString NodeUniqueID = Parser.GetFbxHelper()->GetFbxNodeHierarchyName(Node);
 				const bool bIsRootNode = Node == SDKScene->GetRootNode();
@@ -999,7 +988,7 @@ namespace UE
 					if (!HasSkeletonAttribute)
 					{
 						//in case the joint node "hierarchy finished" then the SkeletalAnimationTrackNode should be reset:
-						//as on the next occurance of a Joint node a New skeleton will start:
+						//as on the next occurrence of a Joint node a New skeleton will start:
 						SkeletalAnimationTrackNode = nullptr;
 
 						SkeletalAnimationAddedToContainer = false;
@@ -1038,8 +1027,7 @@ namespace UE
 
 				//Some fbx file have node without attribute that are link in cluster,
 				//We must consider those node has joint
-				TArray<FbxNode*> ForceJointNodes;
-				FindForceJointNode(SDKScene, ForceJointNodes);
+				TArray<FbxNode*> ForceJointNodes = FindForceJointNode(SDKScene);
 
 				//acquire Skeletal Node Uids from Meshes (via the skeletondependencies: )
 				TSet<FString> SkeletonRootNodeUids;
diff --git a/Engine/Plugins/Interchange/Runtime/Source/Parsers/Fbx/Private/FbxScene.h b/Engine/Plugins/Interchange/Runtime/Source/Parsers/Fbx/Private/FbxScene.h
--- a/Engine/Plugins/Interchange/Runtime/Source/Parsers/Fbx/Private/FbxScene.h
+++ b/Engine/Plugins/Interchange/Runtime/Source/Parsers/Fbx/Private/FbxScene.h
@@ -38,7 +38,8 @@ namespace UE
 
 				struct FRootJointInfo
 				{
-					bool bValidBindPose = false;
+					FRootJointInfo(bool InValue) : bValidBindPose(InValue) {}
+					const bool bValidBindPose;
 				};
 
 			protected:
@@ -50,7 +51,7 @@ namespace UE
 					, FbxScene* SDKScene
 					, UInterchangeBaseNodeContainer& NodeContainer
 					, TMap<FString, TSharedPtr<FPayloadContextBase, ESPMode::ThreadSafe>>& PayloadContexts
-					, TArray<FbxNode*>& ForceJointNodes
+					, const TArray<FbxNode*>& ForceJointNodes
 					, bool& bBadBindPoseMessageDisplay
 					, FFbxJointMeshBindPoseGenerator& FbxJointMeshBindPoseGenerator);
 
@@ -61,23 +62,21 @@ namespace UE
 					, UInterchangeSkeletalAnimationTrackNode* SkeletalAnimationTrackNode, bool SkeletalAnimationAddedToContainer
 					, const FString& RootSceneNodeUid, const TSet<FString>& SkeletonRootNodeUids
 					, const int32& AnimationIndex
-					, TArray<FbxNode*>& ForceJointNodes);
+					, const TArray<FbxNode*>& ForceJointNodes);
 
 			private:
 				void AddRigidAnimation(FbxNode* Node
-					, UInterchangeSceneNode* UnrealNode
+					, const UInterchangeSceneNode* UnrealNode
 					, UInterchangeBaseNodeContainer& NodeContainer
 					, TMap<FString, TSharedPtr<FPayloadContextBase, ESPMode::ThreadSafe>>& PayloadContexts);
+				
+				void FillCommonJointRootNode(FbxScene* SDKScene, const TArray<FbxNode*>& ForceJointNodes);
 
-				FbxNode* Internal_GetRootSkeleton(FbxScene* SDKScene, FbxNode* Link);
-				void FindCommonJointRootNode(FbxScene* SDKScene, const TArray<FbxNode*>& ForceJointNodes);
-
-				void FindForceJointNode(FbxScene* SDKScene, TArray<FbxNode*>& ForceJointNodes);
-
-				bool IsValidBindPose(FbxScene* SDKScene, FbxNode* RootJoint) const;
+				static FbxNode* Internal_GetRootSkeleton(FbxScene* SDKScene, FbxNode* Link, bool IsBlenderArmatureBone);
+				static TArray<FbxNode*> FindForceJointNode(FbxScene* SDKScene);
+				static bool IsValidBindPose(FbxScene* SDKScene, FbxNode* RootJoint);
 
 				TMap<FbxNode*, FRootJointInfo> CommonJointRootNodes;
-
 				FFbxParser& Parser;
 			};
 		}//ns Private
diff --git a/Engine/Plugins/Interchange/Runtime/Source/Pipelines/Private/InterchangeGenericScenesPipeline.cpp b/Engine/Plugins/Interchange/Runtime/Source/Pipelines/Private/InterchangeGenericScenesPipeline.cpp
--- a/Engine/Plugins/Interchange/Runtime/Source/Pipelines/Private/InterchangeGenericScenesPipeline.cpp
+++ b/Engine/Plugins/Interchange/Runtime/Source/Pipelines/Private/InterchangeGenericScenesPipeline.cpp
@@ -75,52 +75,24 @@ if(AttributeType AttributeName; TranslatedNode->GetCustom##AttributeName(Attribu
 
 namespace UE::Interchange::Private
 {
-	//Either a (TransformSpecialized || !JointSpecialized || RootJoint) can be a parent (only those get FactoryNodes) :
+	//Either a (!Joint || RootJoint) can be a parent (only those get FactoryNodes) :
 	FString FindFactoryParentSceneNodeUid(UInterchangeBaseNodeContainer* BaseNodeContainer, TArray<FString>& ActiveSkeletonUids, const UInterchangeSceneNode* SceneNode)
 	{
 		FString ParentUid = SceneNode->GetParentUid();
-		if (const UInterchangeSceneNode* ParentSceneNode = Cast<UInterchangeSceneNode>(BaseNodeContainer->GetNode(ParentUid)))
-		{
-			if (!ActiveSkeletonUids.Contains(ParentUid) || ParentSceneNode->IsSpecializedTypeContains(UE::Interchange::FSceneNodeStaticData::GetTransformSpecializeTypeString()))
-			{
-				return ParentUid;
-			}
-			else
-			{
 
-				bool bParentIsJoint = ParentSceneNode->IsSpecializedTypeContains(UE::Interchange::FSceneNodeStaticData::GetJointSpecializeTypeString());
-				if (bParentIsJoint)
-				{
-					FString GreatParentUid = ParentSceneNode->GetParentUid();
-					if (const UInterchangeSceneNode* GreatParentSceneNode = Cast<UInterchangeSceneNode>(BaseNodeContainer->GetNode(GreatParentUid)))
-					{
-						return FindFactoryParentSceneNodeUid(BaseNodeContainer, ActiveSkeletonUids, ParentSceneNode);
-					}
+		if (!ActiveSkeletonUids.Contains(ParentUid))
+		{
+			return ParentUid;
+		}
 
-					// #interchange_skeletal_revisit: Code above reflects actual behavior of code below 
-					//FString GreatParentUid = ParentSceneNode->GetParentUid();
-					//if (const UInterchangeSceneNode* GreatParentSceneNode = Cast<UInterchangeSceneNode>(BaseNodeContainer->GetNode(GreatParentUid)))
-					//{
-					//  Should we call ParentSceneNode or GreatParentSceneNode below?
-					//	bool bGreatParentIsJoint = (ParentSceneNode or GreatParentSceneNode?)->IsSpecializedTypeContains(UE::Interchange::FSceneNodeStaticData::GetJointSpecializeTypeString());
-					//	if (bGreatParentIsJoint)
-					//	{
-					//		return FindFactoryParentSceneNodeUid(BaseNodeContainer, ActiveSkeletonUids, ParentSceneNode);
-					//	}
-					//	else
-					//	{
-					//		return ParentUid;
-					//	}
-					//}
-				}
-				else
-				{
-					return ParentUid;
-				}
-			}
+		const UInterchangeSceneNode* ParentSceneNode = Cast<UInterchangeSceneNode>(BaseNodeContainer->GetNode(ParentUid));
+		if (!ParentSceneNode)
+		{
+			//Should never happen in theory since we always have the RootNode that cannot be a Joint and will never be in ActiveSkeletonUids
+			return UInterchangeBaseNode::InvalidNodeUid();
 		}
 
-		return UInterchangeBaseNode::InvalidNodeUid();
+		return FindFactoryParentSceneNodeUid(BaseNodeContainer, ActiveSkeletonUids, ParentSceneNode);
 	}
 
 #if WITH_EDITOR
@@ -395,40 +367,37 @@ void UInterchangeGenericLevelPipeline::ExecutePipeline(UInterchangeBaseNodeConta
 
 			if (SpecializeTypes.Num() > 0)
 			{
-				if (!SpecializeTypes.Contains(UE::Interchange::FSceneNodeStaticData::GetTransformSpecializeTypeString()))
+				bool bSkipNode = true;
+				if (SpecializeTypes.Contains(UE::Interchange::FSceneNodeStaticData::GetJointSpecializeTypeString()))
 				{
-					bool bSkipNode = true;
-					if (SpecializeTypes.Contains(UE::Interchange::FSceneNodeStaticData::GetJointSpecializeTypeString()))
+					if (!Cached_ActiveJointUids.Contains(SceneNode->GetUniqueID()))
 					{
-						if(!Cached_ActiveJointUids.Contains(SceneNode->GetUniqueID()))
-						{
-							bSkipNode = false;
-						}
-						else
+						bSkipNode = false;
+					}
+					else
+					{
+						//check if its the root joint (we want to create an actor for the root joint)
+						FString CurrentNodesParentUid = SceneNode->GetParentUid();
+						const UInterchangeBaseNode* ParentNode = BaseNodeContainer->GetNode(CurrentNodesParentUid);
+						if (const UInterchangeSceneNode* ParentSceneNode = Cast<UInterchangeSceneNode>(ParentNode))
 						{
-							//check if its the root joint (we want to create an actor for the root joint)
-							FString CurrentNodesParentUid = SceneNode->GetParentUid();
-							const UInterchangeBaseNode* ParentNode = BaseNodeContainer->GetNode(CurrentNodesParentUid);
-							if (const UInterchangeSceneNode* ParentSceneNode = Cast<UInterchangeSceneNode>(ParentNode))
+							if (!ParentSceneNode->IsSpecializedTypeContains(UE::Interchange::FSceneNodeStaticData::GetJointSpecializeTypeString()))
 							{
-								if (!ParentSceneNode->IsSpecializedTypeContains(UE::Interchange::FSceneNodeStaticData::GetJointSpecializeTypeString()))
-								{
-									bSkipNode = false;
-								}
+								bSkipNode = false;
 							}
 						}
 					}
-					else if (SpecializeTypes.Contains(UE::Interchange::FSceneNodeStaticData::GetLodGroupSpecializeTypeString()))
-					{
-						//Do not skip lod group, we always treat it has a mesh (import LODs option control if we import all lod or only the first)
-						bSkipNode = false;
-					}
+				}
+				else if (SpecializeTypes.Contains(UE::Interchange::FSceneNodeStaticData::GetLodGroupSpecializeTypeString()))
+				{
+					//Do not skip lod group, we always treat it has a mesh (import LODs option control if we import all lod or only the first)
+					bSkipNode = false;
+				}
 
-					if (bSkipNode)
-					{
-						//Skip any scene node that have specialized types but not the "Transform" type.
-						continue;
-					}
+				if (bSkipNode)
+				{
+					//Skip any scene node that have specialized types but not the "Transform" type.
+					continue;
 				}
 			}
 
@@ -1774,9 +1743,7 @@ void UInterchangeGenericLevelPipeline::CacheActiveJointUids()
 					{
 						if (const UInterchangeSceneNode* SceneNode = Cast<UInterchangeSceneNode>(Node))
 						{
-							TArray<FString> SpecializeTypes;
-							SceneNode->GetSpecializedTypes(SpecializeTypes);
-							if (SpecializeTypes.Contains(UE::Interchange::FSceneNodeStaticData::GetJointSpecializeTypeString()))
+							if (SceneNode->IsSpecializedTypeContains(UE::Interchange::FSceneNodeStaticData::GetJointSpecializeTypeString()))
 							{
 								Cached_ActiveJointUids.Add(Node->GetUniqueID());
 							}
