MaterialX:
- Fix <blur> by avoiding calling its nodegraph since since 1.39.4 the nodegraph is just a pass through

#Synced-CL 42983915

diff --git a/Engine/Plugins/Interchange/Runtime/Source/Import/Private/MaterialX/MaterialXUtils/MaterialXManager.cpp b/Engine/Plugins/Interchange/Runtime/Source/Import/Private/MaterialX/MaterialXUtils/MaterialXManager.cpp
--- a/Engine/Plugins/Interchange/Runtime/Source/Import/Private/MaterialX/MaterialXUtils/MaterialXManager.cpp
+++ b/Engine/Plugins/Interchange/Runtime/Source/Import/Private/MaterialX/MaterialXUtils/MaterialXManager.cpp
@@ -287,6 +287,7 @@ FMaterialXManager::FMaterialXManager()
 	, CategoriesToSkip {
 		mx::Category::ACEScgToLinRec709,
 		mx::Category::AdobeRgbToLinRec709,
+		mx::Category::Blur,
 		mx::Category::Checkerboard,
 		mx::Category::ChiangHairAbsorptionFromColor,
 		mx::Category::ChiangHairRoughness,
