Interchange:
- Grey out Autodetect cubemap texture option at reimport

#Synced-CL 42983915

diff --git a/Engine/Plugins/Interchange/Runtime/Source/Pipelines/Private/InterchangeGenericTexturePipeline.cpp b/Engine/Plugins/Interchange/Runtime/Source/Pipelines/Private/InterchangeGenericTexturePipeline.cpp
--- a/Engine/Plugins/Interchange/Runtime/Source/Pipelines/Private/InterchangeGenericTexturePipeline.cpp
+++ b/Engine/Plugins/Interchange/Runtime/Source/Pipelines/Private/InterchangeGenericTexturePipeline.cpp
@@ -398,7 +398,11 @@ void UInterchangeGenericTexturePipeline::AdjustSettingsForContext(const FInterch
 			HidePropertiesOfCategory(OuterMostPipeline, this, HideCategoryName);
 		}
 	}
+	
 #endif //WITH_EDITOR
+#if WITH_EDITORONLY_DATA
+	bIsReimport = ContextParams.ReimportAsset != nullptr;
+#endif
 }
 
 #if WITH_EDITOR
diff --git a/Engine/Plugins/Interchange/Runtime/Source/Pipelines/Public/InterchangeGenericTexturePipeline.h b/Engine/Plugins/Interchange/Runtime/Source/Pipelines/Public/InterchangeGenericTexturePipeline.h
--- a/Engine/Plugins/Interchange/Runtime/Source/Pipelines/Public/InterchangeGenericTexturePipeline.h
+++ b/Engine/Plugins/Interchange/Runtime/Source/Pipelines/Public/InterchangeGenericTexturePipeline.h
@@ -51,8 +51,8 @@ public:
 	UPROPERTY(EditAnywhere, BlueprintReadWrite, Category = "Textures", Meta=(EditCondition="bImportTextures"))
 	bool bImportUDIMs = true;
 
-	/** Detect if a texture should be imported as long/lat cubemap*/
-	UPROPERTY(EditAnywhere, BlueprintReadWrite, Category = "Textures", Meta = (EditCondition = "bImportTextures"))
+	/** Detect if a texture should be imported as long/lat cubemap. The minimum size to detect a cubemap texture is a 512 pixels height. Use the CVar Interchange.Textures.MinCubeMapDetection if you want to control that setting.*/
+	UPROPERTY(EditAnywhere, BlueprintReadWrite, Category = "Textures", Meta = (EditCondition = "bImportTextures && (bIsReimport == false)"))
 	bool bDetectLongLatCubemap = false;
 
 	/** Specify the file types that should be imported as long/lat cubemaps. */
@@ -106,6 +106,11 @@ private:
 
 	/** Texture factory assets nodes */
 	TArray<UInterchangeTextureFactoryNode*> TextureFactoryNodes;
+
+#if WITH_EDITORONLY_DATA
+	UPROPERTY()
+	bool bIsReimport = false;
+#endif
 };
 
 #undef UE_API
