[On behalf of Stefan Ivanovski from XRite]
- Added new InterchangeAxF plugin in the experimental folder
- Set the existing AxFImporter as deprecated

#virtualized
#Synced-CL 42983915

diff --git a/Engine/Plugins/Experimental/InterchangeAxF/Assets/Config/BaseInterchangeAxFAssets.ini b/Engine/Plugins/Experimental/InterchangeAxF/Assets/Config/BaseInterchangeAxFAssets.ini
new file mode 100644
--- /dev/null
+++ b/Engine/Plugins/Experimental/InterchangeAxF/Assets/Config/BaseInterchangeAxFAssets.ini
@@ -0,0 +1,43 @@
+[CoreRedirects]
+
+; Redirectors for /Interchange/Functions/ (30)
++PackageRedirects=(OldName="/InterchangeAxF/BaseMaterials/DefaultSpecularProfile.DefaultSpecularProfile", NewName="/InterchangeAxFAssets/BaseMaterials/DefaultSpecularProfile.DefaultSpecularProfile")
++PackageRedirects=(OldName="/InterchangeAxF/BaseMaterials/SVBRDF.SVBRDF", NewName="/InterchangeAxFAssets/BaseMaterials/SVBRDF.SVBRDF")
++PackageRedirects=(OldName="/InterchangeAxF/BaseMaterials/SVBRDF_Alpha.SVBRDF_Alpha", NewName="/InterchangeAxFAssets/BaseMaterials/SVBRDF_Alpha.SVBRDF_Alpha")
++PackageRedirects=(OldName="/InterchangeAxF/BaseMaterials/SVBRDF_Alpha_POM.SVBRDF_Alpha_POM", NewName="/InterchangeAxFAssets/BaseMaterials/SVBRDF_Alpha_POM.SVBRDF_Alpha_POM")
++PackageRedirects=(OldName="/InterchangeAxF/BaseMaterials/SVBRDF_POM.SVBRDF_POM", NewName="/InterchangeAxFAssets/BaseMaterials/SVBRDF_POM.SVBRDF_POM")
++PackageRedirects=(OldName="/InterchangeAxF/BaseMaterials/CPA2.CPA2", NewName="/InterchangeAxFAssets/BaseMaterials/CPA2.CPA2")
+
++PackageRedirects=(OldName="/InterchangeAxF/MaterialFunctions/MF_InvertMultipleScatteringGGX.MF_InvertMultipleScatteringGGX", NewName="/InterchangeAxFAssets/MaterialFunctions/MF_InvertMultipleScatteringGGX.MF_InvertMultipleScatteringGGX")
++PackageRedirects=(OldName="/InterchangeAxF/MaterialFunctions/MF_MakeTriPlanarUV.MF_MakeTriPlanarUV", NewName="/InterchangeAxFAssets/MaterialFunctions/MF_MakeTriPlanarUV.MF_MakeTriPlanarUV")
++PackageRedirects=(OldName="/InterchangeAxF/MaterialFunctions/MF_Noise_InterLeavedGradientGloden_1d.MF_Noise_InterLeavedGradientGloden_1d", NewName="/InterchangeAxFAssets/MaterialFunctions/MF_Noise_InterLeavedGradientGloden_1d.MF_Noise_InterLeavedGradientGloden_1d")
++PackageRedirects=(OldName="/InterchangeAxF/MaterialFunctions/MF_NormalStrength.MF_NormalStrength", NewName="/InterchangeAxFAssets/MaterialFunctions/MF_NormalStrength.MF_NormalStrength")
++PackageRedirects=(OldName="/InterchangeAxF/MaterialFunctions/MF_POM_TC.MF_POM_TC", NewName="/InterchangeAxFAssets/MaterialFunctions/MF_POM_TC.MF_POM_TC")
++PackageRedirects=(OldName="/InterchangeAxF/MaterialFunctions/MF_POM_TP.MF_POM_TP", NewName="/InterchangeAxFAssets/MaterialFunctions/MF_POM_TP.MF_POM_TP")
++PackageRedirects=(OldName="/InterchangeAxF/MaterialFunctions/MF_POMOverrideBlend.MF_POMOverrideBlend", NewName="/InterchangeAxFAssets/MaterialFunctions/MF_POMOverrideBlend.MF_POMOverrideBlend")
++PackageRedirects=(OldName="/InterchangeAxF/MaterialFunctions/MF_Rotate2D.MF_Rotate2D", NewName="/InterchangeAxFAssets/MaterialFunctions/MF_Rotate2D.MF_Rotate2D")
++PackageRedirects=(OldName="/InterchangeAxF/MaterialFunctions/MF_SampleNormalTEX.MF_SampleNormalTEX", NewName="/InterchangeAxFAssets/MaterialFunctions/MF_SampleNormalTEX.MF_SampleNormalTEX")
++PackageRedirects=(OldName="/InterchangeAxF/MaterialFunctions/MF_SampleTEX.MF_SampleTEX", NewName="/InterchangeAxFAssets/MaterialFunctions/MF_SampleTEX.MF_SampleTEX")
++PackageRedirects=(OldName="/InterchangeAxF/MaterialFunctions/MF_SheenRoughness.MF_SheenRoughness", NewName="/InterchangeAxFAssets/MaterialFunctions/MF_SheenRoughness.MF_SheenRoughness")
++PackageRedirects=(OldName="/InterchangeAxF/MaterialFunctions/MF_smin.MF_smin", NewName="/InterchangeAxFAssets/MaterialFunctions/MF_smin.MF_smin")
++PackageRedirects=(OldName="/InterchangeAxF/MaterialFunctions/MF_TexCoodUVNormalSampler.MF_TexCoodUVNormalSampler", NewName="/InterchangeAxFAssets/MaterialFunctions/MF_TexCoodUVNormalSampler.MF_TexCoodUVNormalSampler")
++PackageRedirects=(OldName="/InterchangeAxF/MaterialFunctions/MF_TextureCoordinationTransform.MF_TextureCoordinationTransform", NewName="/InterchangeAxFAssets/MaterialFunctions/MF_TextureCoordinationTransform.MF_TextureCoordinationTransform")
++PackageRedirects=(OldName="/InterchangeAxF/MaterialFunctions/MF_Triplanar_Core.MF_Triplanar_Core", NewName="/InterchangeAxFAssets/MaterialFunctions/MF_Triplanar_Core.MF_Triplanar_Core")
++PackageRedirects=(OldName="/InterchangeAxF/MaterialFunctions/MF_Triplanar_NormalCorrection.MF_Triplanar_NormalCorrection", NewName="/InterchangeAxFAssets/MaterialFunctions/MF_Triplanar_NormalCorrection.MF_Triplanar_NormalCorrection")
++PackageRedirects=(OldName="/InterchangeAxF/MaterialFunctions/MF_TriplanarBlendHeightA.MF_TriplanarBlendHeightA", NewName="/InterchangeAxFAssets/MaterialFunctions/MF_TriplanarBlendHeightA.MF_TriplanarBlendHeightA")
++PackageRedirects=(OldName="/InterchangeAxF/MaterialFunctions/MF_TriPlanarUV_Coods.MF_TriPlanarUV_Coods", NewName="/InterchangeAxFAssets/MaterialFunctions/MF_TriPlanarUV_Coods.MF_TriPlanarUV_Coods")
++PackageRedirects=(OldName="/InterchangeAxF/MaterialFunctions/MF_TriPlanarUVNormalSampler.MF_TriPlanarUVNormalSampler", NewName="/InterchangeAxFAssets/MaterialFunctions/MF_TriPlanarUVNormalSampler.MF_TriPlanarUVNormalSampler")
++PackageRedirects=(OldName="/InterchangeAxF/MaterialFunctions/MF_ViewTaceLength_TP.MF_ViewTaceLength_TP", NewName="/InterchangeAxFAssets/MaterialFunctions/MF_ViewTaceLength_TP.MF_ViewTaceLength_TP")
++PackageRedirects=(OldName="/InterchangeAxF/MaterialFunctions/MF_ViewTraceLength_TC.MF_ViewTraceLength_TC", NewName="/InterchangeAxFAssets/MaterialFunctions/MF_ViewTraceLength_TC.MF_ViewTraceLength_TC")
++PackageRedirects=(OldName="/InterchangeAxF/MaterialFunctions/SheenRoughnessComponent.SheenRoughnessComponent", NewName="/InterchangeAxFAssets/MaterialFunctions/SheenRoughnessComponent.SheenRoughnessComponent")
++PackageRedirects=(OldName="/InterchangeAxF/MaterialFunctions/TextureCoordinateTransform.TextureCoordinateTransform", NewName="/InterchangeAxFAssets/MaterialFunctions/TextureCoordinateTransform.TextureCoordinateTransform")
++PackageRedirects=(OldName="/InterchangeAxF/MaterialFunctions/AnisotropyComponent.AnisotropyComponent", NewName="/InterchangeAxFAssets/MaterialFunctions/AnisotropyComponent.AnisotropyComponent")
++PackageRedirects=(OldName="/InterchangeAxF/MaterialFunctions/InvertMultipleScatteringGGX.InvertMultipleScatteringGGX", NewName="/InterchangeAxFAssets/MaterialFunctions/InvertMultipleScatteringGGX.InvertMultipleScatteringGGX")
++PackageRedirects=(OldName="/InterchangeAxF/MaterialFunctions/MF_AnisotropyComponent.MF_AnisotropyComponent", NewName="/InterchangeAxFAssets/MaterialFunctions/MF_AnisotropyComponent.MF_AnisotropyComponent")
++PackageRedirects=(OldName="/InterchangeAxF/MaterialFunctions/MF_AnisotropyTangent.MF_AnisotropyTangent", NewName="/InterchangeAxFAssets/MaterialFunctions/MF_AnisotropyTangent.MF_AnisotropyTangent")
++PackageRedirects=(OldName="/InterchangeAxF/MaterialFunctions/MF_ClearCoatTerm.MF_ClearCoatTerm", NewName="/InterchangeAxFAssets/MaterialFunctions/MF_ClearCoatTerm.MF_ClearCoatTerm")
++PackageRedirects=(OldName="/InterchangeAxF/MaterialFunctions/MF_DefaultBlend.MF_DefaultBlend", NewName="/InterchangeAxFAssets/MaterialFunctions/MF_DefaultBlend.MF_DefaultBlend")
+
++PackageRedirects=(OldName="/InterchangeAxF/Meshes/SM_MatPreviewMesh_01.SM_MatPreviewMesh_01", NewName="/InterchangeAxFAssets/Meshes/SM_MatPreviewMesh_01.SM_MatPreviewMesh_01")
+
++PackageRedirects=(OldName="/InterchangeAxF/Textures/GGX_ReflectionEnergy.GGX_ReflectionEnergy", NewName="/InterchangeAxFAssets/Textures/GGX_ReflectionEnergy.GGX_ReflectionEnergy")
\ No newline at end of file
diff --git a/Engine/Plugins/Experimental/InterchangeAxF/Assets/Config/FilterPlugin.ini b/Engine/Plugins/Experimental/InterchangeAxF/Assets/Config/FilterPlugin.ini
new file mode 100644
--- /dev/null
+++ b/Engine/Plugins/Experimental/InterchangeAxF/Assets/Config/FilterPlugin.ini
@@ -0,0 +1,8 @@
+[FilterPlugin]
+; This section lists additional files which will be packaged along with your plugin. Paths should be listed relative to the root plugin directory, and
+; may include "...", "*", and "?" wildcards to match directories, files, and individual characters respectively.
+;
+; Examples:
+;    /README.txt
+;    /Extras/...
+;    /Binaries/ThirdParty/*.dll
diff --git a/Engine/Plugins/Experimental/InterchangeAxF/Assets/InterchangeAxfAssets.uplugin b/Engine/Plugins/Experimental/InterchangeAxF/Assets/InterchangeAxfAssets.uplugin
new file mode 100644
--- /dev/null
+++ b/Engine/Plugins/Experimental/InterchangeAxF/Assets/InterchangeAxfAssets.uplugin
@@ -0,0 +1,24 @@
+{
+	"FileVersion": 3,
+	"Version": 1,
+	"VersionName": "1.0",
+	"FriendlyName": "Interchange AxF Assets",
+	"Description": "The Interchange AxF Assets plugin exposes the assets used by the InterchangeAxf import plugin.",
+	"Category": "",
+	"CreatedBy": "",
+	"CreatedByURL": "",
+	"DocsURL": "",
+	"MarketplaceURL": "",
+	"SupportURL": "",
+	"EnabledByDefault": true,
+	"CanContainContent": true,
+	"IsExperimentalVersion": true,
+	"Installed": false,
+	"Modules": [
+		{
+			"Name": "InterchangeAxFAssets",
+			"Type": "Runtime",
+			"LoadingPhase": "PostConfigInit"
+		}
+	]
+}
\ No newline at end of file
diff --git a/Engine/Plugins/Experimental/InterchangeAxF/Assets/Source/Assets/InterchangeAxFAssets.Build.cs b/Engine/Plugins/Experimental/InterchangeAxF/Assets/Source/Assets/InterchangeAxFAssets.Build.cs
new file mode 100644
--- /dev/null
+++ b/Engine/Plugins/Experimental/InterchangeAxF/Assets/Source/Assets/InterchangeAxFAssets.Build.cs
@@ -0,0 +1,19 @@
+// Copyright Epic Games, Inc. All Rights Reserved.
+
+namespace UnrealBuildTool.Rules
+{
+	public class InterchangeAxFAssets : ModuleRules
+	{
+		public InterchangeAxFAssets(ReadOnlyTargetRules Target) : base(Target)
+		{
+			PrivateDependencyModuleNames.AddRange(
+				new string[]
+				{
+					"Core",
+					"Projects",
+					"RenderCore"
+				}
+			);
+		}
+	}
+}
diff --git a/Engine/Plugins/Experimental/InterchangeAxF/Assets/Source/Assets/Private/InterchangeAxFAssetsModule.cpp b/Engine/Plugins/Experimental/InterchangeAxF/Assets/Source/Assets/Private/InterchangeAxFAssetsModule.cpp
new file mode 100644
--- /dev/null
+++ b/Engine/Plugins/Experimental/InterchangeAxF/Assets/Source/Assets/Private/InterchangeAxFAssetsModule.cpp
@@ -0,0 +1,18 @@
+// Copyright Epic Games, Inc. All Rights Reserved.
+
+#include "CoreMinimal.h"
+#include "Modules/ModuleInterface.h"
+#include "Modules/ModuleManager.h"
+
+class FInterchangeAxFAssetsModule : public IModuleInterface
+{
+	/** IModuleInterface implementation */
+	virtual void StartupModule() override
+	{}
+
+	virtual void ShutdownModule() override
+	{}
+};
+
+IMPLEMENT_MODULE(FInterchangeAxFAssetsModule, InterchangeAxFAssets)
+
diff --git a/Engine/Plugins/Experimental/InterchangeAxF/Runtime/Config/FilterPlugin.ini b/Engine/Plugins/Experimental/InterchangeAxF/Runtime/Config/FilterPlugin.ini
new file mode 100644
--- /dev/null
+++ b/Engine/Plugins/Experimental/InterchangeAxF/Runtime/Config/FilterPlugin.ini
@@ -0,0 +1,8 @@
+[FilterPlugin]
+; This section lists additional files which will be packaged along with your plugin. Paths should be listed relative to the root plugin directory, and
+; may include "...", "*", and "?" wildcards to match directories, files, and individual characters respectively.
+;
+; Examples:
+;    /README.txt
+;    /Extras/...
+;    /Binaries/ThirdParty/*.dll
diff --git a/Engine/Plugins/Experimental/InterchangeAxF/Runtime/InterchangeAxF.uplugin b/Engine/Plugins/Experimental/InterchangeAxF/Runtime/InterchangeAxF.uplugin
new file mode 100644
--- /dev/null
+++ b/Engine/Plugins/Experimental/InterchangeAxF/Runtime/InterchangeAxF.uplugin
@@ -0,0 +1,38 @@
+{
+	"FileVersion": 3,
+	"Version": 1,
+	"VersionName": "1.0",
+	"FriendlyName": "Interchange AxF",
+	"Description": "Imports AxF2UE files from X-Rite Pantora as Substrate materials via the Interchange framework.",
+	"Category": "",
+	"CreatedBy": "",
+	"CreatedByURL": "",
+	"DocsURL": "",
+	"MarketplaceURL": "",
+	"SupportURL": "",
+	"CanContainContent": true,
+	"IsExperimentalVersion": true,
+	"Installed": false,
+	"SupportedTargetPlatforms": [
+		"Win64",
+		"Linux",
+		"Mac"
+	],
+	"Modules": [
+		{
+			"Name": "InterchangeAxF",
+			"Type": "Editor",
+			"LoadingPhase": "PostEngineInit"
+		}
+	],
+	"Plugins": [
+		{
+			"Name": "InterchangeAxFAssets",
+			"Enabled": true
+		},
+		{
+			"Name": "Interchange",
+			"Enabled": true
+		}
+	]
+}
\ No newline at end of file
diff --git a/Engine/Plugins/Experimental/InterchangeAxF/Runtime/Source/InterchangeAxF/InterchangeAxF.Build.cs b/Engine/Plugins/Experimental/InterchangeAxF/Runtime/Source/InterchangeAxF/InterchangeAxF.Build.cs
new file mode 100644
--- /dev/null
+++ b/Engine/Plugins/Experimental/InterchangeAxF/Runtime/Source/InterchangeAxF/InterchangeAxF.Build.cs
@@ -0,0 +1,41 @@
+// Copyright Epic Games, Inc. All Rights Reserved.
+
+using UnrealBuildTool;
+
+public class InterchangeAxF : ModuleRules
+{
+    public InterchangeAxF(ReadOnlyTargetRules Target) : base(Target)
+    {
+        PCHUsage = ModuleRules.PCHUsageMode.UseExplicitOrSharedPCHs;
+		//OptimizeCode = CodeOptimization.Never;
+
+        PublicDependencyModuleNames.AddRange(
+            new string[]
+            {
+                "Core",
+			}
+            );
+
+
+        PrivateDependencyModuleNames.AddRange(
+            new string[]
+            {
+                "Core",
+                "CoreUObject",
+                "Engine",
+                "Slate",
+                "SlateCore",
+                "UnrealEd",
+                "MaterialEditor",
+                "InterchangeCore",
+                "InterchangeEngine",
+                "InterchangeImport",
+                "InterchangeNodes",
+                "InterchangeFactoryNodes",
+                "InterchangePipelines",
+                "InterchangeAxFAssets",
+                "Json"
+            }
+            );
+    }
+}
diff --git a/Engine/Plugins/Experimental/InterchangeAxF/Runtime/Source/InterchangeAxF/Private/AxFInterchangePipeline.cpp b/Engine/Plugins/Experimental/InterchangeAxF/Runtime/Source/InterchangeAxF/Private/AxFInterchangePipeline.cpp
new file mode 100644
--- /dev/null
+++ b/Engine/Plugins/Experimental/InterchangeAxF/Runtime/Source/InterchangeAxF/Private/AxFInterchangePipeline.cpp
@@ -0,0 +1,451 @@
+// Copyright Epic Games, Inc. All Rights Reserved.
+
+#include "AxFInterchangePipeline.h"
+#include "AxFMaterialObjectNode.h"
+#include "Engine/SpecularProfile.h"
+#include "InterchangeGenericTexturePipeline.h"
+#include "Material/InterchangeMaterialFactory.h"
+#include "InterchangeMaterialFactoryNode.h"
+#include "InterchangePipelineHelper.h"
+#include "InterchangeTexture2DFactoryNode.h"
+#include "InterchangeTexture2DNode.h"
+#include "Materials/Material.h"
+#include "Materials/MaterialInstanceConstant.h"
+#include "Nodes/InterchangeSourceNode.h"
+
+DEFINE_LOG_CATEGORY(LogAxFInterchangePipeline);
+
+// Helper structure to differentiate scalar from vector parameters stored in
+// UAxFMaterialObjectNode::PayloadData::ValuesMap
+struct FMaterialParameterTypes
+{
+	TSet<FString> ScalarParameters;
+	TSet<FString> VectorParameters;
+};
+
+// Map to store the names of the scalar and vector parameters of the predefined materials.
+static TMap<FString, FMaterialParameterTypes> ParameterTypesPerMaterial;
+
+static FString GetParameterNameFromUID(const FString & UID)
+{
+    FString LeftPart, ParameterName;
+    UID.Split("_", &LeftPart, &ParameterName, ESearchCase::IgnoreCase, ESearchDir::FromEnd);
+    return ParameterName;
+}
+
+UAxFInterchangePipeline::UAxFInterchangePipeline()
+{
+    TexturePipeline = CreateDefaultSubobject<UInterchangeGenericTexturePipeline>(TEXT("TexturePipeline"));
+    TexturePipeline->bDetectNormalMapTexture = false;
+    TexturePipeline->bAllowNonPowerOfTwo = true;
+}
+
+void UAxFInterchangePipeline::ExecutePipeline(UInterchangeBaseNodeContainer* BaseNodeContainer, 
+                                              TArray<UInterchangeSourceData*> const& SourceDatas, 
+                                              FString const& ContentBasePath)
+{
+    if (TexturePipeline)
+    {
+        TexturePipeline->ScriptedExecutePipeline(BaseNodeContainer, SourceDatas, ContentBasePath);
+    }
+
+    // Find all translated node we need for this pipeline
+    BaseNodeContainer->IterateNodes(
+        [this](FString const& NodeUid, UInterchangeBaseNode* Node) {
+        switch (Node->GetNodeContainerType())
+        {
+            case EInterchangeNodeContainerType::TranslatedAsset:
+            {
+                if (UAxFMaterialObjectNode* CustomDemoObjectNode = Cast<UAxFMaterialObjectNode>(Node))
+                {
+                    AxFObjectNodes.Add(CustomDemoObjectNode);
+                }
+                else if (UInterchangeTexture2DNode* Texture2DNode = Cast<UInterchangeTexture2DNode>(Node))
+                {
+					double WidthMM = 0., HeightMM = 0.;
+                    Texture2DNode->GetAttribute(TEXT("WidthMM"), WidthMM);
+                    Texture2DNode->GetAttribute(TEXT("HeightMM"), HeightMM);
+
+                    FString ParameterName = GetParameterNameFromUID(Texture2DNode->GetUniqueID());
+                    FAxFTextureMap& Pair = TexturesMap.FindOrAdd(ParameterName);
+                    Pair.UV = FVector2d{ WidthMM, HeightMM };
+				}
+            }
+            break;
+            case EInterchangeNodeContainerType::FactoryData:
+            {
+                if (UInterchangeTexture2DFactoryNode* Texture2DFactoryNode = Cast<UInterchangeTexture2DFactoryNode>(Node))
+                {
+                    Texture2DFactoryNodes.Add(Texture2DFactoryNode);
+                }
+            }
+            break;
+            default:;
+        }
+    });
+
+    if (AxFObjectNodes.IsEmpty())
+    {
+        return;
+    }
+
+	for (const UAxFMaterialObjectNode* AxFNode : AxFObjectNodes)
+    {
+		const TArray<EAxFFeature>& Features = AxFNode->PayloadData.UsedFeatures;
+		const TMap<FString, FLinearColor>& ValuesMap = AxFNode->PayloadData.ValuesMap;
+
+		FString BaseMaterialName;
+		if (Features.Contains(EAxFFeature::CPA2))
+		{
+			BaseMaterialName += TEXT("CPA2");
+            BaseMaterialName += (Features.Contains(EAxFFeature::Flakes) ? TEXT("_Flakes") : TEXT(""));
+		}
+		else
+		{
+			BaseMaterialName += "SVBRDF";
+			{
+				BaseMaterialName += (Features.Contains(EAxFFeature::Alpha) ? TEXT("_Alpha") : TEXT(""));
+				BaseMaterialName += (Features.Contains(EAxFFeature::Heightmap) ? TEXT("_POM") : TEXT(""));
+			}
+		}
+
+		FString FullBaseMaterialPath =
+			TEXT("/InterchangeAxFAssets/BaseMaterials/") + BaseMaterialName + TEXT(".") + BaseMaterialName;
+
+		if (!ensure(ParameterTypesPerMaterial.Contains(BaseMaterialName)))
+		{
+			UE_LOG(LogAxFInterchangePipeline, Error, TEXT("Failed to create material instance %s: %s was not found"), *AxFNode->GetDisplayLabel(), *FullBaseMaterialPath);
+			continue;
+		}
+
+		const FMaterialParameterTypes& ParameterTypes = ParameterTypesPerMaterial[BaseMaterialName];
+
+		FString MaterialInstanceUID = TEXT("MI_") + AxFNode->GetDisplayLabel();
+
+        UInterchangeMaterialInstanceFactoryNode* MaterialInstance =
+            NewObject<UInterchangeMaterialInstanceFactoryNode>(BaseNodeContainer, FName(MaterialInstanceUID));
+        MaterialInstance->InitializeNode(
+            MaterialInstanceUID, AxFNode->GetDisplayLabel(), EInterchangeNodeContainerType::FactoryData);
+
+		for (const EAxFFeature Feature : Features)
+		{
+			const FString Value = UEnum::GetValueAsString(Feature);
+			const FString AttributeKey = UInterchangeShaderPortsAPI::MakeInputValueKey(TEXT("Has") + Value + TEXT("?"));
+
+			MaterialInstance->AddBooleanAttribute(AttributeKey, true);
+		}
+
+		TArray<FString> ParameterNames;
+		ValuesMap.GetKeys(ParameterNames);
+		constexpr float Epsilon = 1e-6;
+		for (FString const& Param : ParameterNames)
+		{
+            // NOTE: Exceptional case where we don't have a "GlintUVScale" parameter in the base material,
+            // rather we need to tweak individual parameters for each texture coordinate axis
+            if (Param == TEXT("GlintUVScale"))
+            {
+                const FString XAttributeKey = UInterchangeShaderPortsAPI::MakeInputValueKey(TEXT("ScaleX"));
+                const FString YAttributeKey = UInterchangeShaderPortsAPI::MakeInputValueKey(TEXT("ScaleY"));
+				const FLinearColor& GlintUVScale = ValuesMap[Param];
+				MaterialInstance->AddFloatAttribute(XAttributeKey, 1.0f / (GlintUVScale.R + Epsilon));
+				MaterialInstance->AddFloatAttribute(YAttributeKey, 1.0f / (GlintUVScale.G + Epsilon));
+            }
+
+			// NOTE: Disable texture path in the case that we have a value for this parameter
+			FString TextureToggleAttributeKey = UInterchangeShaderPortsAPI::MakeInputValueKey(TEXT("Has") + Param + TEXT("Texture?"));
+			MaterialInstance->AddBooleanAttribute(TextureToggleAttributeKey, false);
+
+			const FString AttributeKey = UInterchangeShaderPortsAPI::MakeInputValueKey(Param);
+			const FLinearColor& Value = ValuesMap[Param];
+			if (ParameterTypes.ScalarParameters.Contains(Param))
+			{
+				MaterialInstance->AddFloatAttribute(AttributeKey, Value.R);
+			}
+			else
+			{
+				MaterialInstance->AddLinearColorAttribute(AttributeKey, Value);
+			}
+		}
+        
+        if (bUseTriplanarMappingByDefault)
+        {
+            FString TriplanarUVsToggleAttributeKey = UInterchangeShaderPortsAPI::MakeInputValueKey(TEXT("EnableTriplanarUVs"));
+            MaterialInstance->AddBooleanAttribute(TriplanarUVsToggleAttributeKey, true);
+        }
+
+		for (UInterchangeTexture2DFactoryNode* Texture2DFactoryNode : Texture2DFactoryNodes)
+		{
+			Texture2DFactoryNode->SetCustomSubPath(AxFNode->GetDisplayLabel());
+
+			const FString Param = GetParameterNameFromUID(Texture2DFactoryNode->GetUniqueID());
+
+			{
+				FString AttributeKey = UInterchangeShaderPortsAPI::MakeInputValueKey(TEXT("Has") + Param + TEXT("Texture?"));
+				MaterialInstance->AddBooleanAttribute(AttributeKey, true);
+			}
+
+			{
+				FString AttributeKey = UInterchangeShaderPortsAPI::MakeInputValueKey(Param + TEXT("Texture"));
+				MaterialInstance->AddStringAttribute(AttributeKey, Texture2DFactoryNode->GetUniqueID());
+			}
+
+			if (Param == TEXT("HeightMap"))
+			{
+				FString AttributeKey = UInterchangeShaderPortsAPI::MakeInputValueKey(Param);
+				MaterialInstance->AddStringAttribute(AttributeKey, Texture2DFactoryNode->GetUniqueID());
+			}
+
+            if (Param != TEXT("CPA2ColorTable"))
+			{
+				const FVector2d UVScale = FVector2d::UnitVector / (TexturesMap.FindOrAdd(Param).UV + Epsilon);
+				FString UVScaleXAttributeKey = UInterchangeShaderPortsAPI::MakeInputValueKey(TEXT("ScaleX"));
+				FString UVScaleYAttributeKey = UInterchangeShaderPortsAPI::MakeInputValueKey(TEXT("ScaleY"));
+				MaterialInstance->AddFloatAttribute(UVScaleXAttributeKey, UVScale.X);
+				MaterialInstance->AddFloatAttribute(UVScaleYAttributeKey, UVScale.Y);
+			}
+
+			MaterialInstance->AddFactoryDependencyUid(Texture2DFactoryNode->GetUniqueID());
+		}
+		Texture2DFactoryNodes.Empty();
+
+        MaterialInstance->SetCustomParent(FullBaseMaterialPath);
+        MaterialInstance->MarkPackageDirty();
+        MaterialInstance->PostEditChange();
+        BaseNodeContainer->AddNode(MaterialInstance);
+    }
+
+	BaseNodeContainer->IterateNodesOfType<UInterchangeFactoryBaseNode>(
+		[ReimportStrategyClosure = ReimportStrategy](const FString& NodeUid, UInterchangeFactoryBaseNode* FactoryNode)
+		{
+			// Make sure all factory nodes have the specified strategy
+			FactoryNode->SetReimportStrategyFlags(ReimportStrategyClosure);
+		}
+	);
+}
+
+void UAxFInterchangePipeline::ExecutePostImportPipeline(
+    UInterchangeBaseNodeContainer const* BaseNodeContainer, FString const& NodeKey, UObject* CreatedAsset,
+    bool bIsAReimport)
+{
+    if (TexturePipeline)
+    {
+        TexturePipeline->ScriptedExecutePostImportPipeline(BaseNodeContainer, NodeKey, CreatedAsset, bIsAReimport);
+    }
+
+    UMaterialInstanceConstant* Instance = Cast<UMaterialInstanceConstant>(CreatedAsset);
+    if (Instance != nullptr && !AxFObjectNodes.IsEmpty())
+    {
+        TObjectPtr<UAxFMaterialObjectNode>* FoundData = AxFObjectNodes.FindByPredicate(
+            [Instance](UAxFMaterialObjectNode const* Node) { return Node->GetDisplayLabel() == Instance->GetName(); });
+
+        if (FoundData == nullptr)
+        {
+            UE_LOG(
+                LogAxFInterchangePipeline, Error,
+                TEXT("Couldn't find appropriate AxF translated node to create Material Instance."))
+            return;
+        }
+
+		UAxFMaterialObjectNode* AxFData = *FoundData;
+
+		const TArray<EAxFFeature>& Features = AxFData->PayloadData.UsedFeatures;
+		
+        if (SpecularProfile != nullptr)
+        {
+            Instance->bOverrideSpecularProfile = true;
+            Instance->SpecularProfileOverride = SpecularProfile;
+            SpecularProfile = nullptr;
+        }
+
+        if (!Instance->MarkPackageDirty())
+        {
+            UE_LOG(
+                LogAxFInterchangePipeline, Warning,
+                TEXT("Failed to mark package as dirty. Some changes to the material instance %s might not be saved."),
+                *Instance->GetName());
+        }
+
+        Instance->PostEditChange();
+
+        AxFObjectNodes.RemoveSingle(AxFData);
+    }
+
+    Super::ExecutePostImportPipeline(BaseNodeContainer, NodeKey, CreatedAsset, bIsAReimport);
+}
+
+void UAxFInterchangePipeline::ExecutePostFactoryPipeline(
+    UInterchangeBaseNodeContainer const* BaseNodeContainer, FString const& NodeKey, UObject* CreatedAsset,
+    bool bIsAReimport)
+{
+    if (TexturePipeline)
+    {
+        TexturePipeline->ScriptedExecutePostFactoryPipeline(BaseNodeContainer, NodeKey, CreatedAsset, bIsAReimport);
+    }
+
+    FString ParameterName = GetParameterNameFromUID(NodeKey);
+
+    UTexture2D* Texture2D = Cast<UTexture2D>(CreatedAsset);
+    if (Texture2D != nullptr)
+    {
+        ETextureSourceFormat Format = Texture2D->Source.GetFormat();
+        if (Format == TSF_BGRA8 || Format == TSF_G8 || Format == TSF_RGBA16)
+        {
+            Texture2D->CompressionSettings = TC_BC7;
+        }
+        else if (Format == TSF_RGBA32F)
+        {
+            Texture2D->CompressionSettings = TC_HDR_Compressed;
+        }
+        else if (Format == TSF_R32F || Format == TSF_R16F)
+        {
+            Texture2D->CompressionSettings = TC_HalfFloat;
+        }
+        else
+        {
+            UE_LOG(
+                LogAxFInterchangePipeline, Warning,
+                TEXT("Unsupported texture source format! Can't set specific compression settings for %s"),
+                *Texture2D->GetName());
+        }
+
+        Texture2D->MipGenSettings = TMGS_FromTextureGroup;
+        Texture2D->PowerOfTwoMode = ETexturePowerOfTwoSetting::StretchToPowerOfTwo;
+        if (ParameterName == TEXT("Normal") || ParameterName == TEXT("ClearCoatNormal"))
+        {
+            Texture2D->bFlipGreenChannel = true;
+            Texture2D->CompressionSettings = TC_Normalmap;
+        }
+        Texture2D->SRGB = 0;
+
+        if (!Texture2D->MarkPackageDirty())
+        {
+            UE_LOG(
+                LogAxFInterchangePipeline, Warning,
+                TEXT("Failed to mark package as dirty. Some changes to the Texture2D %s might not be saved."),
+                *Texture2D->GetName());
+        }
+        Texture2D->PostEditChange();
+
+		FAxFTextureMap& Pair = TexturesMap.FindOrAdd(ParameterName);
+		Pair.Texture = Texture2D;
+    }
+
+    USpecularProfile* Profile = Cast<USpecularProfile>(CreatedAsset);
+    if (Profile != nullptr)
+    {
+        SpecularProfile = Profile;
+    }
+
+    Super::ExecutePostFactoryPipeline(BaseNodeContainer, NodeKey, CreatedAsset, bIsAReimport);
+}
+
+void UAxFInterchangePipeline::AdjustSettingsForContext(FInterchangePipelineContextParams const& ContextParams)
+{
+    if (TexturePipeline)
+    {
+        TexturePipeline->AdjustSettingsForContext(ContextParams);
+    }
+
+    // Check if the base materials are loaded
+	static bool bRunOnce = []()
+		{
+			const FString BasePath = TEXT("/InterchangeAxFAssets/BaseMaterials/");
+            const FString BaseMaterialNames[] = { 
+				TEXT("SVBRDF"),
+				TEXT("SVBRDF_Alpha"),
+				TEXT("SVBRDF_POM"),
+				TEXT("SVBRDF_Alpha_POM"),
+				TEXT("CPA2"),
+				TEXT("CPA2_Flakes")
+			};
+
+			for (const FString& BaseMaterialName : BaseMaterialNames)
+			{
+				FSoftObjectPath UMaterialPath(FPaths::Combine(BasePath, BaseMaterialName) + TEXT(".") + BaseMaterialName);
+				UObject* MaterialObject = UMaterialPath.TryLoad();
+				if (MaterialObject == nullptr)
+				{
+					UE_LOG(
+						LogAxFInterchangePipeline, Error, TEXT("Failed to load base material from plugin Content: %s"),
+						*UMaterialPath.ToString());
+					continue;
+				}
+
+				UMaterial* Material = Cast<UMaterial>(MaterialObject);
+				if (Material == nullptr)
+				{
+					UE_LOG(
+						LogAxFInterchangePipeline, Error, TEXT("Loaded asset is not a material: %s"),
+						*UMaterialPath.ToString());
+					continue;
+				}
+
+				// Caching the scalar and vector parameters per predefined materials
+				FMaterialParameterTypes& ParameterTypes = ParameterTypesPerMaterial.Add(BaseMaterialName);
+				{
+					TArray<FMaterialParameterInfo> ParameterInfos;
+					TArray<FGuid> ParameterIds;
+					Material->GetAllScalarParameterInfo(ParameterInfos, ParameterIds);
+
+					for (const FMaterialParameterInfo& Info : ParameterInfos)
+					{
+						ParameterTypes.ScalarParameters.Add(Info.Name.ToString());
+					}
+				}
+
+				{
+					TArray<FMaterialParameterInfo> ParameterInfos;
+					TArray<FGuid> ParameterIds;
+					Material->GetAllVectorParameterInfo(ParameterInfos, ParameterIds);
+
+					for (const FMaterialParameterInfo& Info : ParameterInfos)
+					{
+						ParameterTypes.VectorParameters.Add(Info.Name.ToString());
+					}
+				}
+			}
+
+			return true;
+		}();
+
+    Super::AdjustSettingsForContext(ContextParams);
+}
+
+inline void UAxFInterchangePipeline::PreDialogCleanup(FName const PipelineStackName)
+{
+    if (TexturePipeline)
+    {
+        TexturePipeline->PreDialogCleanup(PipelineStackName);
+    }
+}
+
+bool UAxFInterchangePipeline::IsSettingsAreValid(TOptional<FText>& OutInvalidReason) const
+{
+    if (TexturePipeline && !TexturePipeline->IsSettingsAreValid(OutInvalidReason))
+    {
+        return false;
+    }
+
+    return Super::IsSettingsAreValid(OutInvalidReason);
+}
+
+inline void UAxFInterchangePipeline::FilterPropertiesFromTranslatedData(
+    UInterchangeBaseNodeContainer* InBaseNodeContainer)
+{
+    if (TexturePipeline)
+    {
+        TexturePipeline->FilterPropertiesFromTranslatedData(InBaseNodeContainer);
+    }
+
+    Super::FilterPropertiesFromTranslatedData(InBaseNodeContainer);
+}
+
+void UAxFInterchangePipeline::GetSupportAssetClasses(TArray<UClass*>& PipelineSupportAssetClasses) const
+{
+    if (TexturePipeline)
+    {
+        TexturePipeline->GetSupportAssetClasses(PipelineSupportAssetClasses);
+    }
+
+    PipelineSupportAssetClasses.Add(UMaterialInstance::StaticClass());
+}
diff --git a/Engine/Plugins/Experimental/InterchangeAxF/Runtime/Source/InterchangeAxF/Private/AxFMaterialObjectNode.cpp b/Engine/Plugins/Experimental/InterchangeAxF/Runtime/Source/InterchangeAxF/Private/AxFMaterialObjectNode.cpp
new file mode 100644
--- /dev/null
+++ b/Engine/Plugins/Experimental/InterchangeAxF/Runtime/Source/InterchangeAxF/Private/AxFMaterialObjectNode.cpp
@@ -0,0 +1,53 @@
+// Copyright Epic Games, Inc. All Rights Reserved.
+
+#include "AxFMaterialObjectNode.h"
+
+UAxFMaterialObjectNode::UAxFMaterialObjectNode() {}
+
+void UAxFMaterialObjectNode::InitializeAxFMaterialObjectNode(const FString &UniqueID, const FString &DisplayLabel)
+{
+    InitializeNode(UniqueID, DisplayLabel, EInterchangeNodeContainerType::TranslatedAsset);
+}
+
+FString UAxFMaterialObjectNode::GetTypeName() const
+{
+    FString TypeName = TEXT("AxFMaterialObjectNode");
+    return TypeName;
+}
+
+const TOptional<FString> UAxFMaterialObjectNode::GetPayloadKey() const
+{
+    FString PayloadKey;
+
+    UE::Interchange::FAttributeKey PayloadKeyAttributeName(TEXT("__PayloadKey__"));
+    if (!Attributes->ContainAttribute(PayloadKeyAttributeName))
+    {
+        return TOptional<FString>();
+    }
+
+    UE::Interchange::FAttributeStorage::TAttributeHandle<FString> AttributeHandle =
+        Attributes->GetAttributeHandle<FString>(PayloadKeyAttributeName);
+    if (!AttributeHandle.IsValid())
+    {
+        return TOptional<FString>();
+    }
+
+    UE::Interchange::EAttributeStorageResult Result = AttributeHandle.Get(PayloadKey);
+    if (!UE::Interchange::IsAttributeStorageResultSuccess(Result))
+    {
+        LogAttributeStorageErrors(Result, TEXT("UAxFMaterialObjectNode.GetPayLoadKey"), PayloadKeyAttributeName);
+        return TOptional<FString>();
+    }
+
+    return TOptional(PayloadKey);
+}
+
+void UAxFMaterialObjectNode::SetPayloadKey(const FString &PayloadKey)
+{
+    UE::Interchange::FAttributeKey PayloadKeyAttributeName(TEXT("__PayloadKey__"));
+    UE::Interchange::EAttributeStorageResult Result = Attributes->RegisterAttribute(PayloadKeyAttributeName, PayloadKey);
+    if (!UE::Interchange::IsAttributeStorageResultSuccess(Result))
+    {
+        LogAttributeStorageErrors(Result, TEXT("UAxFMaterialObjectNode.SetPayLoadKey"), PayloadKeyAttributeName);
+    }
+}
\ No newline at end of file
diff --git a/Engine/Plugins/Experimental/InterchangeAxF/Runtime/Source/InterchangeAxF/Private/AxFTranslator.cpp b/Engine/Plugins/Experimental/InterchangeAxF/Runtime/Source/InterchangeAxF/Private/AxFTranslator.cpp
new file mode 100644
--- /dev/null
+++ b/Engine/Plugins/Experimental/InterchangeAxF/Runtime/Source/InterchangeAxF/Private/AxFTranslator.cpp
@@ -0,0 +1,348 @@
+// Copyright Epic Games, Inc. All Rights Reserved.
+
+#include "AxFTranslator.h"
+#include "AxFMaterialObjectNode.h"
+#include "Engine/SpecularProfile.h"
+#include "Engine/RendererSettings.h"
+#include "ImageUtils.h"
+#include "InterchangeSpecularProfileFactoryNode.h"
+#include "InterchangeTexture2DNode.h"
+#include "Modules/ModuleManager.h"
+#include "InterchangeImportModule.h"
+
+#include "Serialization/JsonSerializer.h"
+
+DEFINE_LOG_CATEGORY(LogAxFTranslator);
+
+#include UE_INLINE_GENERATED_CPP_BY_NAME(AxFTranslator)
+
+static bool GInterchangeEnableAxFImport = true;
+static FAutoConsoleVariableRef CCvarInterchangeEnableAxFImport(
+    TEXT("Interchange.FeatureFlags.Import.AxF"),
+    GInterchangeEnableAxFImport, 
+    TEXT("Whether AxF Interchange support is enabled."),
+    ECVF_Default);
+
+bool UAxFTranslator::CanImportSourceData(const UInterchangeSourceData* InSourceData) const
+{
+	if (Super::CanImportSourceData(InSourceData))
+	{
+		if (IInterchangeImportModule::Get().IsSubstrateEnabled())
+		{
+			return true;
+		}
+		else
+		{
+			UE_LOG(LogAxFTranslator, Warning, TEXT("Could not create material from %s because Substrate materials are not enabled. Enable \"Substrate materials\" in Project Settings under Engine > Rendering."), *InSourceData->GetFilename());
+		}
+	}
+
+	return false;
+}
+
+TArray<FString> UAxFTranslator::GetSupportedFormats() const
+{
+    TArray<FString> Formats;
+    if (GInterchangeEnableAxFImport && IInterchangeImportModule::Get().IsSubstrateEnabled())
+    {
+        Formats.Add(TEXT("axf2ue;Intermediary AxF resource file"));
+    }
+
+    return Formats;
+}
+
+bool UAxFTranslator::Translate(UInterchangeBaseNodeContainer &BaseNodeContainer) const
+{
+    FString Filename = SourceData->GetFilename();
+
+    FString Path, Name, Extension;
+    FPaths::Split(Filename, Path, Name, Extension);
+
+    FPaths::NormalizeFilename(Filename);
+    if (!FPaths::FileExists(Filename))
+    {
+        return false;
+    }
+
+    FString DisplayLabel = FPaths::GetBaseFilename(Filename);
+    FString NodeUID(Filename);
+    UAxFMaterialObjectNode *AxFMaterialObjectNode = NewObject<UAxFMaterialObjectNode>(&BaseNodeContainer);
+    if (!ensure(AxFMaterialObjectNode))
+    {
+        return false;
+    }
+
+    AxFMaterialObjectNode->InitializeAxFMaterialObjectNode(NodeUID, DisplayLabel);
+    AxFMaterialObjectNode->SetPayloadKey(Filename);
+
+    FString JSONString;
+    if (!FFileHelper::LoadFileToString(JSONString, *SourceData->GetFilename()))
+    {
+        UE_LOG(LogAxFTranslator, Error, TEXT("Failed to read JSON file at %s"), *SourceData->GetFilename());
+        return {};
+    }
+
+    TSharedPtr<FJsonObject> JSON;
+    TSharedRef<TJsonReader<>> Reader = TJsonReaderFactory<>::Create(JSONString);
+    if (!FJsonSerializer::Deserialize(Reader, JSON))
+    {
+        UE_LOG(LogAxFTranslator, Error, TEXT("Failed to deserialize JSON read from file at %s"), *SourceData->GetFilename());
+        return {};
+    }
+
+	static const TMap<FString, FString> ParameterNameMapping{
+		{TEXT("alpha"), TEXT("Alpha")},
+		{TEXT("anisotropicRotation"), TEXT("AnisotropicRotation")},
+		{TEXT("clearCoatIOR"), TEXT("ClearCoatIOR")},
+		{TEXT("clearCoatNormal"), TEXT("ClearCoatNormal")},
+		{TEXT("clearCoatRoughness"), TEXT("ClearCoatRoughness")},
+		{TEXT("clearCoatTransmissionColor"), TEXT("ClearCoatTransmissionColor")},
+		{TEXT("cpaColorTable"), TEXT("CPA2ColorTable")},
+		{TEXT("diffuseColor"), TEXT("DiffuseColor")},
+		{TEXT("glintBrightness"), TEXT("GlintBrightness")},
+		{TEXT("glintDensity"), TEXT("GlintDensity")},
+		{TEXT("glintRoughness"), TEXT("GlintRoughness")},
+		{TEXT("glintUVScale"), TEXT("GlintUVScale")},
+		{TEXT("heightmap"), TEXT("HeightMap")},
+		{TEXT("heightmapScale"), TEXT("HeightStrength")},
+		{TEXT("heightmapReference"), TEXT("HeightReferencePlane")},
+		{TEXT("f0"), TEXT("F0")},
+		{TEXT("f90"), TEXT("F90")},
+		{TEXT("normal"), TEXT("Normal")},
+		{TEXT("roughness"), TEXT("Roughness")},
+		{TEXT("sheenColor"), TEXT("SheenColor")},
+		{TEXT("sheenRoughness"), TEXT("SheenRoughness")},
+		{TEXT("transmissionColor"), TEXT("TransmissionColor")},
+	};
+
+    UE::Interchange::FAxFMaterialObjectData& AxFData = AxFMaterialObjectNode->PayloadData;
+	for (const TPair<FString, FString>& Pair: ParameterNameMapping)
+    {
+		const FString& ParameterName = Pair.Key;
+        const FString & UEParameterName =  Pair.Value;
+        bool ParameterPresent = false;
+
+        if (const TArray<TSharedPtr<FJsonValue>> *TopLevelArray; JSON->TryGetArrayField(ParameterName, TopLevelArray))
+        {
+            ParameterPresent = true;
+
+            TArray<FLinearColor> ColorsFromJSON;
+            FLinearColor ValuesFromJSON(0.0f, 0.0f, 0.0f);
+			const int32 TopLevelArrayNum = FMath::Clamp(TopLevelArray->Num(), 0, 4);
+
+            for (size_t Index = 0; Index < TopLevelArrayNum; ++Index)
+            {
+				TSharedPtr<FJsonValue> JsonValue = (*TopLevelArray)[Index];
+                if (JsonValue->Type != EJson::Number && JsonValue->Type != EJson::Array)
+                {
+                    UE_LOG(LogAxFTranslator, Error, TEXT("Invalid value type in array found"));
+                    return {};
+                }
+
+                if (JsonValue->Type == EJson::Number)
+                {
+                    ValuesFromJSON.Component(Index) = JsonValue->AsNumber();
+                }
+            }
+
+            if (TopLevelArray->Num() == 1)
+            {
+                ValuesFromJSON = FLinearColor{ValuesFromJSON.R, ValuesFromJSON.R, ValuesFromJSON.R};
+            }
+
+            AxFData.ValuesMap.Add(UEParameterName) = ValuesFromJSON;
+        }
+        else if (const TSharedPtr<FJsonObject> *ValueObject; JSON->TryGetObjectField(ParameterName, ValueObject))
+        {
+            ParameterPresent = true;
+
+            FString TexturePath = ValueObject->Get()->GetStringField(TEXT("path"));
+            double WidthMM = ValueObject->Get()->GetNumberField(TEXT("widthMM"));
+            double HeightMM = ValueObject->Get()->GetNumberField(TEXT("heightMM"));
+
+            FString TextureUID = TEXT("T_") + DisplayLabel + TEXT("_") + UEParameterName;
+            UInterchangeTexture2DNode *Texture2D = NewObject<UInterchangeTexture2DNode>(&BaseNodeContainer, FName(TextureUID));
+            Texture2D->InitializeNode(TextureUID, TextureUID, EInterchangeNodeContainerType::TranslatedAsset);
+            Texture2D->SetPayLoadKey(Path + "|" + Name + "." + Extension + "|" + TexturePath);
+            Texture2D->AddDoubleAttribute(TEXT("WidthMM"), WidthMM);
+            Texture2D->AddDoubleAttribute(TEXT("HeightMM"), HeightMM);
+            BaseNodeContainer.AddNode(Texture2D);
+
+            if (UEParameterName == TEXT("CPA2ColorTable"))
+            {
+                FString SpecProfileUID = TEXT("SP_") + DisplayLabel;
+                UInterchangeSpecularProfileFactoryNode* SpecularProfileNode = NewObject<UInterchangeSpecularProfileFactoryNode>(&BaseNodeContainer, FName(SpecProfileUID));
+                SpecularProfileNode->InitializeNode(SpecProfileUID, SpecProfileUID, EInterchangeNodeContainerType::TranslatedAsset);
+                SpecularProfileNode->SetCustomFormat(ESpecularProfileFormat::HalfVector);
+                SpecularProfileNode->SetCustomTexture(TextureUID);
+                SpecularProfileNode->SetCustomSubPath(AxFMaterialObjectNode->GetDisplayLabel());
+                BaseNodeContainer.AddNode(SpecularProfileNode);
+            }
+        }
+        else if (bool ValueBool; JSON->TryGetBoolField(ParameterName, ValueBool)) 
+        {
+            ParameterPresent = ValueBool;
+        }
+
+        // If this parameter is present as either a value [array]
+        // or a texture, check for optional AxF features to be used.
+        if (ParameterPresent)
+        {
+            // Clear Coat
+            if (UEParameterName == TEXT("ClearCoatTransmissionColor") || UEParameterName == TEXT("ClearCoatRoughness") ||
+				UEParameterName == TEXT("ClearCoatIOR") || UEParameterName == TEXT("ClearCoatNormal"))
+            {
+                AxFData.UsedFeatures.AddUnique(EAxFFeature::ClearCoat);
+            }
+
+            // Sheen
+            if (UEParameterName == TEXT("SheenColor") || UEParameterName == TEXT("SheenRoughness"))
+            {
+                AxFData.UsedFeatures.AddUnique(EAxFFeature::Sheen);
+            }
+
+            // Anisotropy
+            if (UEParameterName == TEXT("Roughness") && AxFData.ValuesMap.Contains(UEParameterName))
+            {
+                if (const FLinearColor& RoughnessValues = AxFData.ValuesMap[UEParameterName]; RoughnessValues.R != RoughnessValues.G)
+                {
+                    AxFData.UsedFeatures.AddUnique(EAxFFeature::Anisotropy);
+                }
+            }
+            if (UEParameterName == TEXT("AnisotropicRotation"))
+            {
+                AxFData.UsedFeatures.AddUnique(EAxFFeature::Anisotropy);
+            }
+
+            // CPA2
+            if (UEParameterName == TEXT("CPA2ColorTable"))
+            {
+                AxFData.UsedFeatures.AddUnique(EAxFFeature::CPA2);
+            }
+
+            // Alpha transparency
+            if (UEParameterName == TEXT("Alpha"))
+            {
+                AxFData.UsedFeatures.AddUnique(EAxFFeature::Alpha);
+            }
+
+            // Transmission
+            if (UEParameterName == TEXT("TransmissionColor"))
+            {
+                AxFData.UsedFeatures.AddUnique(EAxFFeature::Transmission);
+            }
+
+            // Heightmap
+            if (UEParameterName == TEXT("HeightMap") || UEParameterName == TEXT("HeightStrength") || UEParameterName == TEXT("HeightReferencePlane"))
+            {
+                AxFData.UsedFeatures.AddUnique(EAxFFeature::Heightmap);
+            }
+
+            if (UEParameterName == TEXT("GlintDensity") || UEParameterName == TEXT("GlintUVScale") ||
+				UEParameterName == TEXT("GlintRoughness") || UEParameterName == TEXT("GlintBrightness"))
+            {
+                AxFData.UsedFeatures.AddUnique(EAxFFeature::Flakes);
+            }
+        }
+    }
+
+    BaseNodeContainer.AddNode(AxFMaterialObjectNode);
+
+    return true;
+}
+
+TOptional<UE::Interchange::FImportImage>
+UAxFTranslator::GetTexturePayloadData(const FString &PayloadKey, TOptional<FString> &AlternateTexturePath) const
+{
+    FString BasePath, RightPart;
+    if (!PayloadKey.Split("|", &BasePath, &RightPart))
+    {
+        UE_LOG(LogAxFTranslator, Error, TEXT("Failed to properly split Interchange payload key!"));
+        return {};
+    }
+
+    FString JSONFileName, LocalTexturePath;
+    if (!RightPart.Split("|", &JSONFileName, &LocalTexturePath))
+    {
+        UE_LOG(LogAxFTranslator, Error, TEXT("Failed to properly split Interchange payload key!"));
+        return {};
+    }
+
+    // TODO: Abstract into separate function
+    FString JSONString;
+    FString JSONFilePath = FPaths::Combine(BasePath, JSONFileName);
+    if (!FFileHelper::LoadFileToString(JSONString, *JSONFilePath))
+    {
+        UE_LOG(LogAxFTranslator, Error, TEXT("Failed to read JSON file at %s"), *JSONFilePath);
+        return {};
+    }
+
+    TSharedPtr<FJsonObject> JSON;
+    TSharedRef<TJsonReader<>> Reader = TJsonReaderFactory<>::Create(JSONString);
+    if (!FJsonSerializer::Deserialize(Reader, JSON))
+    {
+        UE_LOG(LogAxFTranslator, Error, TEXT("Failed to deserialize JSON read from file at %s"), *JSONFilePath);
+        return {};
+    }
+
+    TSharedPtr<TArray<uint8>> ByteArray(new TArray<uint8>());
+    FString TexturePath = FPaths::Combine(BasePath, LocalTexturePath);
+    if (!FFileHelper::LoadFileToArray(*ByteArray, *TexturePath))
+    {
+        UE_LOG(LogAxFTranslator, Error, TEXT("Failed to load texture from file: %s"), *TexturePath);
+        return {};
+    }
+
+    FImage RawImage;
+    if (FImageUtils::DecompressImage(ByteArray->GetData(), ByteArray->Num(), RawImage))
+    {
+        int32 Width = RawImage.GetWidth();
+        int32 Height = RawImage.GetHeight();
+
+        ETextureSourceFormat SourceFormat;
+        if (RawImage.Format == ERawImageFormat::G8)
+        {
+            SourceFormat = TSF_G8;
+        }
+        else if (RawImage.Format == ERawImageFormat::BGRA8)
+        {
+            SourceFormat = TSF_BGRA8;
+        }
+        else if (RawImage.Format == ERawImageFormat::R32F)
+        {
+            SourceFormat = TSF_R32F;
+        }
+        else if (RawImage.Format == ERawImageFormat::RGBA32F)
+        {
+            SourceFormat = TSF_RGBA32F;
+        }
+        else
+        {
+            UE_LOG(LogAxFTranslator, Warning, TEXT("Unsupported number of channels!"));
+            return {};
+        }
+
+		UE::Interchange::FImportImage Image;
+        Image.Init2DWithParams(Width, Height, SourceFormat, false);
+        FMemory::Memcpy(Image.RawData.GetData(), RawImage.RawData.GetData(), RawImage.RawData.NumBytes());
+
+		return Image;
+    }
+    else
+    {
+        UE_LOG(LogAxFTranslator, Error, TEXT("Texture asset could not be created: %s"), *TexturePath);
+		return {};
+    }
+}
+
+bool UAxFTranslator::SupportCompressedTexturePayloadData() const
+{
+    return IInterchangeTexturePayloadInterface::SupportCompressedTexturePayloadData();
+}
+
+TOptional<UE::Interchange::FImportImage>
+UAxFTranslator::GetCompressedTexturePayloadData(const FString &PayloadKey,
+                                                TOptional<FString> &AlternateTexturePath) const
+{
+    return IInterchangeTexturePayloadInterface::GetCompressedTexturePayloadData(PayloadKey, AlternateTexturePath);
+}
\ No newline at end of file
diff --git a/Engine/Plugins/Experimental/InterchangeAxF/Runtime/Source/InterchangeAxF/Private/InterchangeAxF.cpp b/Engine/Plugins/Experimental/InterchangeAxF/Runtime/Source/InterchangeAxF/Private/InterchangeAxF.cpp
new file mode 100644
--- /dev/null
+++ b/Engine/Plugins/Experimental/InterchangeAxF/Runtime/Source/InterchangeAxF/Private/InterchangeAxF.cpp
@@ -0,0 +1,82 @@
+// Copyright Epic Games, Inc. All Rights Reserved.
+
+#include "InterchangeAxF.h"
+#include "AxFTranslator.h"
+#include "Engine/Engine.h"
+#include "Engine/RendererSettings.h"
+#include "InterchangeManager.h"
+#include "InterchangeProjectSettings.h"
+#include "Misc/CoreDelegates.h"
+
+#include <Modules/ModuleManager.h>
+
+#define LOCTEXT_NAMESPACE "FInterchangeAxFModule"
+
+DEFINE_LOG_CATEGORY(LogInterchangeAxF);
+
+// This code will execute after your module is loaded into memory; the exact timing is specified in the .uplugin file per-module
+void FInterchangeAxFModule::StartupModule()
+{
+    bool IsSubstrateEnabled = GetDefault<URendererSettings>() ? GetDefault<URendererSettings>()->bEnableSubstrate : false;
+    if (!IsSubstrateEnabled)
+    {
+        UE_LOG(LogInterchangeAxF, Warning, TEXT("Interchange AxF requires Substrate to be enabled. Enable \"Substrate materials\" in Project Settings under Engine > Rendering."));
+    }
+
+    auto RegisterItems = []()
+    {
+        UInterchangeManager &InterchangeManager = UInterchangeManager::GetInterchangeManager();
+
+        // Register the translators
+        InterchangeManager.RegisterTranslator(UAxFTranslator::StaticClass());
+
+		// Add AxF translator' to the Interchange engine's pipeline set
+		if (UInterchangeProjectSettings* ProjectSettings = GetMutableDefault<UInterchangeProjectSettings>())
+		{
+			bool bPipelinesWereAdded = false;
+			FInterchangeContentImportSettings& AssetImportSettings = ProjectSettings->ContentImportSettings;
+
+			if (AssetImportSettings.DefaultPipelineStackOverride.Contains(EInterchangeTranslatorAssetType::Materials))
+			{
+				const FName DefaultMaterialsStack = AssetImportSettings.DefaultPipelineStackOverride[EInterchangeTranslatorAssetType::Materials];
+
+				if (AssetImportSettings.PipelineStacks.Contains(DefaultMaterialsStack))
+				{
+					FInterchangeTranslatorPipelines TranslatorPipeplines;
+					TranslatorPipeplines.Translator = TSoftClassPtr<UInterchangeTranslatorBase>(UAxFTranslator::StaticClass());
+					TranslatorPipeplines.Pipelines.Add(FSoftObjectPath(TEXT("/InterchangeAxF/AxFInterchangePipeline.AxFInterchangePipeline")));
+					{
+						FInterchangePipelineStack& PipelineStack = AssetImportSettings.PipelineStacks[DefaultMaterialsStack];
+						PipelineStack.PerTranslatorPipelines.Add(TranslatorPipeplines);
+					}
+
+					bPipelinesWereAdded = true;
+					UE_LOG(LogInterchangeAxF, Log, TEXT("Interchange AxF's pipelines added to default Material pipeline stack, %s"), *DefaultMaterialsStack.ToString());
+				}
+			}
+
+			if(!bPipelinesWereAdded)
+			{
+				UE_LOG(LogInterchangeAxF, Warning, TEXT("Cannot find Interchange default Materials' stack to add pipelines. Interchange AxF material import may not work properly."));
+			}
+		}
+	};
+
+    if (GEngine)
+    {
+        RegisterItems();
+    }
+    else
+    {
+        FCoreDelegates::OnPostEngineInit.AddLambda(RegisterItems);
+    }
+}
+
+// This function may be called during shutdown to clean up your module.  For modules that support dynamic reloading,
+// we call this function before unloading the module.
+void FInterchangeAxFModule::ShutdownModule()
+{
+}
+
+IMPLEMENT_MODULE(FInterchangeAxFModule, InterchangeAxF)
+#undef LOCTEXT_NAMESPACE
\ No newline at end of file
diff --git a/Engine/Plugins/Experimental/InterchangeAxF/Runtime/Source/InterchangeAxF/Public/AxFInterchangePipeline.h b/Engine/Plugins/Experimental/InterchangeAxF/Runtime/Source/InterchangeAxF/Public/AxFInterchangePipeline.h
new file mode 100644
--- /dev/null
+++ b/Engine/Plugins/Experimental/InterchangeAxF/Runtime/Source/InterchangeAxF/Public/AxFInterchangePipeline.h
@@ -0,0 +1,99 @@
+// Copyright Epic Games, Inc. All Rights Reserved.
+
+#pragma once
+
+#include "CoreMinimal.h"
+#include "InterchangeGenericTexturePipeline.h"
+#include "InterchangePipelineBase.h"
+#include "InterchangeSpecularProfileFactoryNode.h"
+#include "AxFInterchangePipeline.generated.h"
+
+#define UE_API INTERCHANGEAXF_API
+
+DECLARE_LOG_CATEGORY_EXTERN(LogAxFInterchangePipeline, Log, All);
+
+class UInterchangeShaderGraphNode;
+class UAxFMaterialObjectNode;
+class UInterchangeBaseNodeContainer;
+class UInterchangeSourceData;
+class UInterchangeTexture2DFactoryNode;
+class UInterchangeGenericTexturePipeline;
+class USpecularProfile;
+
+USTRUCT()
+struct FAxFTextureMap
+{
+	GENERATED_BODY()
+
+	UPROPERTY()
+	TObjectPtr<UTexture2D> Texture;
+
+	FVector2d UV;
+};
+
+UCLASS(BlueprintType)
+class UE_API UAxFInterchangePipeline : public UInterchangePipelineBase
+{
+    GENERATED_BODY()
+
+public:
+    UAxFInterchangePipeline();
+
+    UPROPERTY(
+        EditAnywhere, BlueprintReadWrite, Category = "General",
+        meta = (StandAlonePipelineProperty = "True", PipelineInternalEditionData = "True"))
+    FString PipelineDisplayName;
+
+	/* Set the reimport strategy. */
+	UPROPERTY(
+		EditAnywhere, BlueprintReadWrite, Category = "Common",
+		meta = (AdjustPipelineAndRefreshDetailOnChange = "True"))
+	EReimportStrategyFlags ReimportStrategy = EReimportStrategyFlags::ApplyNoProperties;
+
+    UPROPERTY(EditAnywhere, BlueprintReadWrite, Category = "AxF")
+    bool bUseTriplanarMappingByDefault = false;
+
+protected:
+    virtual void ExecutePipeline(
+        UInterchangeBaseNodeContainer* BaseNodeContainer, TArray<UInterchangeSourceData*> const& SourceDatas,
+        FString const& ContentBasePath) override;
+
+    virtual void ExecutePostImportPipeline(
+        UInterchangeBaseNodeContainer const* BaseNodeContainer, FString const& NodeKey, UObject* CreatedAsset,
+        bool bIsAReimport) override;
+
+    virtual void ExecutePostFactoryPipeline(
+        UInterchangeBaseNodeContainer const* BaseNodeContainer, FString const& NodeKey, UObject* CreatedAsset,
+        bool bIsAReimport) override;
+
+    virtual void AdjustSettingsForContext(FInterchangePipelineContextParams const& ContextParams) override;
+
+    virtual void PreDialogCleanup(FName const PipelineStackName) override;
+
+    virtual bool IsSettingsAreValid(TOptional<FText>& OutInvalidReason) const override;
+
+    virtual void FilterPropertiesFromTranslatedData(UInterchangeBaseNodeContainer* InBaseNodeContainer) override;
+
+    virtual void GetSupportAssetClasses(TArray<UClass*>& PipelineSupportAssetClasses) const override;
+
+private:
+	UPROPERTY()
+	TArray<TObjectPtr<UAxFMaterialObjectNode>> AxFObjectNodes;
+
+    // Texture handling
+	UPROPERTY()
+    TArray<TObjectPtr<UInterchangeTexture2DFactoryNode>> Texture2DFactoryNodes;
+
+	UPROPERTY()
+    TMap<FString, FAxFTextureMap> TexturesMap;
+
+    // Specular profile handling
+	UPROPERTY()
+    TObjectPtr<USpecularProfile> SpecularProfile;
+
+    // Generic texture pipeline to be used to create the necessary texture assets the materials use
+	UPROPERTY()
+    TObjectPtr<UInterchangeGenericTexturePipeline> TexturePipeline;
+};
+
+#undef UE_API
diff --git a/Engine/Plugins/Experimental/InterchangeAxF/Runtime/Source/InterchangeAxF/Public/AxFMaterialObjectNode.h b/Engine/Plugins/Experimental/InterchangeAxF/Runtime/Source/InterchangeAxF/Public/AxFMaterialObjectNode.h
new file mode 100644
--- /dev/null
+++ b/Engine/Plugins/Experimental/InterchangeAxF/Runtime/Source/InterchangeAxF/Public/AxFMaterialObjectNode.h
@@ -0,0 +1,63 @@
+// Copyright Epic Games, Inc. All Rights Reserved.
+
+#pragma once
+
+#include "CoreMinimal.h"
+#include "Nodes/InterchangeBaseNode.h"
+#include "AxFMaterialObjectNode.generated.h"
+
+#define UE_API INTERCHANGEAXF_API
+
+UENUM()
+enum EAxFFeature : uint8
+{
+    Alpha,
+    Anisotropy,
+    ClearCoat,
+    CPA2,
+	Flakes,
+    Heightmap,
+    Sheen,
+    Transmission
+};
+
+namespace UE
+{
+    namespace Interchange
+    {
+        struct UE_API FAxFMaterialObjectData
+        {
+            virtual ~FAxFMaterialObjectData() = default;
+            FAxFMaterialObjectData() = default;
+            FAxFMaterialObjectData(FAxFMaterialObjectData &&) = default;
+            FAxFMaterialObjectData &operator=(FAxFMaterialObjectData &&) = default;
+            FAxFMaterialObjectData(const FAxFMaterialObjectData &) = delete;
+            FAxFMaterialObjectData &operator=(const FAxFMaterialObjectData &) = delete;
+
+            TMap<FString, FLinearColor> ValuesMap;
+            TArray<EAxFFeature> UsedFeatures;
+        };
+    } // namespace Interchange
+} // namespace UE
+
+UCLASS(BlueprintType)
+class UE_API UAxFMaterialObjectNode : public UInterchangeBaseNode
+{
+    GENERATED_BODY()
+
+    UAxFMaterialObjectNode();
+
+public:
+    UFUNCTION(BlueprintCallable, Category = "Interchange | Node | AxFObject")
+    void InitializeAxFMaterialObjectNode(const FString &UniqueID, const FString &DisplayLabel);
+
+    virtual FString GetTypeName() const override;
+
+    virtual const TOptional<FString> GetPayloadKey() const;
+
+    virtual void SetPayloadKey(const FString &PayloadKey);
+
+    UE::Interchange::FAxFMaterialObjectData PayloadData;
+};
+
+#undef UE_API
\ No newline at end of file
diff --git a/Engine/Plugins/Experimental/InterchangeAxF/Runtime/Source/InterchangeAxF/Public/AxFTranslator.h b/Engine/Plugins/Experimental/InterchangeAxF/Runtime/Source/InterchangeAxF/Public/AxFTranslator.h
new file mode 100644
--- /dev/null
+++ b/Engine/Plugins/Experimental/InterchangeAxF/Runtime/Source/InterchangeAxF/Public/AxFTranslator.h
@@ -0,0 +1,46 @@
+// Copyright Epic Games, Inc. All Rights Reserved.
+
+#pragma once
+
+#include "CoreMinimal.h"
+#include "InterchangeTranslatorBase.h"
+#include "Import/Public/Texture/InterchangeTexturePayloadInterface.h"
+
+#include "AxFTranslator.generated.h"
+
+#define UE_API INTERCHANGEAXF_API
+
+DECLARE_LOG_CATEGORY_EXTERN(LogAxFTranslator, Log, All);
+
+namespace UE::Interchange
+{
+    struct FAxFMaterialObjectData;
+}
+
+UCLASS(BlueprintType)
+class UE_API UAxFTranslator
+    : public UInterchangeTranslatorBase
+    , public IInterchangeTexturePayloadInterface
+{
+	GENERATED_BODY()
+public:
+
+	virtual bool CanImportSourceData(const UInterchangeSourceData* InSourceData) const override;
+
+	virtual TArray<FString> GetSupportedFormats() const override;
+
+	virtual EInterchangeTranslatorAssetType GetSupportedAssetTypes() const override
+	{
+	    return EInterchangeTranslatorAssetType::Materials;
+	}
+
+	virtual bool Translate(UInterchangeBaseNodeContainer &BaseNodeContainer) const override;
+
+	// IInterchangeTexturePayloadInterface
+        virtual TOptional<UE::Interchange::FImportImage> GetTexturePayloadData(const FString &PayloadKey, TOptional<FString> &AlternateTexturePath) const override;
+        virtual bool SupportCompressedTexturePayloadData() const override;
+        virtual TOptional<UE::Interchange::FImportImage> GetCompressedTexturePayloadData(const FString &PayloadKey, TOptional<FString> &AlternateTexturePath) const override;
+	// IInterchangeTexturePayloadInterface
+};
+
+#undef UE_API
\ No newline at end of file
diff --git a/Engine/Plugins/Experimental/InterchangeAxF/Runtime/Source/InterchangeAxF/Public/InterchangeAxF.h b/Engine/Plugins/Experimental/InterchangeAxF/Runtime/Source/InterchangeAxF/Public/InterchangeAxF.h
new file mode 100644
--- /dev/null
+++ b/Engine/Plugins/Experimental/InterchangeAxF/Runtime/Source/InterchangeAxF/Public/InterchangeAxF.h
@@ -0,0 +1,17 @@
+// Copyright Epic Games, Inc. All Rights Reserved.
+
+#pragma once
+
+#include "Logging/LogMacros.h"
+#include "Modules/ModuleManager.h"
+
+DECLARE_LOG_CATEGORY_EXTERN(LogInterchangeAxF, Log, All);
+
+class FInterchangeAxFModule : public IModuleInterface
+{
+public:
+
+	/** IModuleInterface implementation */
+	virtual void StartupModule() override;
+	virtual void ShutdownModule() override;
+};
\ No newline at end of file
