Interchange: Use a set instead of repeatedly linear searching when combining material slots.

#Synced-CL 48758811

diff --git a/Engine/Plugins/Interchange/Runtime/Source/FactoryNodes/Private/InterchangeMeshFactoryNode.cpp b/Engine/Plugins/Interchange/Runtime/Source/FactoryNodes/Private/InterchangeMeshFactoryNode.cpp
--- a/Engine/Plugins/Interchange/Runtime/Source/FactoryNodes/Private/InterchangeMeshFactoryNode.cpp
+++ b/Engine/Plugins/Interchange/Runtime/Source/FactoryNodes/Private/InterchangeMeshFactoryNode.cpp
@@ -287,12 +287,12 @@ void UInterchangeMeshFactoryNode::GetSlotMaterialDependencies(TMap<FString, FStr
 	else
 	{
 		TMap<FString, FString> MaterialDependencies = SlotMaterialDependencies.ToMap();
-		TArray<FString> MaterialUids;
+		TSet<FString> MaterialUids;
 		MaterialUids.Reserve(MaterialDependencies.Num());
 		OutMaterialDependencies.Empty(MaterialDependencies.Num());
 		for (TPair<FString, FString>& MaterialDependency : MaterialDependencies)
 		{
-			if (MaterialUids.Find(MaterialDependency.Value) == INDEX_NONE)
+			if (!MaterialUids.Contains(MaterialDependency.Value))
 			{
 				OutMaterialDependencies.Add(MaterialDependency);
 				MaterialUids.Add(MaterialDependency.Value);
