[Interchange] ufbx - Camera, Lights

- additionally, support for camera orthographic projection

#Synced-CL 48545231

diff --git a/Engine/Plugins/Interchange/Runtime/Source/Parsers/Fbx/Private/Ufbx/UfbxParser.cpp b/Engine/Plugins/Interchange/Runtime/Source/Parsers/Fbx/Private/Ufbx/UfbxParser.cpp
--- a/Engine/Plugins/Interchange/Runtime/Source/Parsers/Fbx/Private/Ufbx/UfbxParser.cpp
+++ b/Engine/Plugins/Interchange/Runtime/Source/Parsers/Fbx/Private/Ufbx/UfbxParser.cpp
@@ -18,6 +18,7 @@
 #include "Fbx/InterchangeFbxMessages.h"
 #include "InterchangeHelper.h"
 #include "InterchangeCommonAnimationPayload.h"
+#include "InterchangeLightNode.h"
 
 #include "MeshDescription.h"
 #include "Misc/Paths.h"
@@ -40,6 +41,7 @@ FUfbxParser::FUfbxParser(TWeakObjectPtr<UInterchangeResultsContainer> InResultsC
 	, BoneNames(*this)
 	, MeshNames(*this)
 	, CameraNames(*this)
+	, LightNames(*this)
 	, MaterialNames(*this)
 	{
 		NodeNames.DefaultName = TEXT("Node");
@@ -53,6 +55,8 @@ FUfbxParser::FUfbxParser(TWeakObjectPtr<UInterchangeResultsContainer> InResultsC
 
 		CameraNames.DefaultName = TEXT("Camera");
 
+		LightNames.DefaultName =  TEXT("Light");
+
 		MaterialNames.DefaultName = TEXT("Material");
 
 	}
@@ -300,6 +304,7 @@ void FUfbxParser::Reset()
 	BoneNames.Reset();
 	MeshNames.Reset();
 	CameraNames.Reset();
+	LightNames.Reset();
 	MaterialNames.Reset();
 
 	ElementIdToSceneNode.Reset();
@@ -420,6 +425,21 @@ FString FUfbxParser::GetCameraName(const ufbx_camera& Camera) const
 		});
 }
 
+FString FUfbxParser::GetLightName(const ufbx_light& Light) const
+{
+	return LightNames.GetOrMake(Light.element, [&]
+		{
+			FString LightName = GetElementNameRaw(Light.element);
+
+			if (LightName.IsEmpty())
+			{
+				FString DefaultNamePrefix(UInterchangeBaseLightNode::StaticAssetTypeName());
+				LightName = DefaultNamePrefix + TEXT("_") + FString::FromInt(Light.typed_id);
+			}
+			return LightName;
+		});
+}
+
 void FUfbxParser::ConvertProperty(const ufbx_prop& Prop, UInterchangeBaseNode* SceneNode, const TOptional<FString>& PayloadKey)
 {
 	if (Prop.flags & UFBX_PROP_FLAG_USER_DEFINED)
diff --git a/Engine/Plugins/Interchange/Runtime/Source/Parsers/Fbx/Private/Ufbx/UfbxParser.h b/Engine/Plugins/Interchange/Runtime/Source/Parsers/Fbx/Private/Ufbx/UfbxParser.h
--- a/Engine/Plugins/Interchange/Runtime/Source/Parsers/Fbx/Private/Ufbx/UfbxParser.h
+++ b/Engine/Plugins/Interchange/Runtime/Source/Parsers/Fbx/Private/Ufbx/UfbxParser.h
@@ -187,6 +187,8 @@ namespace UE::Interchange::Private
 		}
 
 		FString GetCameraName(const ufbx_camera& Camera) const;
+		FString GetLightName(const ufbx_light& Light) const;
+
 
 		static void ConvertProperty(const ufbx_prop& Prop, UInterchangeBaseNode* SceneNode, const TOptional<FString>& PayloadKey = TOptional<FString>());
 
@@ -443,6 +445,7 @@ namespace UE::Interchange::Private
 		FUniqueNames BoneNames;
 		FUniqueNames MeshNames;
 		FUniqueNames CameraNames;
+		FUniqueNames LightNames;
 		FUniqueNames MaterialNames;
 
 		// Cache of node Uids
diff --git a/Engine/Plugins/Interchange/Runtime/Source/Parsers/Fbx/Private/Ufbx/UfbxScene.cpp b/Engine/Plugins/Interchange/Runtime/Source/Parsers/Fbx/Private/Ufbx/UfbxScene.cpp
--- a/Engine/Plugins/Interchange/Runtime/Source/Parsers/Fbx/Private/Ufbx/UfbxScene.cpp
+++ b/Engine/Plugins/Interchange/Runtime/Source/Parsers/Fbx/Private/Ufbx/UfbxScene.cpp
@@ -10,6 +10,7 @@
 
 #include "InterchangeMeshNode.h"
 #include "InterchangeCameraNode.h"
+#include "InterchangeLightNode.h"
 #include "Math/UnitConversion.h"
 #include "Nodes/InterchangeBaseNodeContainer.h"
 #include "Nodes/InterchangeUserDefinedAttribute.h"
@@ -302,6 +303,150 @@ namespace UE::Interchange::Private
 			}
 		}
 
+		if (InNode.camera)
+		{
+			const ufbx_camera& Camera = *InNode.camera;
+
+			FString CameraName = Parser.GetCameraName(Camera);
+			FString CameraLabel = CameraName;
+
+			switch (Camera.projection_mode)
+			{
+			case UFBX_PROJECTION_MODE_PERSPECTIVE:
+			{
+				FString CameraUid = FString(TEXT("\\")) + UInterchangePhysicalCameraNode::StaticAssetTypeName() + TEXT("\\") + CameraName;
+
+				const UInterchangePhysicalCameraNode* CameraNode = Cast<const UInterchangePhysicalCameraNode>(NodeContainer.GetNode(CameraUid));
+				if (!CameraNode)
+				{
+					UInterchangePhysicalCameraNode* NewCameraNode = NewObject<UInterchangePhysicalCameraNode>(&NodeContainer, NAME_None);
+					if (!ensure(NewCameraNode))
+					{
+						UInterchangeResultError_Generic* Message = Parser.AddMessage<UInterchangeResultError_Generic>();
+						Message->Text = LOCTEXT("CannotAllocateNode", "Cannot allocate a node when importing FBX.");
+						return;
+					}
+
+					NodeContainer.SetupNode(NewCameraNode, CameraUid, CameraName, EInterchangeNodeContainerType::TranslatedAsset);
+
+					// ufbx computes fov tan for every aperture mode for we can recalculate focal length
+					float ApertureSize = FUnitConversion::Convert(Camera.aperture_size_inch.x, EUnit::Inches, EUnit::Millimeters);
+					float FocalLength = 0.5f * ApertureSize / FMath::Max(Camera.field_of_view_tan.x, UE_SMALL_NUMBER);
+
+					NewCameraNode->SetCustomFocalLength(FocalLength); //Both FBX and UE have their focal length in mm
+					NewCameraNode->SetCustomSensorHeight(FUnitConversion::Convert(Camera.aperture_size_inch.y, EUnit::Inches, EUnit::Millimeters));
+					NewCameraNode->SetCustomSensorWidth(FUnitConversion::Convert(Camera.aperture_size_inch.x, EUnit::Inches, EUnit::Millimeters));
+
+					CameraNode = NewCameraNode;
+				}
+
+				SceneNode->SetCustomAssetInstanceUid(CameraUid);
+			}
+				break;
+			case UFBX_PROJECTION_MODE_ORTHOGRAPHIC:
+			{
+				FString CameraUid = FString(TEXT("\\")) + UInterchangeStandardCameraNode::StaticAssetTypeName() + TEXT("\\") + CameraName;
+
+				const UInterchangeStandardCameraNode* CameraNode = Cast<const UInterchangeStandardCameraNode>(NodeContainer.GetNode(CameraUid));
+				if (!CameraNode)
+				{
+					UInterchangeStandardCameraNode* NewCameraNode = NewObject<UInterchangeStandardCameraNode>(&NodeContainer, NAME_None);
+					if (!ensure(NewCameraNode))
+					{
+						UInterchangeResultError_Generic* Message = Parser.AddMessage<UInterchangeResultError_Generic>();
+						Message->Text = LOCTEXT("CannotAllocateNode", "Cannot allocate a node when importing FBX.");
+						return;
+					}
+
+					NodeContainer.SetupNode(NewCameraNode, CameraUid, CameraName, EInterchangeNodeContainerType::TranslatedAsset);
+
+					NewCameraNode->SetCustomProjectionMode(EInterchangeCameraProjectionType::Orthographic);
+
+					NewCameraNode->SetCustomWidth(Camera.orthographic_size.x);
+					NewCameraNode->SetCustomNearClipPlane(Camera.near_plane);
+					NewCameraNode->SetCustomFarClipPlane(Camera.far_plane);
+					NewCameraNode->SetCustomAspectRatio(Camera.aspect_ratio);
+
+					CameraNode = NewCameraNode;
+				}
+
+				SceneNode->SetCustomAssetInstanceUid(CameraUid);
+			}
+				break;
+			};
+		}
+
+		if (InNode.light)
+		{
+			const ufbx_light& Light = *InNode.light;
+
+			FString LightName = Parser.GetLightName(Light);
+			FString LightUid = FString(TEXT("\\")) + UInterchangeBaseLightNode::StaticAssetTypeName() + TEXT("\\") + LightName;
+			const UInterchangeBaseLightNode* LightNode = Cast<const UInterchangeBaseLightNode>(NodeContainer.GetNode(LightUid));
+
+			if (!LightNode)
+			{
+				UClass* LightClass;
+
+				switch(Light.type)
+				{
+				case UFBX_LIGHT_POINT:
+				case UFBX_LIGHT_VOLUME:
+					LightClass = UInterchangePointLightNode::StaticClass();
+					break;
+				case UFBX_LIGHT_DIRECTIONAL:
+					LightClass = UInterchangeDirectionalLightNode::StaticClass();
+					break;
+				case UFBX_LIGHT_SPOT:
+					LightClass = UInterchangeSpotLightNode::StaticClass();
+					break;
+				case UFBX_LIGHT_AREA:
+					LightClass = UInterchangeRectLightNode::StaticClass();
+					break;
+				default:
+					LightClass = UInterchangePointLightNode::StaticClass();
+					break;
+				}
+
+				UInterchangeBaseLightNode* BaseLightNode = NewObject<UInterchangeBaseLightNode>(&NodeContainer, LightClass, NAME_None);
+				if (!ensure(BaseLightNode))
+				{
+					UInterchangeResultError_Generic* Message = Parser.AddMessage<UInterchangeResultError_Generic>();
+					Message->Text = LOCTEXT("CannotAllocateNode", "Cannot allocate a node when importing FBX.");
+					return;
+				}
+				else
+				{
+					FString LightLabel = LightName;
+				
+					NodeContainer.SetupNode(BaseLightNode, LightUid, LightLabel, EInterchangeNodeContainerType::TranslatedAsset);
+
+					BaseLightNode->SetCustomLightColor(Convert::ConvertColor(ufbx_vec4{Light.color.x, Light.color.y, Light.color.z, 1.0}));
+
+					// ufbx divides intensity by 100, re-multiplied here to match FBX SDK parser
+					BaseLightNode->SetCustomIntensity(static_cast<float>(Light.intensity) * 100.f);
+
+					if (UInterchangeLightNode* InterchangeLightNode = Cast<UInterchangeLightNode>(BaseLightNode))
+					{
+						InterchangeLightNode->SetCustomAttenuationRadius(FBXSDK_FLOAT_MAX);
+						// #ufbx_todo: doesn't seem to have attenuation radius, only decay type 
+						// ufbx_light_decay DecayType = Light.decay;
+					}
+
+					if (Light.type == UFBX_LIGHT_SPOT )
+					{
+						if (UInterchangeSpotLightNode* SpotLightNode = Cast<UInterchangeSpotLightNode>(BaseLightNode))
+						{
+							SpotLightNode->SetCustomInnerConeAngle(Light.inner_angle);
+							SpotLightNode->SetCustomOuterConeAngle(Light.outer_angle);
+						}
+					}
+
+				}
+			}
+			SceneNode->SetCustomAssetInstanceUid(LightUid);
+		}
+
 		if (InNode.attrib_type == UFBX_ELEMENT_LOD_GROUP)
 		{
 			SceneNode->AddSpecializedType(FSceneNodeStaticData::GetLodGroupSpecializeTypeString());
