Build fix with clang

#Synced-CL 48545231

diff --git a/Engine/Plugins/Interchange/Runtime/Source/Parsers/Fbx/Private/Ufbx/UfbxScene.cpp b/Engine/Plugins/Interchange/Runtime/Source/Parsers/Fbx/Private/Ufbx/UfbxScene.cpp
--- a/Engine/Plugins/Interchange/Runtime/Source/Parsers/Fbx/Private/Ufbx/UfbxScene.cpp
+++ b/Engine/Plugins/Interchange/Runtime/Source/Parsers/Fbx/Private/Ufbx/UfbxScene.cpp
@@ -421,7 +421,7 @@ namespace UE::Interchange::Private
 				
 					NodeContainer.SetupNode(BaseLightNode, LightUid, LightLabel, EInterchangeNodeContainerType::TranslatedAsset);
 
-					BaseLightNode->SetCustomLightColor(Convert::ConvertColor(ufbx_vec4{Light.color.x, Light.color.y, Light.color.z, 1.0}));
+					BaseLightNode->SetCustomLightColor(Convert::ConvertColor(ufbx_vec4{{Light.color.x, Light.color.y, Light.color.z, 1.0}}));
 
 					// ufbx divides intensity by 100, re-multiplied here to match FBX SDK parser
 					BaseLightNode->SetCustomIntensity(static_cast<float>(Light.intensity) * 100.f);
