Build fix with clang - another try

- braces for a struct - in a union - in a struct

#Synced-CL 48545231

diff --git a/Engine/Plugins/Interchange/Runtime/Source/Parsers/Fbx/Private/Ufbx/UfbxScene.cpp b/Engine/Plugins/Interchange/Runtime/Source/Parsers/Fbx/Private/Ufbx/UfbxScene.cpp
--- a/Engine/Plugins/Interchange/Runtime/Source/Parsers/Fbx/Private/Ufbx/UfbxScene.cpp
+++ b/Engine/Plugins/Interchange/Runtime/Source/Parsers/Fbx/Private/Ufbx/UfbxScene.cpp
@@ -421,8 +421,7 @@ namespace UE::Interchange::Private
 				
 					NodeContainer.SetupNode(BaseLightNode, LightUid, LightLabel, EInterchangeNodeContainerType::TranslatedAsset);
 
-					const ufbx_vec4 Color = {{Light.color.x, Light.color.y, Light.color.z, 1.0}};
-					BaseLightNode->SetCustomLightColor(Convert::ConvertColor(Color));
+					const ufbx_vec4 Color = {{{Light.color.x, Light.color.y, Light.color.z, 1.0}}};
 
 					// ufbx divides intensity by 100, re-multiplied here to match FBX SDK parser
 					BaseLightNode->SetCustomIntensity(static_cast<float>(Light.intensity) * 100.f);
