[On behalf of Stefan Ivanovski from XRite]
- Removed no longer used TextureCoordinateTransform.uasset
- Remove question marks from parameters of base materials and material functions
- Replace direct calls of GetStringField() with TryGetStringField()
- Assert texture paths to be under the base AxF2UE directory

#virtualized
#Synced-CL 42983915

diff --git a/Engine/Plugins/Experimental/InterchangeAxF/Runtime/InterchangeAxF.uplugin b/Engine/Plugins/Experimental/InterchangeAxF/Runtime/InterchangeAxF.uplugin
--- a/Engine/Plugins/Experimental/InterchangeAxF/Runtime/InterchangeAxF.uplugin
+++ b/Engine/Plugins/Experimental/InterchangeAxF/Runtime/InterchangeAxF.uplugin
@@ -27,11 +27,11 @@
 	],
 	"Plugins": [
 		{
-			"Name": "InterchangeAxFAssets",
+			"Name": "Interchange",
 			"Enabled": true
 		},
 		{
-			"Name": "Interchange",
+			"Name": "InterchangeAxFAssets",
 			"Enabled": true
 		}
 	]
diff --git a/Engine/Plugins/Experimental/InterchangeAxF/Runtime/Source/InterchangeAxF/Private/AxFInterchangePipeline.cpp b/Engine/Plugins/Experimental/InterchangeAxF/Runtime/Source/InterchangeAxF/Private/AxFInterchangePipeline.cpp
--- a/Engine/Plugins/Experimental/InterchangeAxF/Runtime/Source/InterchangeAxF/Private/AxFInterchangePipeline.cpp
+++ b/Engine/Plugins/Experimental/InterchangeAxF/Runtime/Source/InterchangeAxF/Private/AxFInterchangePipeline.cpp
@@ -130,7 +130,7 @@ void UAxFInterchangePipeline::ExecutePipeline(UInterchangeBaseNodeContainer* Bas
 		for (const EAxFFeature Feature : Features)
 		{
 			const FString Value = UEnum::GetValueAsString(Feature);
-			const FString AttributeKey = UInterchangeShaderPortsAPI::MakeInputValueKey(TEXT("Has") + Value + TEXT("?"));
+			const FString AttributeKey = UInterchangeShaderPortsAPI::MakeInputValueKey(TEXT("Has") + Value);
 
 			MaterialInstance->AddBooleanAttribute(AttributeKey, true);
 		}
@@ -152,7 +152,7 @@ void UAxFInterchangePipeline::ExecutePipeline(UInterchangeBaseNodeContainer* Bas
             }
 
 			// NOTE: Disable texture path in the case that we have a value for this parameter
-			FString TextureToggleAttributeKey = UInterchangeShaderPortsAPI::MakeInputValueKey(TEXT("Has") + Param + TEXT("Texture?"));
+			FString TextureToggleAttributeKey = UInterchangeShaderPortsAPI::MakeInputValueKey(TEXT("Has") + Param + TEXT("Texture"));
 			MaterialInstance->AddBooleanAttribute(TextureToggleAttributeKey, false);
 
 			const FString AttributeKey = UInterchangeShaderPortsAPI::MakeInputValueKey(Param);
@@ -180,7 +180,7 @@ void UAxFInterchangePipeline::ExecutePipeline(UInterchangeBaseNodeContainer* Bas
 			const FString Param = GetParameterNameFromUID(Texture2DFactoryNode->GetUniqueID());
 
 			{
-				FString AttributeKey = UInterchangeShaderPortsAPI::MakeInputValueKey(TEXT("Has") + Param + TEXT("Texture?"));
+				FString AttributeKey = UInterchangeShaderPortsAPI::MakeInputValueKey(TEXT("Has") + Param + TEXT("Texture"));
 				MaterialInstance->AddBooleanAttribute(AttributeKey, true);
 			}
 
diff --git a/Engine/Plugins/Experimental/InterchangeAxF/Runtime/Source/InterchangeAxF/Private/AxFTranslator.cpp b/Engine/Plugins/Experimental/InterchangeAxF/Runtime/Source/InterchangeAxF/Private/AxFTranslator.cpp
--- a/Engine/Plugins/Experimental/InterchangeAxF/Runtime/Source/InterchangeAxF/Private/AxFTranslator.cpp
+++ b/Engine/Plugins/Experimental/InterchangeAxF/Runtime/Source/InterchangeAxF/Private/AxFTranslator.cpp
@@ -156,9 +156,24 @@ bool UAxFTranslator::Translate(UInterchangeBaseNodeContainer &BaseNodeContainer)
         {
             ParameterPresent = true;
 
-            FString TexturePath = ValueObject->Get()->GetStringField(TEXT("path"));
-            double WidthMM = ValueObject->Get()->GetNumberField(TEXT("widthMM"));
-            double HeightMM = ValueObject->Get()->GetNumberField(TEXT("heightMM"));
+			FString TexturePath{};
+			if (!ValueObject->Get()->TryGetStringField(TEXT("path"), TexturePath))
+			{
+				UE_LOG(LogAxFTranslator, Error, TEXT("Failed to find '%s' property in AxF2UE file."), TEXT("path"));
+				return false;
+			}
+
+			double WidthMM{}, HeightMM{};
+			if (!ValueObject->Get()->TryGetNumberField(TEXT("widthMM"), WidthMM))
+			{
+				UE_LOG(LogAxFTranslator, Error, TEXT("Failed to find '%s' property in AxF2UE file."), TEXT("widthMM"));
+				return false;
+			}
+			if (!ValueObject->Get()->TryGetNumberField(TEXT("heightMM"), HeightMM))
+			{
+				UE_LOG(LogAxFTranslator, Error, TEXT("Failed to find '%s' property in AxF2UE file."), TEXT("heightMM"));
+				return false;
+			}
 
             FString TextureUID = TEXT("T_") + DisplayLabel + TEXT("_") + UEParameterName;
             UInterchangeTexture2DNode *Texture2D = NewObject<UInterchangeTexture2DNode>(&BaseNodeContainer, FName(TextureUID));
@@ -271,6 +286,11 @@ UAxFTranslator::GetTexturePayloadData(const FString &PayloadKey, TOptional<FStri
     // TODO: Abstract into separate function
     FString JSONString;
     FString JSONFilePath = FPaths::Combine(BasePath, JSONFileName);
+	if (!FPaths::IsUnderDirectory(JSONFilePath, BasePath))
+	{
+		UE_LOG(LogAxFTranslator, Error, TEXT("JSON file %s is not under the directory %s"), *JSONFileName, *BasePath);
+		return {};
+	}
     if (!FFileHelper::LoadFileToString(JSONString, *JSONFilePath))
     {
         UE_LOG(LogAxFTranslator, Error, TEXT("Failed to read JSON file at %s"), *JSONFilePath);
@@ -287,6 +307,11 @@ UAxFTranslator::GetTexturePayloadData(const FString &PayloadKey, TOptional<FStri
 
     TSharedPtr<TArray<uint8>> ByteArray(new TArray<uint8>());
     FString TexturePath = FPaths::Combine(BasePath, LocalTexturePath);
+	if (!FPaths::IsUnderDirectory(TexturePath, BasePath))
+	{
+		UE_LOG(LogAxFTranslator, Error, TEXT("Texture file file %s is not under the directory %s"), *TexturePath, *BasePath);
+		return {};
+	}
     if (!FFileHelper::LoadFileToArray(*ByteArray, *TexturePath))
     {
         UE_LOG(LogAxFTranslator, Error, TEXT("Failed to load texture from file: %s"), *TexturePath);
diff --git a/Engine/Plugins/Experimental/InterchangeAxF/Runtime/Source/InterchangeAxF/Private/InterchangeAxF.cpp b/Engine/Plugins/Experimental/InterchangeAxF/Runtime/Source/InterchangeAxF/Private/InterchangeAxF.cpp
--- a/Engine/Plugins/Experimental/InterchangeAxF/Runtime/Source/InterchangeAxF/Private/InterchangeAxF.cpp
+++ b/Engine/Plugins/Experimental/InterchangeAxF/Runtime/Source/InterchangeAxF/Private/InterchangeAxF.cpp
@@ -30,7 +30,7 @@ void FInterchangeAxFModule::StartupModule()
         // Register the translators
         InterchangeManager.RegisterTranslator(UAxFTranslator::StaticClass());
 
-		// Add AxF translator' to the Interchange engine's pipeline set
+        // Add X-Rite's AxF translator to the Interchange engine's pipeline set
 		if (UInterchangeProjectSettings* ProjectSettings = GetMutableDefault<UInterchangeProjectSettings>())
 		{
 			bool bPipelinesWereAdded = false;
