Interchange: Fix static mesh actors spawned for collision meshes sometimes ending up with the render mesh assigned

#Synced-CL 48758811

diff --git a/Engine/Plugins/Interchange/Runtime/Source/Pipelines/Private/InterchangeGenericScenesPipeline.cpp b/Engine/Plugins/Interchange/Runtime/Source/Pipelines/Private/InterchangeGenericScenesPipeline.cpp
--- a/Engine/Plugins/Interchange/Runtime/Source/Pipelines/Private/InterchangeGenericScenesPipeline.cpp
+++ b/Engine/Plugins/Interchange/Runtime/Source/Pipelines/Private/InterchangeGenericScenesPipeline.cpp
@@ -617,6 +617,28 @@ void UInterchangeGenericLevelPipeline::ExecuteSceneNodePreImport(const FTransfor
 				}
 			}
 		}
+
+		// Don't spawn actor factory nodes if our target asset hasn't produced a factory node itself for us to assign (as we likely don't want
+		// e.g. a static mesh component if we don't have a static mesh for it). Currently this is mainly used to prevent us from spawning actor
+		// factory nodes for mesh nodes that were ultimately consumed as collision meshes for other static meshes. We filter for meshes only
+		// to avoid false positives like for cameras/lights, which have translated asset nodes, but no asset factory nodes.
+		//
+		// Note that without this snippet it may seem like these actors don't get assigned any StaticMesh anyway, so is this really that bad an issue?
+		// Well, that is only due to the fact that UInterchangeGenericLevelPipeline::SetUpFactoryNode() can't specify the factory dependency correctly
+		// in the collision mesh case, as the produced static mesh factory node won't have the UID based on the collision mesh translated node, but
+		// instead the render mesh translated node. So randomly depending on your scene/setup, if the spawned actor gets finalized after the static mesh
+		// you'd get the render mesh assigned to the collision mesh spawned actor (really wrong), or if it gets finalized before the static mesh you'd
+		// get the actor without anything assigned (slightly wrong / useless)
+		if (TranslatedAssetNode &&
+			(TranslatedAssetNode->IsA<const UInterchangeMeshNode>() || TranslatedAssetNode->IsA<const UInterchangeMeshBundleNode>()))
+		{
+			FString AssetFactoryNodeUid = UInterchangeFactoryBaseNode::BuildFactoryNodeUid(TranslatedAssetNode->GetUniqueID());
+			UInterchangeFactoryBaseNode* AssetFactoryNode = BaseNodeContainer->GetFactoryNode(AssetFactoryNodeUid);
+			if (!AssetFactoryNode)
+			{
+				return;
+			}
+		}
 	}
 
 	const bool bIsLodGroup = [&SceneNode]()
