Interchange MaterialX:
- Add a custom nodedef for the displacementshader to support a "bias" input
- Add an attributedef for "scale" and "bias" inputs to choose between a graph or a value in the Details Panel of the displacement
- Add a setting in the MaterialX pipeline to choose if normals should be recomputed after a displacement
- Added "bias" and "bRecomputeNormals" input to the Material Function MX_Displacement
- Set the Displacement Magnitude in the details panel

#virtualized
#Synced-CL 49028617

diff --git a/Engine/Plugins/Interchange/Runtime/Source/FactoryNodes/Private/InterchangeMaterialFactoryNode.cpp b/Engine/Plugins/Interchange/Runtime/Source/FactoryNodes/Private/InterchangeMaterialFactoryNode.cpp
--- a/Engine/Plugins/Interchange/Runtime/Source/FactoryNodes/Private/InterchangeMaterialFactoryNode.cpp
+++ b/Engine/Plugins/Interchange/Runtime/Source/FactoryNodes/Private/InterchangeMaterialFactoryNode.cpp
@@ -425,6 +425,16 @@ bool UInterchangeMaterialFactoryNode::SetCustomDisplacementCenter(float Attribut
 	IMPLEMENT_NODE_ATTRIBUTE_SETTER_NODELEGATE(DisplacementCenter, float);
 }
 
+bool UInterchangeMaterialFactoryNode::GetCustomDisplacementMagnitude(float& AttributeValue) const
+{
+	IMPLEMENT_NODE_ATTRIBUTE_GETTER(DisplacementMagnitude, float);
+}
+
+bool UInterchangeMaterialFactoryNode::SetCustomDisplacementMagnitude(float AttributeValue)
+{
+	IMPLEMENT_NODE_ATTRIBUTE_SETTER_NODELEGATE(DisplacementMagnitude, float);
+}
+
 FString UInterchangeMaterialExpressionFactoryNode::GetTypeName() const
 {
 	const FString TypeName = TEXT("MaterialExpressionFactoryNode");
diff --git a/Engine/Plugins/Interchange/Runtime/Source/FactoryNodes/Public/InterchangeMaterialFactoryNode.h b/Engine/Plugins/Interchange/Runtime/Source/FactoryNodes/Public/InterchangeMaterialFactoryNode.h
--- a/Engine/Plugins/Interchange/Runtime/Source/FactoryNodes/Public/InterchangeMaterialFactoryNode.h
+++ b/Engine/Plugins/Interchange/Runtime/Source/FactoryNodes/Public/InterchangeMaterialFactoryNode.h
@@ -275,6 +275,12 @@ public:
 	UFUNCTION(BlueprintCallable, Category = "Interchange | Node | Material")
 	INTERCHANGEFACTORYNODES_API bool SetCustomDisplacementCenter(float AttributeValue);
 
+	UFUNCTION(BlueprintCallable, Category = "Interchange | Node | Material")
+	INTERCHANGEFACTORYNODES_API bool GetCustomDisplacementMagnitude(float& AttributeValue) const;
+
+	UFUNCTION(BlueprintCallable, Category = "Interchange | Node | Material")
+	INTERCHANGEFACTORYNODES_API bool SetCustomDisplacementMagnitude(float AttributeValue);
+
 private:
 	IMPLEMENT_NODE_ATTRIBUTE_KEY(ShadingModel)
 	IMPLEMENT_NODE_ATTRIBUTE_KEY(TranslucencyLightingMode)
@@ -284,6 +290,7 @@ private:
 	IMPLEMENT_NODE_ATTRIBUTE_KEY(RefractionMethod)
 	IMPLEMENT_NODE_ATTRIBUTE_KEY(ScreenSpaceReflections)
 	IMPLEMENT_NODE_ATTRIBUTE_KEY(DisplacementCenter)
+	IMPLEMENT_NODE_ATTRIBUTE_KEY(DisplacementMagnitude)
 };
 
 UCLASS(MinimalAPI, BlueprintType)
diff --git a/Engine/Plugins/Interchange/Runtime/Source/Import/Private/MaterialX/MaterialXUtils/MaterialXDisplacementShader.cpp b/Engine/Plugins/Interchange/Runtime/Source/Import/Private/MaterialX/MaterialXUtils/MaterialXDisplacementShader.cpp
--- a/Engine/Plugins/Interchange/Runtime/Source/Import/Private/MaterialX/MaterialXUtils/MaterialXDisplacementShader.cpp
+++ b/Engine/Plugins/Interchange/Runtime/Source/Import/Private/MaterialX/MaterialXUtils/MaterialXDisplacementShader.cpp
@@ -22,13 +22,34 @@ TSharedRef<FMaterialXBase> FMaterialXDisplacementShader::MakeInstance(UInterchan
 
 UInterchangeBaseNode* FMaterialXDisplacementShader::Translate(MaterialX::NodePtr DisplacementNode)
 {
+	using namespace UE::Interchange::Materials;
 	this->SurfaceShaderNode = DisplacementNode;
 
-	UInterchangeShaderNode* DisplacementShaderNode = FMaterialXSurfaceShaderAbstract::Translate(EInterchangeMaterialXShaders::Displacement);
-	// by default the center is at 0.5, but in order to compute correctly the normals from the displacement we need it at 0 in MX_Displacement material function
-	ShaderGraphNode->SetCustomDisplacementCenterMode(0.f);
+	float DisplacementCenter = 0.f;
+
+	// if we set the bias and scale in the details panel then we should remove it from the shader graph
+	if (mx::InputPtr Input = DisplacementNode->getInput("scale"))
+	{
+		if (Input->hasValue() && Input->getTypedAttribute<bool>("ue_displacement_attribute"))
+		{
+			ShaderGraphNode->SetCustomDisplacementMagnitude(Input->getValue()->asA<float>());
+			DisplacementNode->removeInput(Input->getName());
+		}
+	}
+
+	if (mx::InputPtr Input = DisplacementNode->getInput("bias"))
+	{
+		if (Input->hasValue() && Input->getTypedAttribute<bool>("ue_displacement_attribute"))
+		{
+			DisplacementCenter = Input->getValue()->asA<float>();
+			DisplacementNode->removeInput(Input->getName());
+		}
+	}
 
-	using namespace UE::Interchange::Materials;
+	UInterchangeShaderNode* DisplacementShaderNode = FMaterialXSurfaceShaderAbstract::Translate(EInterchangeMaterialXShaders::Displacement);
+	// by default the center is at 0.5, but in order to compute correctly the normals from the displacement we need it at 0 in MX_Displacement
+	// only if we're creating a shadergraph for the input and we're recomputing the normals (see MaterialXPipeline)
+	ShaderGraphNode->SetCustomDisplacementCenterMode(DisplacementCenter);
 
 	// Outputs
 	UInterchangeShaderPortsAPI::ConnectOuputToInputByName(ShaderGraphNode, Common::Parameters::Displacement.ToString(), DisplacementShaderNode->GetUniqueID(), Common::Parameters::Displacement.ToString());
diff --git a/Engine/Plugins/Interchange/Runtime/Source/Nodes/Private/InterchangeShaderGraphNode.cpp b/Engine/Plugins/Interchange/Runtime/Source/Nodes/Private/InterchangeShaderGraphNode.cpp
--- a/Engine/Plugins/Interchange/Runtime/Source/Nodes/Private/InterchangeShaderGraphNode.cpp
+++ b/Engine/Plugins/Interchange/Runtime/Source/Nodes/Private/InterchangeShaderGraphNode.cpp
@@ -373,6 +373,16 @@ bool UInterchangeShaderGraphNode::SetCustomDisplacementCenterMode(float Attribut
 	IMPLEMENT_NODE_ATTRIBUTE_SETTER_NODELEGATE(DisplacementCenter, float);
 }
 
+bool UInterchangeShaderGraphNode::GetCustomDisplacementMagnitude(float& AttributeValue) const
+{
+	IMPLEMENT_NODE_ATTRIBUTE_GETTER(DisplacementMagnitude, float);
+}
+
+bool UInterchangeShaderGraphNode::SetCustomDisplacementMagnitude(float AttributeValue)
+{
+	IMPLEMENT_NODE_ATTRIBUTE_SETTER_NODELEGATE(DisplacementMagnitude, float);
+}
+
 FString UInterchangeFunctionCallShaderNode::GetTypeName() const
 {
 	const FString TypeName = TEXT("FunctionCallShaderNode");
diff --git a/Engine/Plugins/Interchange/Runtime/Source/Nodes/Public/InterchangeShaderGraphNode.h b/Engine/Plugins/Interchange/Runtime/Source/Nodes/Public/InterchangeShaderGraphNode.h
--- a/Engine/Plugins/Interchange/Runtime/Source/Nodes/Public/InterchangeShaderGraphNode.h
+++ b/Engine/Plugins/Interchange/Runtime/Source/Nodes/Public/InterchangeShaderGraphNode.h
@@ -331,6 +331,15 @@ public:
 	UFUNCTION(BlueprintCallable, Category = "Interchange | Node | Material")
 	UE_API bool SetCustomDisplacementCenterMode(float AttributeValue);
 
+	/**
+	 * Set the magnitude of the displacement
+	 */
+	UFUNCTION(BlueprintCallable, Category = "Interchange | Node | Material")
+	UE_API bool GetCustomDisplacementMagnitude(float& AttributeValue) const;
+
+	UFUNCTION(BlueprintCallable, Category = "Interchange | Node | Material")
+	UE_API bool SetCustomDisplacementMagnitude(float AttributeValue);
+
 private:
 	IMPLEMENT_NODE_ATTRIBUTE_KEY(TwoSided)
 	IMPLEMENT_NODE_ATTRIBUTE_KEY(TwoSidedTransmission)
@@ -338,7 +347,8 @@ private:
 	IMPLEMENT_NODE_ATTRIBUTE_KEY(IsAShaderFunction)
 	IMPLEMENT_NODE_ATTRIBUTE_KEY(ScreenSpaceReflections)
 	IMPLEMENT_NODE_ATTRIBUTE_KEY(BlendMode)
-	IMPLEMENT_NODE_ATTRIBUTE_KEY(DisplacementCenter);
+	IMPLEMENT_NODE_ATTRIBUTE_KEY(DisplacementCenter)
+	IMPLEMENT_NODE_ATTRIBUTE_KEY(DisplacementMagnitude)
 };
 
 #undef UE_API
diff --git a/Engine/Plugins/Interchange/Runtime/Source/Pipelines/Private/InterchangeGenericMaterialPipeline.cpp b/Engine/Plugins/Interchange/Runtime/Source/Pipelines/Private/InterchangeGenericMaterialPipeline.cpp
--- a/Engine/Plugins/Interchange/Runtime/Source/Pipelines/Private/InterchangeGenericMaterialPipeline.cpp
+++ b/Engine/Plugins/Interchange/Runtime/Source/Pipelines/Private/InterchangeGenericMaterialPipeline.cpp
@@ -2022,6 +2022,11 @@ void UInterchangeGenericMaterialPipeline::HandleCommonParameters(const UIntercha
 			{
 				MaterialFactoryNode->SetCustomDisplacementCenter(DisplacementCenter);
 			}
+
+			if (float DisplacementMagnitude; ShaderGraphNode->GetCustomDisplacementMagnitude(DisplacementMagnitude))
+			{
+				MaterialFactoryNode->SetCustomDisplacementMagnitude(DisplacementMagnitude);
+			}
 		}
 	}
 }
@@ -4059,6 +4064,11 @@ bool UInterchangeGenericMaterialPipeline::HandleSubstrate(const UInterchangeShad
 			{
 				MaterialFactoryNode->SetCustomDisplacementCenter(DisplacementCenter);
 			}
+
+			if (float DisplacementMagnitude; ShaderGraphNode->GetCustomDisplacementMagnitude(DisplacementMagnitude))
+			{
+				MaterialFactoryNode->SetCustomDisplacementMagnitude(DisplacementMagnitude);
+			}
 		}
 
 		if(UInterchangeShaderPortsAPI::HasInput(ShaderGraphNode, SubstrateMaterial::Parameters::Occlusion))
diff --git a/Engine/Plugins/Interchange/Runtime/Source/Pipelines/Private/InterchangeMaterialXPipeline.cpp b/Engine/Plugins/Interchange/Runtime/Source/Pipelines/Private/InterchangeMaterialXPipeline.cpp
--- a/Engine/Plugins/Interchange/Runtime/Source/Pipelines/Private/InterchangeMaterialXPipeline.cpp
+++ b/Engine/Plugins/Interchange/Runtime/Source/Pipelines/Private/InterchangeMaterialXPipeline.cpp
@@ -8,6 +8,7 @@
 #include "InterchangePipelineLog.h"
 
 #include "Materials/MaterialExpressionMaterialFunctionCall.h"
+#include "Materials/MaterialExpressionStaticBoolParameter.h"
 #include "Materials/MaterialFunction.h"
 #include "Misc/PackageName.h"
 
@@ -781,7 +782,6 @@ void UInterchangeMaterialXPipeline::ExecutePipeline(UInterchangeBaseNodeContaine
 	Super::ExecutePipeline(NodeContainer, InSourceDatas, ContentBasePath);
 
 #if WITH_EDITOR
-
 	// we only override the volume shader if it's not a data-driven path
 	if (bVolumetricMaterial)
 	{
@@ -802,32 +802,63 @@ void UInterchangeMaterialXPipeline::ExecutePipeline(UInterchangeBaseNodeContaine
 		}
 	}
 
-	auto UpdateMaterialXNodes = [this, NodeContainer](const FString& NodeUid, UInterchangeMaterialFunctionCallExpressionFactoryNode* FactorNode)
+	TArray<UInterchangeMaterialFunctionCallExpressionFactoryNode*> DisplacementNodes;
+	// only add the StaticBoolParameter if the MX_Displacement hasn't been overridden
+	const bool bRecomputeNormalsStaticBool = DisplacementFunctionsPath ==  MaterialXSettings->PredefinedSurfaceShaders.FindChecked(EInterchangeMaterialXShaders::Displacement).ToString();
+
+	auto UpdateMaterialXNodes = [this, NodeContainer, &DisplacementNodes, bRecomputeNormalsStaticBool](const FString& NodeUid, UInterchangeMaterialFunctionCallExpressionFactoryNode* FactoryNode)
 	{
 		const FString MaterialFunctionMemberName = GET_MEMBER_NAME_CHECKED(UMaterialExpressionMaterialFunctionCall, MaterialFunction).ToString();
 
-		FString FunctionShaderNodeUID = FactorNode->GetUniqueID();
+		FString FunctionShaderNodeUID = FactoryNode->GetUniqueID();
 		FunctionShaderNodeUID.RemoveFromStart(TEXT("Factory_"));
 
 		const UInterchangeFunctionCallShaderNode* FunctionCallShaderNode = Cast<UInterchangeFunctionCallShaderNode>(NodeContainer->GetNode(FunctionShaderNodeUID));
 		
 		if(int32 EnumType; FunctionCallShaderNode->GetInt32Attribute(UE::Interchange::MaterialX::Attributes::EnumType, EnumType))
 		{
-			int32 EnumValue;
-			FunctionCallShaderNode->GetInt32Attribute(UE::Interchange::MaterialX::Attributes::EnumValue, EnumValue);
-			FactorNode->AddStringAttribute(MaterialFunctionMemberName, MaterialXSettings->GetAssetPathString(MaterialXSettings->ToEnumKey(EnumType, EnumValue)));
+			if (int32 EnumValue; FunctionCallShaderNode->GetInt32Attribute(UE::Interchange::MaterialX::Attributes::EnumValue, EnumValue))
+			{
+				FactoryNode->AddStringAttribute(MaterialFunctionMemberName, MaterialXSettings->GetAssetPathString(MaterialXSettings->ToEnumKey(EnumType, EnumValue)));
+
+				if (bRecomputeNormalsStaticBool &&
+					EnumType == UE::Interchange::MaterialX::IndexSurfaceShaders &&
+					EnumValue == int32(EInterchangeMaterialXShaders::Displacement))
+				{
+					DisplacementNodes.Emplace(FactoryNode);
+				}
+			}
 		}
-		if(FString MaterialFunctionPath; FactorNode->GetStringAttribute(MaterialFunctionMemberName, MaterialFunctionPath))
+		if(FString MaterialFunctionPath; FactoryNode->GetStringAttribute(MaterialFunctionMemberName, MaterialFunctionPath))
 		{
 			if (const EInterchangeMaterialXSettings* EnumPtr = PathToEnumMapping.Find(MaterialFunctionPath))
 			{
-				FactorNode->AddStringAttribute(MaterialFunctionMemberName, MaterialXSettings->GetAssetPathString(*EnumPtr));
+				FactoryNode->AddStringAttribute(MaterialFunctionMemberName, MaterialXSettings->GetAssetPathString(*EnumPtr));
 			}
 		}
 	};
 
 	//Find all translated node we need for this pipeline
 	NodeContainer->IterateNodesOfType<UInterchangeMaterialFunctionCallExpressionFactoryNode>(UpdateMaterialXNodes);
+
+
+	// Create StaticBool expression Factory node
+	for(UInterchangeMaterialFunctionCallExpressionFactoryNode* DisplacementExpressionFactoryNode : DisplacementNodes)
+	{
+		const FString ParentUid = DisplacementExpressionFactoryNode->GetUniqueID();
+		const FString InputName = TEXT("bRecomputeNormals");
+		const FString MaterialExpressionUid = ParentUid + TEXT("\\") + InputName;
+
+		UInterchangeMaterialExpressionFactoryNode* StaticBoolExpressionFactoryNode = NewObject<UInterchangeMaterialExpressionFactoryNode>(NodeContainer);
+		StaticBoolExpressionFactoryNode->SetCustomExpressionClassName(UMaterialExpressionStaticBoolParameter::StaticClass()->GetName());
+		NodeContainer->SetupNode(StaticBoolExpressionFactoryNode, MaterialExpressionUid, InputName, EInterchangeNodeContainerType::FactoryData, ParentUid);
+
+		const FName DefaultValueMemberName = GET_MEMBER_NAME_CHECKED(UMaterialExpressionStaticBoolParameter, DefaultValue);
+		StaticBoolExpressionFactoryNode->AddBooleanAttribute(DefaultValueMemberName.ToString(), bRecomputeNormals);
+		StaticBoolExpressionFactoryNode->AddApplyAndFillDelegates<bool>(DefaultValueMemberName.ToString(), UMaterialExpressionStaticBoolParameter::StaticClass(), DefaultValueMemberName);
+
+		UInterchangeShaderPortsAPI::ConnectDefaultOuputToInput(DisplacementExpressionFactoryNode, InputName, StaticBoolExpressionFactoryNode->GetUniqueID());
+	}
 #endif
 }
 
diff --git a/Engine/Plugins/Interchange/Runtime/Source/Pipelines/Public/InterchangeMaterialXPipeline.h b/Engine/Plugins/Interchange/Runtime/Source/Pipelines/Public/InterchangeMaterialXPipeline.h
--- a/Engine/Plugins/Interchange/Runtime/Source/Pipelines/Public/InterchangeMaterialXPipeline.h
+++ b/Engine/Plugins/Interchange/Runtime/Source/Pipelines/Public/InterchangeMaterialXPipeline.h
@@ -106,6 +106,12 @@ public:
 	UPROPERTY(EditAnywhere, config, Category = "Volume")
 	bool bVolumetricMaterial = true;
 
+	/** Recompute the normals for displacement shaders, this option is best suited with procedural normals (for example shader graphs created specifically for offline-renderers).
+	  * Otherwise this option should be disabled if there's already a normal map plugged to the normal input of the surfaceshader.
+	  */
+	UPROPERTY(EditAnywhere, config, Category = "Displacement")
+	bool bRecomputeNormals = true;
+
 protected:
 	UE_API virtual void AdjustSettingsForContext(const FInterchangePipelineContextParams& ContextParams) override;
 	UE_API virtual void ExecutePipeline(UInterchangeBaseNodeContainer* BaseNodeContainer, const TArray<UInterchangeSourceData*>& SourceDatas, const FString& ContentBasePath) override;
