[Interchange] FBX now imports rigid animation from all anim stacks

#Synced-CL 48801832

diff --git a/Engine/Plugins/Interchange/Runtime/Source/Parsers/Fbx/Private/FbxScene.cpp b/Engine/Plugins/Interchange/Runtime/Source/Parsers/Fbx/Private/FbxScene.cpp
--- a/Engine/Plugins/Interchange/Runtime/Source/Parsers/Fbx/Private/FbxScene.cpp
+++ b/Engine/Plugins/Interchange/Runtime/Source/Parsers/Fbx/Private/FbxScene.cpp
@@ -664,9 +664,10 @@ namespace UE
 			}
 
 			void FFbxScene::AddRigidAnimation(FbxNode* Node
+				, FbxAnimStack* AnimationStack
 				, const UInterchangeSceneNode* UnrealNode
 				, UInterchangeBaseNodeContainer& NodeContainer
-				, TMap<FString, TSharedPtr<FPayloadContextBase, ESPMode::ThreadSafe>>& PayloadContexts)
+				, TMap<FString, TSharedPtr<FPayloadContextBase, ESPMode::ThreadSafe>>& PayloadContexts) const
 			{
 				FbxAnimCurveNode* TranlsationCurveNode = nullptr;
 				FbxAnimCurveNode* RotationCurveNode = nullptr;
@@ -679,7 +680,7 @@ namespace UE
 
 					if (FFbxAnimation::IsFbxPropertyTypeSupported(PropertyType))
 					{
-						FbxAnimCurveNode* CurveNode = Property.GetCurveNode();
+						FbxAnimCurveNode* CurveNode = Property.GetCurveNode(AnimationStack);
 
 						//only translation/rotation/scale is supported
 						if (CurveNode && CurveNode->IsAnimated())
@@ -1003,11 +1004,8 @@ namespace UE
 						}
 					}
 
-					//Add the transform payload for all node
-					if (AnimationIndex == 0)
-					{
-						AddRigidAnimation(Node, UnrealNode, NodeContainer, PayloadContexts);
-					}
+					FbxAnimStack* CurrentAnimationStack = (FbxAnimStack*)SDKScene->GetSrcObject<FbxAnimStack>(AnimationIndex);
+					AddRigidAnimation(Node, CurrentAnimationStack, UnrealNode, NodeContainer, PayloadContexts);
 				}
 			
 
diff --git a/Engine/Plugins/Interchange/Runtime/Source/Parsers/Fbx/Private/FbxScene.h b/Engine/Plugins/Interchange/Runtime/Source/Parsers/Fbx/Private/FbxScene.h
--- a/Engine/Plugins/Interchange/Runtime/Source/Parsers/Fbx/Private/FbxScene.h
+++ b/Engine/Plugins/Interchange/Runtime/Source/Parsers/Fbx/Private/FbxScene.h
@@ -66,9 +66,10 @@ namespace UE
 
 			private:
 				void AddRigidAnimation(FbxNode* Node
+					, FbxAnimStack* AnimationStack
 					, const UInterchangeSceneNode* UnrealNode
 					, UInterchangeBaseNodeContainer& NodeContainer
-					, TMap<FString, TSharedPtr<FPayloadContextBase, ESPMode::ThreadSafe>>& PayloadContexts);
+					, TMap<FString, TSharedPtr<FPayloadContextBase, ESPMode::ThreadSafe>>& PayloadContexts) const;
 				
 				void FillCommonJointRootNode(FbxScene* SDKScene, const TArray<FbxNode*>& ForceJointNodes);
 
