Remove removal of degenerated triangles in the FBX SDK parser.
- Note that the algorithm was buggy as it eliminated tiny triangles on high-tessellation surfaces

#rnx
#Synced-CL 49033479

diff --git a/Engine/Plugins/Interchange/Runtime/Source/Parsers/Fbx/Private/FbxMesh.cpp b/Engine/Plugins/Interchange/Runtime/Source/Parsers/Fbx/Private/FbxMesh.cpp
--- a/Engine/Plugins/Interchange/Runtime/Source/Parsers/Fbx/Private/FbxMesh.cpp
+++ b/Engine/Plugins/Interchange/Runtime/Source/Parsers/Fbx/Private/FbxMesh.cpp
@@ -1012,44 +1012,6 @@ bool FMeshDescriptionImporter::FillMeshDescriptionFromFbxMesh(FbxMesh* Mesh, con
 			for (int32 PolygonIndex = 0; PolygonIndex < PolygonCount; PolygonIndex++)
 			{
 				int32 PolygonVertexCount = Mesh->GetPolygonSize(PolygonIndex);
-				//Verify if the polygon is degenerate, in this case do not add them
-				{
-					float ComparisonThreshold = 1.e-12f; //SMALL_NUMBER (1.e-8f)  was not small enough, it produced missing/partial geometry on Drone from UE-192178
-					P.Reset();
-					P.AddUninitialized(PolygonVertexCount);
-					bool bAllCornerValid = true;
-					for (int32 CornerIndex = 0; CornerIndex < PolygonVertexCount; CornerIndex++)
-					{
-						const int32 ControlPointIndex = Mesh->GetPolygonVertex(PolygonIndex, CornerIndex);
-						const FVertexID VertexID(VertexOffset + ControlPointIndex);
-						if (!MeshDescription->Vertices().IsValid(VertexID))
-						{
-							bAllCornerValid = false;
-							break;
-						}
-						P[CornerIndex] = VertexPositions[VertexID];
-					}
-					if (!bAllCornerValid)
-					{
-						if (!bCorruptedMsgDone)
-						{
-							UInterchangeResultMeshError_Generic* Message = AddMessage<UInterchangeResultMeshError_Generic>(Mesh);
-							Message->Text = LOCTEXT("CorruptedFbxPolygon", "Corrupted triangle for the mesh '{MeshName}'.");
-							bCorruptedMsgDone = true;
-						}
-
-						SkippedVertexInstance += PolygonVertexCount;
-						continue;
-					}
-					check(P.Num() > 2); //triangle is the smallest polygon we can have
-					const FVector3f Normal = ((P[1] - P[2]) ^ (P[0] - P[2])).GetSafeNormal(ComparisonThreshold);
-					//Check for degenerated polygons, avoid NAN
-					if (Normal.IsNearlyZero(ComparisonThreshold) || Normal.ContainsNaN())
-					{
-						SkippedVertexInstance += PolygonVertexCount;
-						continue;
-					}
-				}
 
 				int32 RealPolygonIndex = PolygonOffset + PolygonIndex;
 				CornerInstanceIDs.Reset();
