Fixed static analysis warning

#Synced-CL 49033479

diff --git a/Engine/Plugins/Interchange/Runtime/Source/Parsers/Fbx/Private/FbxMesh.cpp b/Engine/Plugins/Interchange/Runtime/Source/Parsers/Fbx/Private/FbxMesh.cpp
--- a/Engine/Plugins/Interchange/Runtime/Source/Parsers/Fbx/Private/FbxMesh.cpp
+++ b/Engine/Plugins/Interchange/Runtime/Source/Parsers/Fbx/Private/FbxMesh.cpp
@@ -999,7 +999,6 @@ bool FMeshDescriptionImporter::FillMeshDescriptionFromFbxMesh(FbxMesh* Mesh, con
 			bool  bBeginGetMeshEdgeIndexForPolygonCalled = false;
 			bool  bBeginGetMeshEdgeIndexForPolygonRequired = true;
 			int32 CurrentVertexInstanceIndex = 0;
-			int32 SkippedVertexInstance = 0;
 
 			// keep those for all iterations to avoid heap allocations
 			TArray<FVertexInstanceID> CornerInstanceIDs;
@@ -1021,7 +1020,7 @@ bool FMeshDescriptionImporter::FillMeshDescriptionFromFbxMesh(FbxMesh* Mesh, con
 				for (int32 CornerIndex = 0; CornerIndex < PolygonVertexCount; CornerIndex++)
 				{
 					int32 VertexInstanceIndex = VertexInstanceOffset + CurrentVertexInstanceIndex + CornerIndex;
-					int32 RealFbxVertexIndex = SkippedVertexInstance + CurrentVertexInstanceIndex + CornerIndex;
+					int32 RealFbxVertexIndex = CurrentVertexInstanceIndex + CornerIndex;
 					const FVertexInstanceID VertexInstanceID(VertexInstanceIndex);
 					CornerInstanceIDs[CornerIndex] = VertexInstanceID;
 					const int32 ControlPointIndex = Mesh->GetPolygonVertex(PolygonIndex, CornerIndex);
@@ -1297,7 +1296,7 @@ bool FMeshDescriptionImporter::FillMeshDescriptionFromFbxMesh(FbxMesh* Mesh, con
 						{
 							int UVMapIndex = (FBXUVs.UVMappingMode[UVLayerIndex] == FbxLayerElement::eByControlPoint)
 								? Mesh->GetPolygonVertex(PolygonIndex, VertexIndex)
-								: SkippedVertexInstance + CurrentVertexInstanceIndex + VertexIndex;
+								: CurrentVertexInstanceIndex + VertexIndex;
 							int32 UVIndex = (FBXUVs.UVReferenceMode[UVLayerIndex] == FbxLayerElement::eDirect)
 								? UVMapIndex
 								: FBXUVs.LayerElementUV[UVLayerIndex]->GetIndexArray().GetAt(UVMapIndex);
@@ -1331,11 +1330,6 @@ bool FMeshDescriptionImporter::FillMeshDescriptionFromFbxMesh(FbxMesh* Mesh, con
 			{
 				Mesh->EndGetMeshEdgeIndexForPolygon();
 			}
-
-			if (SkippedVertexInstance > 0)
-			{
-				check(MeshDescription->Triangles().Num() == MeshDescription->Triangles().GetArraySize());
-			}
 		}
 
 		if (MeshType == EMeshType::Skinned)
