[Interchange][FBX] ufbx - rebuilding BindPose to fix multiple cases

- FBX file can have defined Pose with Type:BindPose for every bone in the skeleton
- Also, Pose might be missing for some bones or every
- Clusters supposed to have all information on bind pose for bones used in actual skinning so they are used as a base for building bind pose
- But other nodes in the skeleton structure need to be set to proper positions - these fallback to initial node position as the last resort

#Synced-CL 48801832

diff --git a/Engine/Plugins/Interchange/Runtime/Source/Parsers/Fbx/Private/Ufbx/UfbxParser.h b/Engine/Plugins/Interchange/Runtime/Source/Parsers/Fbx/Private/Ufbx/UfbxParser.h
--- a/Engine/Plugins/Interchange/Runtime/Source/Parsers/Fbx/Private/Ufbx/UfbxParser.h
+++ b/Engine/Plugins/Interchange/Runtime/Source/Parsers/Fbx/Private/Ufbx/UfbxParser.h
@@ -97,6 +97,7 @@ namespace UE::Interchange::Private
 		virtual void SetConvertSettings(const bool InbConvertScene, const bool InbForceFrontXAxis, const bool InbConvertSceneUnit, const bool InbKeepFbxNamespace, const bool InbConsiderClusterBeforePoseForMeshBindPose)  override
 		{
 			bKeepFbxNamespace = InbKeepFbxNamespace;
+			bConsiderClusterBeforePoseForMeshBindPose = InbConsiderClusterBeforePoseForMeshBindPose;
 		}
 
 		virtual bool LoadFbxFile(const FString& Filename, UInterchangeBaseNodeContainer& NodeContainer) override;
@@ -438,6 +439,7 @@ namespace UE::Interchange::Private
 		};
 
 		bool bKeepFbxNamespace = false;
+		bool bConsiderClusterBeforePoseForMeshBindPose = false;
 
 		// Unique names for elements by type, to base Label and Uid on
 		FUniqueNames ElementNames;
diff --git a/Engine/Plugins/Interchange/Runtime/Source/Parsers/Fbx/Private/Ufbx/UfbxScene.cpp b/Engine/Plugins/Interchange/Runtime/Source/Parsers/Fbx/Private/Ufbx/UfbxScene.cpp
--- a/Engine/Plugins/Interchange/Runtime/Source/Parsers/Fbx/Private/Ufbx/UfbxScene.cpp
+++ b/Engine/Plugins/Interchange/Runtime/Source/Parsers/Fbx/Private/Ufbx/UfbxScene.cpp
@@ -147,6 +147,33 @@ namespace UE::Interchange::Private
 		}
 	}
 
+	FMatrix FUfbxScene::GetBindPoseMatrix(const ufbx_node& InNode)
+	{
+		if (const FMatrix* Found = BindPose.Find(&InNode))
+		{
+			return *Found;
+		}
+		return Convert::ConvertMatrix(InNode.node_to_world);
+	}
+
+	bool FUfbxScene::IsInSkeleton(const ufbx_node& Node)
+	{
+		if (CommonJointRootNodes.IsEmpty())
+		{
+			return false;
+		}
+
+		for (const ufbx_node* Parent = &Node;Parent;Parent = Parent->parent)
+		{
+			if (CommonJointRootNodes.Contains(Parent))
+			{
+				return true;
+			}
+		}
+		return false;
+		
+	};
+	
 	void FUfbxScene::ProcessNode(UInterchangeBaseNodeContainer& NodeContainer, const ufbx_node& InNode)
 	{
 		UInterchangeSceneNode** SceneNodeFound = Parser.ElementIdToSceneNode.Find(InNode.element_id);
@@ -179,24 +206,6 @@ namespace UE::Interchange::Private
 			}
 		}
 
-		auto IsInSkeleton = [this](const ufbx_node& Node)
-		{
-			if (CommonJointRootNodes.IsEmpty())
-			{
-				return false;
-			}
-
-			for (const ufbx_node* Parent = &Node;Parent;Parent = Parent->parent)
-			{
-				if (CommonJointRootNodes.Contains(Parent->element_id))
-				{
-					return true;
-				}
-			}
-			return false;
-			
-		};
-		
 		if (InNode.bone || IsInSkeleton(InNode))
 		{
 			UInterchangeSceneNode* UnrealNode = SceneNode;
@@ -211,12 +220,7 @@ namespace UE::Interchange::Private
 			{
 				// #ufbx_todo: probably, we don't actually need Global transform, as FBX Parser does since FBX SDK returns only evaluate globally, but source is local
 				// so might refactor this carefully to locally evaluated transforms
-				FMatrix GlobalBindPoseJointMatrix = SceneUtils::EvaluateNodeTransformGlobal(Parser, InNode, 0);
-
-				if (FMatrix* Found = JointToBindPoseMap.Find(&InNode))
-				{
-					GlobalBindPoseJointMatrix = *Found;
-				}
+				FMatrix GlobalBindPoseJointMatrix = GetBindPoseMatrix(InNode);
 
 				FTransform GlobalBindPoseJointTransform(GlobalBindPoseJointMatrix);
 
@@ -239,11 +243,7 @@ namespace UE::Interchange::Private
 
 				if (InNode.parent)
 				{
-					FMatrix GlobalFbxParentMatrix = SceneUtils::EvaluateNodeTransformGlobal(Parser, *InNode.parent, 0);
-					if (FMatrix* Found = JointToBindPoseMap.Find(InNode.parent))
-					{
-						GlobalFbxParentMatrix = *Found;
-					}
+					FMatrix GlobalFbxParentMatrix = GetBindPoseMatrix(*InNode.parent);
 
 					FMatrix LocalFbxMatrix = GlobalBindPoseJointMatrix * GlobalFbxParentMatrix.Inverse();
 					FTransform LocalBindPoseJointTransform(LocalFbxMatrix);
@@ -507,7 +507,165 @@ namespace UE::Interchange::Private
 
 	FUfbxScene::FUfbxScene(FUfbxParser& InParser): Parser(InParser)
 	{
-		// Collect bind pose matrices
+		// Build bind pose for skinned mesh.
+
+		// Normally correct Bind Pose would be fully defined in fbx - 
+		//   a Pose object for each bone, with Type: "BindPose"
+		// that would identify bone position for mesh at the time of skin association
+		// But BindPose might be missing.
+		// Still, clusters(an object identifying skinning per bone per mesh)
+		// contain bone bind position
+		// But clusters will lack position for bones that are part of skeleton structure but
+		// not affecting mesh. Those will use just initial pose stored in the file.
+
+		TMap<const ufbx_node*, TMap<FString, FMatrix>> BoneToMeshIdToBindToWorldMap;
+		TMap<const ufbx_node*, TMap<FString, FMatrix>> BoneToMeshIdToGeometryToBone;
+
+		// Collect bind_to_world for each bone
+		for (ufbx_mesh* Mesh : Parser.Scene->meshes)
+		{
+			// #ufbx_todo: fix getting existing mesh element Uid using ufbx_mesh param instead of element
+			FString MeshUid = Parser.GetMeshUid(Mesh->element);
+			for (ufbx_skin_deformer* Deformer : Mesh->skin_deformers)
+			{
+				for (ufbx_skin_cluster* Cluster : Deformer->clusters)
+				{
+					
+					TMap<FString, FMatrix>& MeshIdToGlobalBindPoseJointMap = BoneToMeshIdToBindToWorldMap.FindOrAdd(Cluster->bone_node);
+
+					const FMatrix BindMatrix = Convert::ConvertMatrix(Cluster->bind_to_world);
+					MeshIdToGlobalBindPoseJointMap.Add(MeshUid, BindMatrix);
+
+					const FMatrix GeometryToBone = Convert::ConvertMatrix(Cluster->geometry_to_bone);
+					BoneToMeshIdToGeometryToBone.FindOrAdd(Cluster->bone_node).Add(MeshUid, GeometryToBone);
+				}
+			}
+		}
+
+		// Find root bone of each skeleton
+		for (const ufbx_node* Node : Parser.Scene->nodes)
+		{
+			if (Node->bone)
+			{
+				if (const ufbx_node* SkeletonRoot = SceneUtils::FindSkeletonRoot(Parser, Node))
+				{
+					CommonJointRootNodes.FindOrAdd(SkeletonRoot);
+				}
+			}
+		}
+
+		// Record which node has bind pose defined somehow - in a skin cluster or specified bind pose
+		TSet<const ufbx_node*> bHasPose;
+
+		// Build BindPose from skin clusters
+		for (const ufbx_node* Node : Parser.Scene->nodes)
+		{
+			if (TMap<FString, FMatrix>* Found = BoneToMeshIdToBindToWorldMap.Find(Node))
+			{
+				FString MeshUid;
+				for (const TTuple<FString, FMatrix>& P : *Found)
+				{
+					const FString& NewMeshUid = P.Key;
+					const FMatrix& NewMatrix = P.Value;
+					if (bHasPose.Contains(Node))
+					{
+						FMatrix BindPoseMatrix = BindPose[Node];
+						if (!BindPoseMatrix.Equals(NewMatrix, UE_KINDA_SMALL_NUMBER))
+						{
+							// Normally, each bone is used in one cluster per mesh. But Fbx supports multiple skins on each mesh,
+							// this means there could be different clusters for single bone and with different bind matrices
+							UInterchangeResultWarning_Generic* Message = Parser.AddMessage<UInterchangeResultWarning_Generic>();
+							Message->Text = FText::Format(LOCTEXT("FFbxJointMeshBindPoseGenerator_ConflictingBindPoseMatrices", 
+								"Imported SkeletalMesh and/or Skeleton might be incorrect.\nConflicting Bind Pose Matrices found for Joint [{0}].\n   Currently using Geometry [{1}] for Bind Pose Matrix [{2}].\n   However Geometry [{3}] stores a differing Bind Pose Matrix [{4}]."), 
+								FText::FromString(Parser.GetElementNameRaw(Node->element)), 
+								FText::FromString(MeshUid), FText::FromString(BindPoseMatrix.ToString()),
+								FText::FromString(NewMeshUid), FText::FromString(NewMatrix.ToString()));
+						}
+					}
+					else
+					{
+						MeshUid = NewMeshUid;
+						bHasPose.Add(Node);
+						BindPose.Add(Node, NewMatrix);
+					}
+				}
+			}
+		}
+
+		// Fill bind pose for bones that are not found in clusters from bind_pose found on node
+		bool bClusterVsPoseMissmatch = false;
+		for (ufbx_node* Node : Parser.Scene->nodes)
+		{
+			if (Node->bind_pose && Node->bind_pose->is_bind_pose)
+			{
+				ufbx_bone_pose* Pose = ufbx_get_bone_pose(Node->bind_pose, Node);
+				if (Pose)
+				{
+					const FMatrix BindPoseMatrix = Convert::ConvertMatrix(Pose->bone_to_world);
+					
+					if (!bHasPose.Contains(Node) || !Parser.bConsiderClusterBeforePoseForMeshBindPose)
+					{
+						BindPose.Add(Node, BindPoseMatrix);
+						bHasPose.Add(Node);
+					}
+					else
+					{
+						// Check matrices previously taken from cluster
+						if (const FMatrix* Found = BindPose.Find(Node))
+						{
+							if (!BindPoseMatrix.Equals(*Found, KINDA_SMALL_NUMBER))
+							{
+								bClusterVsPoseMissmatch = true;
+							}
+						}
+						
+					}
+				}
+			}
+		}
+
+		if (bClusterVsPoseMissmatch)
+		{
+			UInterchangeResultWarning_Generic* Message = Parser.AddMessage<UInterchangeResultWarning_Generic>();
+			Message->Text = LOCTEXT("FFbxJointMeshBindPoseGenerator_NonMatchingBindPoses", "BindPose Matrix generation acquired 2 different Matrices from FbxCluster vs FbxPose.\n If import is incorrect, consider using a different setting for FBX option 'bConsiderClusterBeforePoseForMeshBindPose'.");
+		}
+
+		// For nodes not found in bind pose or clusters set BindPose from node_to_world(initial "pose")
+		for (const ufbx_node* Node : Parser.Scene->nodes)
+		{
+			if (!bHasPose.Contains(Node))
+			{
+				BindPose.Add(Node, Convert::ConvertMatrix(Node->node_to_world));
+			}
+		}
+
+		// Bind pose can be detached from skeleton root(this happens when root bone in not in bind pose/clusters)
+		// then we need to rebase it (assume global transforms of nodes as local to this root)
+		for (TTuple<const ufbx_node*, FMatrix>& P : BindPose)
+		{
+			const ufbx_node* Node = P.Key;
+
+			// Check node ancestors having defined pose when this node has the pose
+			if (bHasPose.Contains(Node))
+			{
+				for(const ufbx_node* Parent = Node->parent; Parent && IsInSkeleton(*Parent); Parent = Parent->parent)
+				{
+					if (BindPose.Contains(Parent) && !bHasPose.Contains(Parent) && CommonJointRootNodes.Contains(Parent))
+					{
+						FMatrix BindPoseMatrix = BindPose[Parent];
+
+						// Rebase node under node without bind pose
+						// When bone has no bind pose (transform defined in cluster)
+						// then we treat children bind poses that are defined under it as relative to this
+						// parent's. Usually they are all global.
+						P.Value = P.Value * BindPoseMatrix;
+						break;
+					}
+				}
+			}
+		}
+
+		// Finally, Collect bind pose matrices for reskin
 		for (ufbx_mesh* Mesh : Parser.Scene->meshes)
 		{
 			// #ufbx_todo: fix getting existing mesh element Uid using ufbx_mesh param instead of element
@@ -519,27 +677,19 @@ namespace UE::Interchange::Private
 				{
 					const FMatrix GeometryToBone = Convert::ConvertMatrix(Cluster->geometry_to_bone);
 					const FMatrix BindToWorld = Convert::ConvertMatrix(Cluster->bind_to_world);
+
 					JointToMeshIdToGlobalBindPoseReferenceMap.FindOrAdd(Cluster->bone_node).Add(MeshUid, GeometryToBone*BindToWorld);
 				}
 			}
+
 			for (ufbx_skin_deformer* Deformer : Mesh->skin_deformers)
 			{
 				for (ufbx_skin_cluster* Cluster : Deformer->clusters)
 				{
-					
 					TMap<FString, FMatrix>& MeshIdToGlobalBindPoseJointMap = JointToMeshIdToGlobalBindPoseJointMap.FindOrAdd(Cluster->bone_node);
 
 					const FMatrix BindMatrix = Convert::ConvertMatrix(Cluster->bind_to_world);
 					MeshIdToGlobalBindPoseJointMap.Add(MeshUid, BindMatrix);
-
-					if (!JointToBindPoseMap.Contains(Cluster->bone_node))
-					{
-						JointToBindPoseMap.Add(Cluster->bone_node, BindMatrix);
-					}
-					else
-					{
-						ensureMsgf(JointToBindPoseMap[Cluster->bone_node].Equals(BindMatrix, KINDA_SMALL_NUMBER), TEXT("Expected to have same joint node in different clusters to have same Bind Pose matrix!"));
-					}
 				}
 			}
 		}
@@ -581,14 +731,7 @@ namespace UE::Interchange::Private
 			{
 				continue;
 			}
-			
-			if (Node->bone)
-			{
-				if (uint32 SkeletonRoot = SceneUtils::FindSkeletonRoot(Parser, Node)->element_id)
-				{
-					CommonJointRootNodes.FindOrAdd(SkeletonRoot);
-				}
-			}
+
 
 			if (!Node->parent && (Node->element_id != RootNode->element_id)) 
 			{
diff --git a/Engine/Plugins/Interchange/Runtime/Source/Parsers/Fbx/Private/Ufbx/UfbxScene.h b/Engine/Plugins/Interchange/Runtime/Source/Parsers/Fbx/Private/Ufbx/UfbxScene.h
--- a/Engine/Plugins/Interchange/Runtime/Source/Parsers/Fbx/Private/Ufbx/UfbxScene.h
+++ b/Engine/Plugins/Interchange/Runtime/Source/Parsers/Fbx/Private/Ufbx/UfbxScene.h
@@ -27,15 +27,17 @@ namespace UE::Interchange::Private
 
 	private:
 		void ProcessNode(UInterchangeBaseNodeContainer& NodeContainer, const ufbx_node& Node);
+		bool IsInSkeleton(const ufbx_node& Node);
+		FMatrix GetBindPoseMatrix(const ufbx_node& InNode);
 
 		FUfbxParser& Parser;
 
-		TSet<uint32> CommonJointRootNodes;
-		TMap<ufbx_node*, FString> SkeletonRootPerBone;
+		TSet<const ufbx_node*> CommonJointRootNodes;
+		TMap<const ufbx_node*, FString> SkeletonRootPerBone;
 
+		TMap<const ufbx_node*, FMatrix> BindPose;
 		TMap<const ufbx_node*, TMap<FString, FMatrix>> JointToMeshIdToGlobalBindPoseReferenceMap;
 		TMap<const ufbx_node*, TMap<FString, FMatrix>> JointToMeshIdToGlobalBindPoseJointMap;
-		TMap<const ufbx_node*, FMatrix> JointToBindPoseMap;
 	};
 
 }
