[Interchang] FBX Actor visibility is imported from the FBX file.

#rnx
#Synced-CL 49373489

diff --git a/Engine/Plugins/Interchange/Runtime/Source/Parsers/Fbx/Private/FbxScene.cpp b/Engine/Plugins/Interchange/Runtime/Source/Parsers/Fbx/Private/FbxScene.cpp
--- a/Engine/Plugins/Interchange/Runtime/Source/Parsers/Fbx/Private/FbxScene.cpp
+++ b/Engine/Plugins/Interchange/Runtime/Source/Parsers/Fbx/Private/FbxScene.cpp
@@ -159,6 +159,8 @@ namespace UE
 					UnrealNode->SetCustomIsSceneRoot(true);
 				}
 
+				UnrealNode->SetCustomActorVisibility(Node->GetVisibility());
+
 				auto GetConvertedTransform = [Node](FbxAMatrix& NewFbxMatrix)
 				{
 					FTransform Transform = FFbxConvert::ConvertTransform<FTransform, FVector, FQuat>(NewFbxMatrix);
