InterchangeAxF:
- Remove warning from the StartupModule as it 's triggered by some internal tool. The warning will still be triggered in the Translator if Substrate is disabled.

#Synced-CL 49028617

diff --git a/Engine/Plugins/Experimental/InterchangeAxF/Runtime/Source/InterchangeAxF/Private/InterchangeAxF.cpp b/Engine/Plugins/Experimental/InterchangeAxF/Runtime/Source/InterchangeAxF/Private/InterchangeAxF.cpp
--- a/Engine/Plugins/Experimental/InterchangeAxF/Runtime/Source/InterchangeAxF/Private/InterchangeAxF.cpp
+++ b/Engine/Plugins/Experimental/InterchangeAxF/Runtime/Source/InterchangeAxF/Private/InterchangeAxF.cpp
@@ -17,12 +17,6 @@ DEFINE_LOG_CATEGORY(LogInterchangeAxF);
 // This code will execute after your module is loaded into memory; the exact timing is specified in the .uplugin file per-module
 void FInterchangeAxFModule::StartupModule()
 {
-    bool IsSubstrateEnabled = GetDefault<URendererSettings>() ? GetDefault<URendererSettings>()->bEnableSubstrate : false;
-    if (!IsSubstrateEnabled)
-    {
-        UE_LOG(LogInterchangeAxF, Warning, TEXT("Interchange AxF requires Substrate to be enabled. Enable \"Substrate materials\" in Project Settings under Engine > Rendering."));
-    }
-
     auto RegisterItems = []()
     {
         UInterchangeManager &InterchangeManager = UInterchangeManager::GetInterchangeManager();
