[Interchange] Fixed sockets not properly detected on joint roots

[RN] Fixed sockets not properly detected on joint roots
#Synced-CL 49021560

diff --git a/Engine/Plugins/Interchange/Runtime/Source/Import/Private/Mesh/InterchangeSkeletalMeshFactory.cpp b/Engine/Plugins/Interchange/Runtime/Source/Import/Private/Mesh/InterchangeSkeletalMeshFactory.cpp
--- a/Engine/Plugins/Interchange/Runtime/Source/Import/Private/Mesh/InterchangeSkeletalMeshFactory.cpp
+++ b/Engine/Plugins/Interchange/Runtime/Source/Import/Private/Mesh/InterchangeSkeletalMeshFactory.cpp
@@ -125,23 +125,23 @@ namespace UE
 
 		namespace Private
 		{
-			void FindAllSockets(const UInterchangeBaseNodeContainer* NodeContainer, const FString& JointNodeId, const FString JointDisplayLabel, TMap<FString, TArray<FString>>& OutSocketsPerBoneMap)
+			void FindAllSockets(const UInterchangeBaseNodeContainer* NodeContainer, const UInterchangeBaseNode* BaseNode, TMap<FString, TSet<const UInterchangeSceneNode*>>& OutSocketsPerBoneMap)
 			{
-				const TArray<FString> ChildrenIds = NodeContainer->GetNodeChildrenUids(JointNodeId);
+				const TArray<FString> ChildrenIds = NodeContainer->GetNodeChildrenUids(BaseNode->GetUniqueID());
 				for (int32 ChildIndex = 0; ChildIndex < ChildrenIds.Num(); ++ChildIndex)
 				{
 					const FString& ChildUid = ChildrenIds[ChildIndex];
-					if (const UInterchangeSceneNode* ChildJointNode = Cast<UInterchangeSceneNode>(NodeContainer->GetNode(ChildUid)))
+					if (const UInterchangeSceneNode* ChildNode = Cast<UInterchangeSceneNode>(NodeContainer->GetNode(ChildUid)))
 					{
-						if (FSkeletonHelper::IsValidSocket(NodeContainer, ChildJointNode))
+						if (FSkeletonHelper::IsValidSocket(NodeContainer, ChildNode))
 						{
 							//We found a socket add it to the map
-							TArray<FString>& Sockets = OutSocketsPerBoneMap.FindOrAdd(JointDisplayLabel);
-							Sockets.AddUnique(ChildUid);
+							TSet<const UInterchangeSceneNode*>& Sockets = OutSocketsPerBoneMap.FindOrAdd(BaseNode->GetDisplayLabel());
+							Sockets.Add(ChildNode);
 							//Socket cannot have children
 							continue;
 						}
-						FindAllSockets(NodeContainer, ChildrenIds[ChildIndex], ChildJointNode->GetDisplayLabel(), OutSocketsPerBoneMap);
+						FindAllSockets(NodeContainer, ChildNode, OutSocketsPerBoneMap);
 					}
 				}
 			}
@@ -2078,9 +2078,14 @@ UInterchangeFactoryBase::FImportAssetResult UInterchangeSkeletalMeshFactory::End
 			if(!RootSkeletonNodeUid.IsEmpty())
 			{
 				const FReferenceSkeleton& RefSkeleton = SkeletonReference->GetReferenceSkeleton();
-				TMap<FString, TArray<FString>> SocketsPerBoneMap;
-				UE::Interchange::Private::FindAllSockets(Arguments.NodeContainer, RootSkeletonNodeUid, RootSkeletonNodeUid, SocketsPerBoneMap);
-				for (const TPair<FString, TArray<FString>>& BoneAndSockets : SocketsPerBoneMap)
+				TMap<FString, TSet<const UInterchangeSceneNode*>> SocketsPerBoneMap;
+
+				if (const UInterchangeBaseNode* RootSkeletonNode = Arguments.NodeContainer->GetNode(RootSkeletonNodeUid))
+				{
+					UE::Interchange::Private::FindAllSockets(Arguments.NodeContainer, RootSkeletonNode, SocketsPerBoneMap);
+				}
+
+				for (const TPair<FString, TSet<const UInterchangeSceneNode*>>& BoneAndSockets : SocketsPerBoneMap)
 				{
 					//Find the bone in the skeleton
 					const FString& BoneName = BoneAndSockets.Key;
@@ -2088,35 +2093,32 @@ UInterchangeFactoryBase::FImportAssetResult UInterchangeSkeletalMeshFactory::End
 					if (RefSkeleton.IsValidIndex(BoneIndex))
 					{
 						//Add or update sockets
-						for (const FString& SocketNodeUid : BoneAndSockets.Value)
+						for (const UInterchangeSceneNode* SocketSceneNode : BoneAndSockets.Value)
 						{
-							if (const UInterchangeSceneNode* SocketSceneNode = Cast<UInterchangeSceneNode>(Arguments.NodeContainer->GetNode(SocketNodeUid)))
+							const FString SocketNameString = SocketSceneNode->GetDisplayLabel();
+
+							//Remove the SOCKET_ prefix from the socket node name
+							const int32 SocketPrefixLength = UInterchangeMeshFactoryNode::GetMeshSocketPrefix().Len();
+							FName SocketName = (SocketNameString.Len() > SocketPrefixLength) ? *FString(SocketNameString.RightChop(SocketPrefixLength)) : *SocketNameString;
+							USkeletalMeshSocket* Socket = SkeletalMesh->FindSocket(SocketName);
+							if (!Socket)
 							{
-								const FString SocketNameString = SocketSceneNode->GetDisplayLabel();
-								
-								//Remove the SOCKET_ prefix from the socket node name
-								const int32 SocketPrefixLength = UInterchangeMeshFactoryNode::GetMeshSocketPrefix().Len();
-								FName SocketName = (SocketNameString.Len() > SocketPrefixLength) ? *FString(SocketNameString.RightChop(SocketPrefixLength)) : *SocketNameString;
-								USkeletalMeshSocket* Socket = SkeletalMesh->FindSocket(SocketName);
-								if (!Socket)
-								{
-									//Create the socket
-									Socket = NewObject<USkeletalMeshSocket>(SkeletalMesh);
-									Socket->SocketName = SocketName;
-									Sockets.Add(Socket);
-								}
+								//Create the socket
+								Socket = NewObject<USkeletalMeshSocket>(SkeletalMesh);
+								Socket->SocketName = SocketName;
+								Sockets.Add(Socket);
+							}
 
-								if(ensure(Socket))
-								{
-									//Update the socket bone here in case the socket exist but the bone change
-									Socket->BoneName = *BoneName;
-
-									FTransform LocalSocketTransform;
-									SocketSceneNode->GetCustomLocalTransform(LocalSocketTransform);
-									Socket->RelativeLocation = LocalSocketTransform.GetLocation();
-									Socket->RelativeRotation = LocalSocketTransform.GetRotation().Rotator();
-									Socket->RelativeScale = LocalSocketTransform.GetScale3D();
-								}
+							if (ensure(Socket))
+							{
+								//Update the socket bone here in case the socket exist but the bone change
+								Socket->BoneName = *BoneName;
+
+								FTransform LocalSocketTransform;
+								SocketSceneNode->GetCustomLocalTransform(LocalSocketTransform);
+								Socket->RelativeLocation = LocalSocketTransform.GetLocation();
+								Socket->RelativeRotation = LocalSocketTransform.GetRotation().Rotator();
+								Socket->RelativeScale = LocalSocketTransform.GetScale3D();
 							}
 						}
 					}
