Interchange: Add support for point/spot light source radius.

#Synced-CL 49553638
#Synced-CL 49553638

diff --git a/Engine/Plugins/Interchange/Extensions/OpenUSD/Source/Import/Private/SchemaHandlers/LightSchemaHandler.cpp b/Engine/Plugins/Interchange/Extensions/OpenUSD/Source/Import/Private/SchemaHandlers/LightSchemaHandler.cpp
--- a/Engine/Plugins/Interchange/Extensions/OpenUSD/Source/Import/Private/SchemaHandlers/LightSchemaHandler.cpp
+++ b/Engine/Plugins/Interchange/Extensions/OpenUSD/Source/Import/Private/SchemaHandlers/LightSchemaHandler.cpp
@@ -197,8 +197,15 @@ namespace UE::LuxLightSchemaHandler::Private
 		{
 			const FUsdStageInfo StageInfo(Prim.GetStage());
 
-			float Radius = 0.5f;
-			UsdUtils::GetAuthoredAttributeValue<float>(Prim, RadiusToken, Radius);
+			float Radius = 1.0f;
+			if (UInterchangePointLightNode* CastLightNode = Cast<UInterchangePointLightNode>(LightNode))
+			{
+				if (UsdUtils::GetAuthoredAttributeValue<float>(Prim, RadiusToken, Radius))
+				{
+					Radius = UsdToUnreal::ConvertDistance(StageInfo, Radius);
+					CastLightNode->SetCustomSourceRadius(Radius);
+				}
+			}
 
 			if (Prim.HasAPI(TEXT("ShapingAPI")))
 			{
diff --git a/Engine/Plugins/Interchange/Runtime/Source/FactoryNodes/Private/InterchangeLightFactoryNode.cpp b/Engine/Plugins/Interchange/Runtime/Source/FactoryNodes/Private/InterchangeLightFactoryNode.cpp
--- a/Engine/Plugins/Interchange/Runtime/Source/FactoryNodes/Private/InterchangeLightFactoryNode.cpp
+++ b/Engine/Plugins/Interchange/Runtime/Source/FactoryNodes/Private/InterchangeLightFactoryNode.cpp
@@ -196,6 +196,16 @@ bool UInterchangePointLightFactoryNode::SetCustomLightFalloffExponent(float Attr
 	IMPLEMENT_NODE_ATTRIBUTE_SETTER(UInterchangePointLightFactoryNode, LightFalloffExponent, float, UPointLightComponent)
 }
 
+bool UInterchangePointLightFactoryNode::GetCustomSourceRadius(float& AttributeValue) const
+{
+	IMPLEMENT_NODE_ATTRIBUTE_GETTER(SourceRadius, float)
+}
+
+bool UInterchangePointLightFactoryNode::SetCustomSourceRadius(float AttributeValue, bool bAddApplyDelegate)
+{
+	IMPLEMENT_NODE_ATTRIBUTE_SETTER(UInterchangePointLightFactoryNode, SourceRadius, float, UPointLightComponent)
+}
+
 void UInterchangePointLightFactoryNode::CopyWithObject(const UInterchangeFactoryBaseNode* SourceNode, UObject* Object)
 {
 	Super::CopyWithObject(SourceNode, Object);
@@ -204,6 +214,7 @@ void UInterchangePointLightFactoryNode::CopyWithObject(const UInterchangeFactory
 	{
 		COPY_NODE_DELEGATES(LightFactoryNode, bUseInverseSquaredFalloff, bool, UPointLightComponent)
 		COPY_NODE_DELEGATES(LightFactoryNode, LightFalloffExponent, float, UPointLightComponent)
+		COPY_NODE_DELEGATES(LightFactoryNode, SourceRadius, float, UPointLightComponent)
 	}
 }
 
diff --git a/Engine/Plugins/Interchange/Runtime/Source/FactoryNodes/Public/InterchangeLightFactoryNode.h b/Engine/Plugins/Interchange/Runtime/Source/FactoryNodes/Public/InterchangeLightFactoryNode.h
--- a/Engine/Plugins/Interchange/Runtime/Source/FactoryNodes/Public/InterchangeLightFactoryNode.h
+++ b/Engine/Plugins/Interchange/Runtime/Source/FactoryNodes/Public/InterchangeLightFactoryNode.h
@@ -174,11 +174,18 @@ public:
 	UFUNCTION(BlueprintCallable, Category = "Interchange | Node | PointLightFactory")
 	UE_API bool SetCustomLightFalloffExponent(float AttributeValue, bool bAddApplyDelegate = true);
 
+	UFUNCTION(BlueprintCallable, Category = "Interchange | Node | PointLightFactory")
+	UE_API bool GetCustomSourceRadius(float& AttributeValue) const;
+
+	UFUNCTION(BlueprintCallable, Category = "Interchange | Node | PointLightFactory")
+	UE_API bool SetCustomSourceRadius(float AttributeValue, bool bAddApplyDelegate = true);
+
 	UE_API virtual void CopyWithObject(const UInterchangeFactoryBaseNode* SourceNode, UObject* Object) override;
 
 private:
 	IMPLEMENT_NODE_ATTRIBUTE_KEY(bUseInverseSquaredFalloff)
 	IMPLEMENT_NODE_ATTRIBUTE_KEY(LightFalloffExponent)
+	IMPLEMENT_NODE_ATTRIBUTE_KEY(SourceRadius)
 };
 
 UCLASS(MinimalAPI, BlueprintType)
diff --git a/Engine/Plugins/Interchange/Runtime/Source/Nodes/Private/InterchangeLightNode.cpp b/Engine/Plugins/Interchange/Runtime/Source/Nodes/Private/InterchangeLightNode.cpp
--- a/Engine/Plugins/Interchange/Runtime/Source/Nodes/Private/InterchangeLightNode.cpp
+++ b/Engine/Plugins/Interchange/Runtime/Source/Nodes/Private/InterchangeLightNode.cpp
@@ -145,6 +145,16 @@ bool UInterchangePointLightNode::SetCustomLightFalloffExponent(float AttributeVa
 	IMPLEMENT_NODE_ATTRIBUTE_SETTER_NODELEGATE(LightFalloffExponent, float)
 }
 
+bool UInterchangePointLightNode::GetCustomSourceRadius(float& AttributeValue) const
+{
+	IMPLEMENT_NODE_ATTRIBUTE_GETTER(SourceRadius, float)
+}
+
+bool UInterchangePointLightNode::SetCustomSourceRadius(float AttributeValue)
+{
+	IMPLEMENT_NODE_ATTRIBUTE_SETTER_NODELEGATE(SourceRadius, float)
+}
+
 FString UInterchangeSpotLightNode::GetTypeName() const
 {
 	return TEXT("SpotLightNode");
diff --git a/Engine/Plugins/Interchange/Runtime/Source/Nodes/Public/InterchangeLightNode.h b/Engine/Plugins/Interchange/Runtime/Source/Nodes/Public/InterchangeLightNode.h
--- a/Engine/Plugins/Interchange/Runtime/Source/Nodes/Public/InterchangeLightNode.h
+++ b/Engine/Plugins/Interchange/Runtime/Source/Nodes/Public/InterchangeLightNode.h
@@ -149,9 +149,16 @@ public:
 	UFUNCTION(BlueprintCallable, Category = "Interchange | Node | PointLight")
 	UE_API bool SetCustomLightFalloffExponent(float AttributeValue);
 
+	UFUNCTION(BlueprintCallable, Category = "Interchange | Node | PointLight")
+	UE_API bool GetCustomSourceRadius(float& AttributeValue) const;
+
+	UFUNCTION(BlueprintCallable, Category = "Interchange | Node | PointLight")
+	UE_API bool SetCustomSourceRadius(float AttributeValue);
+
 private:
 	IMPLEMENT_NODE_ATTRIBUTE_KEY(UseInverseSquaredFalloff)
 	IMPLEMENT_NODE_ATTRIBUTE_KEY(LightFalloffExponent)
+	IMPLEMENT_NODE_ATTRIBUTE_KEY(SourceRadius)
 };
 
 UCLASS(MinimalAPI, BlueprintType)
diff --git a/Engine/Plugins/Interchange/Runtime/Source/Pipelines/Private/InterchangeGenericScenesPipeline.cpp b/Engine/Plugins/Interchange/Runtime/Source/Pipelines/Private/InterchangeGenericScenesPipeline.cpp
--- a/Engine/Plugins/Interchange/Runtime/Source/Pipelines/Private/InterchangeGenericScenesPipeline.cpp
+++ b/Engine/Plugins/Interchange/Runtime/Source/Pipelines/Private/InterchangeGenericScenesPipeline.cpp
@@ -1099,6 +1099,8 @@ void UInterchangeGenericLevelPipeline::SetUpFactoryNode(UInterchangeActorFactory
 					{
 						if (UInterchangePointLightFactoryNode* PointLightFactoryNode = Cast<UInterchangePointLightFactoryNode>(LightFactoryNode))
 						{
+							COPY_FROM_TRANSLATED_TO_FACTORY(PointLightNode, PointLightFactoryNode, SourceRadius, float)
+
 							if (bool bUseInverseSquaredFalloff; PointLightNode->GetCustomUseInverseSquaredFalloff(bUseInverseSquaredFalloff))
 							{
 								PointLightFactoryNode->SetCustomUseInverseSquaredFalloff(bUseInverseSquaredFalloff);
