Interchange FBX:
- Fix material rendered black by defaulting the diffuse color to white
- Remove static switch parameters : bHasShininessMap, bHasOpacityMaskMap, bHasOpacityMap to reduce the size of the assets during serialization
- Default static switch parameter bIsPhong to true to avoid deduplication of shaders for material instances (resulting in larger size file)

#virtualized
#Synced-CL 49645631

diff --git a/Engine/Plugins/Interchange/Runtime/Source/Parsers/Fbx/Private/FbxMaterial.cpp b/Engine/Plugins/Interchange/Runtime/Source/Parsers/Fbx/Private/FbxMaterial.cpp
--- a/Engine/Plugins/Interchange/Runtime/Source/Parsers/Fbx/Private/FbxMaterial.cpp
+++ b/Engine/Plugins/Interchange/Runtime/Source/Parsers/Fbx/Private/FbxMaterial.cpp
@@ -297,7 +297,6 @@ namespace UE
 							TVariant<FLinearColor, float> DefaultValue;
 							DefaultValue.Set<float>(0.f); // Opaque
 							FName InputName = Phong::Parameters::Opacity;
-							FString HasOpacity{ TEXT("bHasOpacity") };
 							EBlendMode BlendMode = EBlendMode::BLEND_Translucent;							
 
 							// The texture is hooked to the OpacityMask when transparency is with a texture
@@ -305,12 +304,8 @@ namespace UE
 							{
 								InputName = Phong::Parameters::OpacityMask;
 								BlendMode = EBlendMode::BLEND_Masked;
-								MaterialInstanceNode->AddStaticSwitchParameterValue(TEXT("bHasOpacityMaskMap"), true);
-								HasOpacity += TEXT("Mask");
 							}
 
-							HasOpacity += TEXT("Map");
-							MaterialInstanceNode->AddStaticSwitchParameterValue(HasOpacity, true);
 							MaterialInstanceNode->SetCustomBlendMode(BlendMode);
 							
 							bHasInput |= ConvertPropertyToMaterialInstanceNode(NodeContainer, MaterialInstanceNode, MaterialProperty, Factor, InputName, DefaultValue, true);
@@ -327,15 +322,15 @@ namespace UE
 					}
 				}
 
-				// if it has no specular no shininess then just take a Lambert material
-				bool bIsPhong = false;
+				// Default to Phong, if there's no specular, no shininess then it's a Lambert material
+				bool bIsPhong = true;
 
 				// Specular
 				{
 					const float Factor = GetFactor(FbxSurfaceMaterial::sSpecularFactor);
 					TVariant<FLinearColor, float> DefaultValue;
 					DefaultValue.Set<FLinearColor>(FLinearColor::Black);
-					bIsPhong |= ConnectInput(Phong::Parameters::SpecularColor, FbxSurfaceMaterial::sSpecular, Factor, DefaultValue, false);
+					bIsPhong &= ConnectInput(Phong::Parameters::SpecularColor, FbxSurfaceMaterial::sSpecular, Factor, DefaultValue, false);
 					bHasInput |= bIsPhong;
 				}
 
@@ -346,11 +341,6 @@ namespace UE
 					{
 						TVariant<FLinearColor, float> DefaultValue;
 						DefaultValue.Set<float>(20.f);
-						if (ShininessProperty.GetSrcObjectCount<FBXSDK_NAMESPACE::FbxTexture>() > 0)
-						{
-							MaterialInstanceNode->AddStaticSwitchParameterValue(TEXT("bHasShininessMap"), true);
-						}
-
 						bIsPhong |= ConnectInput(Phong::Parameters::Shininess, FbxSurfaceMaterial::sShininess, 1.f, DefaultValue, false);
 						bHasInput |= bIsPhong;
 					}
