[Interchange] [UFbx] Fixed ensure in UfbxParser.cpp

Regression introduced in Joint Nodes refactor (UE-354035)
FUfbxParser::GetNodeUid was relying on ElementIdToSceneNode which might not be populated at that point.
Instead, we can use NodeToUid map to build ids using the parent's one.

#rnx
#Synced-CL 49752044

diff --git a/Engine/Plugins/Interchange/Runtime/Source/Parsers/Fbx/Private/Ufbx/UfbxParser.cpp b/Engine/Plugins/Interchange/Runtime/Source/Parsers/Fbx/Private/Ufbx/UfbxParser.cpp
--- a/Engine/Plugins/Interchange/Runtime/Source/Parsers/Fbx/Private/Ufbx/UfbxParser.cpp
+++ b/Engine/Plugins/Interchange/Runtime/Source/Parsers/Fbx/Private/Ufbx/UfbxParser.cpp
@@ -351,13 +351,10 @@ FString FUfbxParser::GetNodeUid(const ufbx_node& Node) const
 	// reference FBX parser's uses UE::Interchange::Private::FFbxHelper::GetFbxNodeHierarchyName to create NodeUid
 	if (Node.parent)
 	{
-		const UInterchangeSceneNode*const* ParentSceneNodeFound = ElementIdToSceneNode.Find(Node.parent->element_id);
-		
-		if (ensureMsgf(ParentSceneNodeFound, TEXT("Expected to have child node Uid generated after parent was processed first")))
+		const FString* ParentNodeUid = NodeToUid.Find(Node.parent);
+		if (ensureMsgf(ParentNodeUid, TEXT("Expected to have child node Uid generated after parent was processed first")))
 		{
-			const UInterchangeSceneNode* ParentSceneNode = *ParentSceneNodeFound;
-	
-			NodeUid = ParentSceneNode->GetUniqueID()+TEXT(".")+NodeUid;
+			NodeUid = *ParentNodeUid + TEXT(".") + NodeUid;
 		}
 	}
 
