[Interchange] ufbx - fix naming to match FBX SDK Parser's current

#Synced-CL 49765931

diff --git a/Engine/Plugins/Interchange/Runtime/Source/Parsers/Fbx/Private/Ufbx/UfbxParser.cpp b/Engine/Plugins/Interchange/Runtime/Source/Parsers/Fbx/Private/Ufbx/UfbxParser.cpp
--- a/Engine/Plugins/Interchange/Runtime/Source/Parsers/Fbx/Private/Ufbx/UfbxParser.cpp
+++ b/Engine/Plugins/Interchange/Runtime/Source/Parsers/Fbx/Private/Ufbx/UfbxParser.cpp
@@ -49,13 +49,16 @@ FUfbxParser::FUfbxParser(TWeakObjectPtr<UInterchangeResultsContainer> InResultsC
 		NodeNames.bUseNodeRules = true;
 
 		BoneNames.DefaultName = TEXT("Node");
+		BoneNames.bUseNCL = false;
 		BoneNames.bIsJoint = true;
 
 		MeshNames.DefaultName = TEXT("Mesh");
 
 		CameraNames.DefaultName = TEXT("Camera");
+		CameraNames.bUseNCL = false;
 
 		LightNames.DefaultName =  TEXT("Light");
+		LightNames.bUseNCL = false;
 
 		MaterialNames.DefaultName = TEXT("Material");
 
diff --git a/Engine/Plugins/Interchange/Runtime/Source/Parsers/Fbx/Private/Ufbx/UfbxParser.h b/Engine/Plugins/Interchange/Runtime/Source/Parsers/Fbx/Private/Ufbx/UfbxParser.h
--- a/Engine/Plugins/Interchange/Runtime/Source/Parsers/Fbx/Private/Ufbx/UfbxParser.h
+++ b/Engine/Plugins/Interchange/Runtime/Source/Parsers/Fbx/Private/Ufbx/UfbxParser.h
@@ -318,8 +318,11 @@ namespace UE::Interchange::Private
 			 * @param bOverrideUseNCL whether to use '_ncl_' suffix when making unique name, or just add a number
 			 * @return 
 			 */
-			FString MakeUniqueName(const ufbx_element& Element, const FString& BaseName, bool bOverrideUseNCL=false) const 
+			FString MakeUniqueName(const ufbx_element& Element, const FString& BaseName) const 
 			{
+
+				bool bOverrideUseNCL = BaseName.IsEmpty();
+
 				// Fix base name, non-node elements might need to have namespace removed
 				const FString FixedBaseName = [&]
 					{
diff --git a/Engine/Plugins/Interchange/Runtime/Source/Parsers/Fbx/Private/Ufbx/UfbxScene.cpp b/Engine/Plugins/Interchange/Runtime/Source/Parsers/Fbx/Private/Ufbx/UfbxScene.cpp
--- a/Engine/Plugins/Interchange/Runtime/Source/Parsers/Fbx/Private/Ufbx/UfbxScene.cpp
+++ b/Engine/Plugins/Interchange/Runtime/Source/Parsers/Fbx/Private/Ufbx/UfbxScene.cpp
@@ -708,22 +708,14 @@ namespace UE::Interchange::Private
 		{
 			FString BaseName = Parser.GetElementNameRaw(Node->element);
 
-			bool bUseNCL = false;
 
-			if (BaseName.IsEmpty())
+			if (Node->is_root)
 			{
-				if (Node->is_root)
-				{
-					// ufbx doesn't forcefully initialize root node name, so in case it's empty - set it
-					BaseName = TEXT("RootNode");
-				}
-				else
-				{
-					bUseNCL = true;
-				}
+				// ufbx doesn't forcefully initialize root node name, so in case it's empty - set it
+				BaseName = TEXT("RootNode");
 			}
 
-			Parser.NodeNames.MakeUniqueName(Node->element, BaseName, bUseNCL);
+			Parser.NodeNames.MakeUniqueName(Node->element, BaseName);
 		}
 
 		TArray<UnrealNodeInfo> UnrealNodeInfos;
