[BUGFIX] Interchange - Static Mesh Combine Option was not correctly migrated from the previously deprecated property.

- Fixed the checks done in the PostLoad that now takes care of the migration properly.

#rnx
#virtualized
#Synced-CL 49809867

diff --git a/Engine/Plugins/Interchange/Runtime/Source/Pipelines/Private/InterchangeGenericMeshPipeline.cpp b/Engine/Plugins/Interchange/Runtime/Source/Pipelines/Private/InterchangeGenericMeshPipeline.cpp
--- a/Engine/Plugins/Interchange/Runtime/Source/Pipelines/Private/InterchangeGenericMeshPipeline.cpp
+++ b/Engine/Plugins/Interchange/Runtime/Source/Pipelines/Private/InterchangeGenericMeshPipeline.cpp
@@ -399,15 +399,12 @@ void UInterchangeGenericMeshPipeline::PostLoad()
 	{
 		bCollision = bImportCollision_DEPRECATED;
 	}
-
-	if (!bCombineStaticMeshes_DEPRECATED)
-	{
-		CombineStaticMeshesBehavior = EInterchangeCombineStaticMeshesBehavior::DoNotCombine;
-	}
-	else if (bCombineStaticMeshes_DEPRECATED)
+	
+	if (bCombineStaticMeshes_DEPRECATED && CombineStaticMeshesBehavior == EInterchangeCombineStaticMeshesBehavior::DoNotCombine)
 	{
 		CombineStaticMeshesBehavior = EInterchangeCombineStaticMeshesBehavior::All;
 	}
+	
 }
 
 UInterchangePipelineMeshesUtilities* UInterchangeGenericMeshPipeline::CreateMeshPipelineUtilities(UInterchangeBaseNodeContainer* InBaseNodeContainer
