Interchange:
- Improve performance when re-importing large scenes.
- Set displacement magnitude in the detail panels
- Plug refraction for substrate materials

#Synced-CL 49645631

#Synced-CL 49645631

diff --git a/Engine/Plugins/Interchange/Runtime/Source/Import/Private/Material/InterchangeMaterialFactory.cpp b/Engine/Plugins/Interchange/Runtime/Source/Import/Private/Material/InterchangeMaterialFactory.cpp
--- a/Engine/Plugins/Interchange/Runtime/Source/Import/Private/Material/InterchangeMaterialFactory.cpp
+++ b/Engine/Plugins/Interchange/Runtime/Source/Import/Private/Material/InterchangeMaterialFactory.cpp
@@ -1061,7 +1061,9 @@ UInterchangeFactoryBase::FImportAssetResult UInterchangeMaterialFactory::BeginIm
 							if(ParentPath == PreviousParentPath)
 							{
 								SetupReimportedMaterialInstance(*MaterialInstance, *Arguments.NodeContainer, *MaterialInstanceFactoryNode, *PreviousNode);
-								MaterialInstance->UpdateStaticPermutation();
+								FStaticParameterSet StaticParameters;
+								MaterialInstance->GetStaticParameterValues(StaticParameters);
+								MaterialInstance->UpdateStaticPermutation(StaticParameters);
 								bApplyPipelineProperties = false;
 							}
 						}
@@ -1447,6 +1449,11 @@ void UInterchangeMaterialFactory::SetupMaterial(UMaterial* Material, const FImpo
 						{
 							Material->DisplacementScaling.Center = DisplacementCenter;
 						}
+
+						if (float DisplacementMagnitude; MaterialFactoryNode->GetCustomDisplacementMagnitude(DisplacementMagnitude))
+						{
+							Material->DisplacementScaling.Magnitude = DisplacementMagnitude;
+						}
 					}
 				}
 			}
@@ -1474,6 +1481,33 @@ void UInterchangeMaterialFactory::SetupMaterial(UMaterial* Material, const FImpo
 			}
 		}
 
+		// Refraction
+		if (UInterchangeShaderPortsAPI::HasInput(MaterialFactoryNode, Common::Parameters::Refraction))
+		{
+			FString ExpressionNodeUid;
+			FString OutputName;
+
+			UInterchangeShaderPortsAPI::GetInputConnection(MaterialFactoryNode, Common::Parameters::Refraction.ToString(), ExpressionNodeUid, OutputName);
+
+			const UInterchangeMaterialExpressionFactoryNode* Refraction = Cast<UInterchangeMaterialExpressionFactoryNode>(Arguments.NodeContainer->GetNode(ExpressionNodeUid));
+
+			if (Refraction)
+			{
+				const UInterchangeMaterialExpressionFactoryNode* RefractionNode = Cast<UInterchangeMaterialExpressionFactoryNode>(Arguments.NodeContainer->GetNode(ExpressionNodeUid));
+
+				if (RefractionNode)
+				{
+					if (UMaterialExpression* RefractionExpression = Builder.CreateExpressionsForNode(*RefractionNode))
+					{
+						if (FExpressionInput* RefractionInput = Material->GetExpressionInputForProperty(MP_Refraction))
+						{
+							RefractionExpression->ConnectExpression(RefractionInput, GetOutputIndex(*RefractionExpression, OutputName));
+						}
+					}
+				}
+			}
+		}
+
 		UMaterialEditingLibrary::LayoutMaterialExpressions(Material);
 
 		return;
@@ -1950,6 +1984,11 @@ void UInterchangeMaterialFactory::SetupMaterial(UMaterial* Material, const FImpo
 						{
 							Material->DisplacementScaling.Center = DisplacementCenter;
 						}
+
+						if (float DisplacementMagnitude; MaterialFactoryNode->GetCustomDisplacementMagnitude(DisplacementMagnitude))
+						{
+							Material->DisplacementScaling.Magnitude = DisplacementMagnitude;
+						}
 					}
 				}
 			}
