[Interchange][FBX] ufbx opacity and diffuse/base color factor handling fixed

- with legacy materials ufbx needs some treatment to make sure opaque material imported as opaque

- although FBX format exports Factor value for maps, including, Base/Diffuse the Unreal's legacy and FBX impoters both don't use and keep either clean texture or color(when texture is missing).

#Synced-CL 49917634

diff --git a/Engine/Plugins/Interchange/Runtime/Source/Parsers/Fbx/Private/Ufbx/UfbxMaterial.cpp b/Engine/Plugins/Interchange/Runtime/Source/Parsers/Fbx/Private/Ufbx/UfbxMaterial.cpp
--- a/Engine/Plugins/Interchange/Runtime/Source/Parsers/Fbx/Private/Ufbx/UfbxMaterial.cpp
+++ b/Engine/Plugins/Interchange/Runtime/Source/Parsers/Fbx/Private/Ufbx/UfbxMaterial.cpp
@@ -236,7 +236,8 @@ namespace UE::Interchange::Private
 
 				MaterialInstanceNode->AddTextureParameterValue(InputName.ToString() + TEXT("Map"), TextureNode->GetUniqueID());
 
-				MaterialInstanceNode->AddScalarParameterValue(InputName.ToString() + TEXT("MapWeight"), Factor);
+				// Using MapWeight=1.0 to have same behavior as FBX SDK/legacy parser
+				MaterialInstanceNode->AddScalarParameterValue(InputName.ToString() + TEXT("MapWeight"), 1.0);
 
 				// #ufbx_todo: Handle texture transform
 				// Current material asset used as parent of the material instance does not handle UV transform yet.
@@ -343,10 +344,26 @@ namespace UE::Interchange::Private
 
 			MaterialInstanceNode->AddStaticSwitchParameterValue(TEXT("bIsPhong"), bIsPhong);
 
+
 			// Opacity
+
+			// A Legacy Material might be opaque but with fbx.transparency_color - still should be treated as opaque
+			bool bOpaque = false;
+			if (Material->shader_type == UFBX_SHADER_FBX_PHONG
+				|| Material->shader_type == UFBX_SHADER_FBX_LAMBERT)
+			{
+				for (const ufbx_prop& Prop : Material->props.props)
+				{
+					if (Convert::ToUnrealString(Prop.name) == TEXT("Opacity"))
+					{
+						bOpaque = Prop.value_real >= 1.f;
+					}
+				}
+			}
+
 			// Connect only if transparency is either a texture or different from 0.f
 			// Check fbx.transparency first...
-			if (Material->fbx.transparency_color.has_value)
+			if (!bOpaque && Material->fbx.transparency_color.has_value)
 			{
 				FName InputName = Phong::Parameters::Opacity;
 				const ufbx_material_map& ColorMap = Material->fbx.transparency_color;
