[Interchange] UFBX - Actor Visibility Information is added to Scene Nodes.

#rnx

#Synced-CL 50223086

diff --git a/Engine/Plugins/Interchange/Runtime/Source/Parsers/Fbx/Private/Ufbx/UfbxScene.cpp b/Engine/Plugins/Interchange/Runtime/Source/Parsers/Fbx/Private/Ufbx/UfbxScene.cpp
--- a/Engine/Plugins/Interchange/Runtime/Source/Parsers/Fbx/Private/Ufbx/UfbxScene.cpp
+++ b/Engine/Plugins/Interchange/Runtime/Source/Parsers/Fbx/Private/Ufbx/UfbxScene.cpp
@@ -200,6 +200,8 @@ namespace UE::Interchange::Private
 			SceneNode->SetCustomIsSceneRoot(true);
 		}
 
+		SceneNode->SetCustomActorVisibility(InNode.visible);
+
 		const FTransform GeometryTransform = Convert::ConvertTransform(InNode.geometry_transform);
 		if (!GeometryTransform.Equals(FTransform::Identity))
 		{
