[On behalf of Stefan Ivanovski from XRite]
Interchange AxF

#14250

#virtualized
#Synced-CL 50150178

diff --git a/Engine/Plugins/Experimental/InterchangeAxF/Assets/InterchangeAxfAssets.uplugin b/Engine/Plugins/Experimental/InterchangeAxF/Assets/InterchangeAxfAssets.uplugin
--- a/Engine/Plugins/Experimental/InterchangeAxF/Assets/InterchangeAxfAssets.uplugin
+++ b/Engine/Plugins/Experimental/InterchangeAxF/Assets/InterchangeAxfAssets.uplugin
@@ -20,5 +20,15 @@
 			"Type": "Runtime",
 			"LoadingPhase": "PostConfigInit"
 		}
+	],
+	"Plugins": [
+		{
+			"Name": "InterchangeAssets",
+			"Enabled": true
+		},
+		{
+			"Name": "BaseMaterial",
+			"Enabled": true
+		}
 	]
 }
\ No newline at end of file
diff --git a/Engine/Plugins/Experimental/InterchangeAxF/Runtime/Source/InterchangeAxF/Private/AxFInterchangePipeline.cpp b/Engine/Plugins/Experimental/InterchangeAxF/Runtime/Source/InterchangeAxF/Private/AxFInterchangePipeline.cpp
--- a/Engine/Plugins/Experimental/InterchangeAxF/Runtime/Source/InterchangeAxF/Private/AxFInterchangePipeline.cpp
+++ b/Engine/Plugins/Experimental/InterchangeAxF/Runtime/Source/InterchangeAxF/Private/AxFInterchangePipeline.cpp
@@ -62,13 +62,13 @@ void UAxFInterchangePipeline::ExecutePipeline(UInterchangeBaseNodeContainer* Bas
                 }
                 else if (UInterchangeTexture2DNode* Texture2DNode = Cast<UInterchangeTexture2DNode>(Node))
                 {
-					double WidthMM = 0., HeightMM = 0.;
+					float WidthMM = 0.f, HeightMM = 0.f;
                     Texture2DNode->GetAttribute(TEXT("WidthMM"), WidthMM);
                     Texture2DNode->GetAttribute(TEXT("HeightMM"), HeightMM);
 
                     FString ParameterName = GetParameterNameFromUID(Texture2DNode->GetUniqueID());
                     FAxFTextureMap& Pair = TexturesMap.FindOrAdd(ParameterName);
-                    Pair.UV = FVector2d{ WidthMM, HeightMM };
+                    Pair.SizeMM = FVector2f{ WidthMM, HeightMM };
 				}
             }
             break;
@@ -91,6 +91,11 @@ void UAxFInterchangePipeline::ExecutePipeline(UInterchangeBaseNodeContainer* Bas
 
 	for (const UAxFMaterialObjectNode* AxFNode : AxFObjectNodes)
     {
+		if (!AxFNode)
+		{
+			continue;
+		}
+
 		const TArray<EAxFFeature>& Features = AxFNode->PayloadData.UsedFeatures;
 		const TMap<FString, FLinearColor>& ValuesMap = AxFNode->PayloadData.ValuesMap;
 
@@ -140,15 +145,13 @@ void UAxFInterchangePipeline::ExecutePipeline(UInterchangeBaseNodeContainer* Bas
 		constexpr float Epsilon = 1e-6;
 		for (FString const& Param : ParameterNames)
 		{
-            // NOTE: Exceptional case where we don't have a "GlintUVScale" parameter in the base material,
-            // rather we need to tweak individual parameters for each texture coordinate axis
-            if (Param == TEXT("GlintUVScale"))
+            if (Param == TEXT("FlakePatchSizeMM"))
             {
-                const FString XAttributeKey = UInterchangeShaderPortsAPI::MakeInputValueKey(TEXT("ScaleX"));
-                const FString YAttributeKey = UInterchangeShaderPortsAPI::MakeInputValueKey(TEXT("ScaleY"));
-				const FLinearColor& GlintUVScale = ValuesMap[Param];
-				MaterialInstance->AddFloatAttribute(XAttributeKey, 1.0f / (GlintUVScale.R + Epsilon));
-				MaterialInstance->AddFloatAttribute(YAttributeKey, 1.0f / (GlintUVScale.G + Epsilon));
+                const FString XAttributeKey = UInterchangeShaderPortsAPI::MakeInputValueKey(TEXT("WidthMM"));
+                const FString YAttributeKey = UInterchangeShaderPortsAPI::MakeInputValueKey(TEXT("HeightMM"));
+				const FLinearColor& FlakePatchSizeMM = ValuesMap[Param];
+				MaterialInstance->AddFloatAttribute(XAttributeKey, FlakePatchSizeMM.R + Epsilon);
+				MaterialInstance->AddFloatAttribute(YAttributeKey, FlakePatchSizeMM.G + Epsilon);
             }
 
 			// NOTE: Disable texture path in the case that we have a value for this parameter
@@ -169,8 +172,8 @@ void UAxFInterchangePipeline::ExecutePipeline(UInterchangeBaseNodeContainer* Bas
         
         if (bUseTriplanarMappingByDefault)
         {
-            FString TriplanarUVsToggleAttributeKey = UInterchangeShaderPortsAPI::MakeInputValueKey(TEXT("EnableTriplanarUVs"));
-            MaterialInstance->AddBooleanAttribute(TriplanarUVsToggleAttributeKey, true);
+            FString GeoUVsToggleAttributeKey = UInterchangeShaderPortsAPI::MakeInputValueKey(TEXT("Use Geo UVs"));
+            MaterialInstance->AddBooleanAttribute(GeoUVsToggleAttributeKey, false);
         }
 
 		for (UInterchangeTexture2DFactoryNode* Texture2DFactoryNode : Texture2DFactoryNodes)
@@ -197,9 +200,9 @@ void UAxFInterchangePipeline::ExecutePipeline(UInterchangeBaseNodeContainer* Bas
 
             if (Param != TEXT("CPA2ColorTable"))
 			{
-				const FVector2d UVScale = FVector2d::UnitVector / (TexturesMap.FindOrAdd(Param).UV + Epsilon);
-				FString UVScaleXAttributeKey = UInterchangeShaderPortsAPI::MakeInputValueKey(TEXT("ScaleX"));
-				FString UVScaleYAttributeKey = UInterchangeShaderPortsAPI::MakeInputValueKey(TEXT("ScaleY"));
+                const FVector2f UVScale = TexturesMap.FindOrAdd(Param).SizeMM;
+				FString UVScaleXAttributeKey = UInterchangeShaderPortsAPI::MakeInputValueKey(TEXT("WidthMM"));
+				FString UVScaleYAttributeKey = UInterchangeShaderPortsAPI::MakeInputValueKey(TEXT("HeightMM"));
 				MaterialInstance->AddFloatAttribute(UVScaleXAttributeKey, UVScale.X);
 				MaterialInstance->AddFloatAttribute(UVScaleYAttributeKey, UVScale.Y);
 			}
diff --git a/Engine/Plugins/Experimental/InterchangeAxF/Runtime/Source/InterchangeAxF/Private/AxFTranslator.cpp b/Engine/Plugins/Experimental/InterchangeAxF/Runtime/Source/InterchangeAxF/Private/AxFTranslator.cpp
--- a/Engine/Plugins/Experimental/InterchangeAxF/Runtime/Source/InterchangeAxF/Private/AxFTranslator.cpp
+++ b/Engine/Plugins/Experimental/InterchangeAxF/Runtime/Source/InterchangeAxF/Private/AxFTranslator.cpp
@@ -99,10 +99,10 @@ bool UAxFTranslator::Translate(UInterchangeBaseNodeContainer &BaseNodeContainer)
 		{TEXT("clearCoatTransmissionColor"), TEXT("ClearCoatTransmissionColor")},
 		{TEXT("cpaColorTable"), TEXT("CPA2ColorTable")},
 		{TEXT("diffuseColor"), TEXT("DiffuseColor")},
+		{TEXT("flakePatchSizeMM"), TEXT("FlakePatchSizeMM")},
 		{TEXT("glintBrightness"), TEXT("GlintBrightness")},
 		{TEXT("glintDensity"), TEXT("GlintDensity")},
 		{TEXT("glintRoughness"), TEXT("GlintRoughness")},
-		{TEXT("glintUVScale"), TEXT("GlintUVScale")},
 		{TEXT("heightmap"), TEXT("HeightMap")},
 		{TEXT("heightmapScale"), TEXT("HeightStrength")},
 		{TEXT("heightmapReference"), TEXT("HeightReferencePlane")},
@@ -156,6 +156,24 @@ bool UAxFTranslator::Translate(UInterchangeBaseNodeContainer &BaseNodeContainer)
         {
             ParameterPresent = true;
 
+            float WidthMM{}, HeightMM{};
+            if (!ValueObject->Get()->TryGetNumberField(TEXT("widthMM"), WidthMM))
+            {
+                UE_LOG(LogAxFTranslator, Error, TEXT("Failed to find '%s' property in AxF2UE file."), TEXT("widthMM"));
+                return false;
+            }
+            if (!ValueObject->Get()->TryGetNumberField(TEXT("heightMM"), HeightMM))
+            {
+                UE_LOG(LogAxFTranslator, Error, TEXT("Failed to find '%s' property in AxF2UE file."), TEXT("heightMM"));
+                return false;
+            }
+
+            if (UEParameterName == TEXT("FlakePatchSizeMM"))
+            {
+                AxFData.ValuesMap.Add(UEParameterName) = FLinearColor{ WidthMM, HeightMM, 0.0f };
+                continue;
+            }
+
 			FString TexturePath{};
 			if (!ValueObject->Get()->TryGetStringField(TEXT("path"), TexturePath))
 			{
@@ -163,24 +181,12 @@ bool UAxFTranslator::Translate(UInterchangeBaseNodeContainer &BaseNodeContainer)
 				return false;
 			}
 
-			double WidthMM{}, HeightMM{};
-			if (!ValueObject->Get()->TryGetNumberField(TEXT("widthMM"), WidthMM))
-			{
-				UE_LOG(LogAxFTranslator, Error, TEXT("Failed to find '%s' property in AxF2UE file."), TEXT("widthMM"));
-				return false;
-			}
-			if (!ValueObject->Get()->TryGetNumberField(TEXT("heightMM"), HeightMM))
-			{
-				UE_LOG(LogAxFTranslator, Error, TEXT("Failed to find '%s' property in AxF2UE file."), TEXT("heightMM"));
-				return false;
-			}
-
             FString TextureUID = TEXT("T_") + DisplayLabel + TEXT("_") + UEParameterName;
             UInterchangeTexture2DNode *Texture2D = NewObject<UInterchangeTexture2DNode>(&BaseNodeContainer, FName(TextureUID));
             Texture2D->InitializeNode(TextureUID, TextureUID, EInterchangeNodeContainerType::TranslatedAsset);
             Texture2D->SetPayLoadKey(Path + "|" + Name + "." + Extension + "|" + TexturePath);
-            Texture2D->AddDoubleAttribute(TEXT("WidthMM"), WidthMM);
-            Texture2D->AddDoubleAttribute(TEXT("HeightMM"), HeightMM);
+            Texture2D->AddFloatAttribute(TEXT("WidthMM"), WidthMM);
+            Texture2D->AddFloatAttribute(TEXT("HeightMM"), HeightMM);
             BaseNodeContainer.AddNode(Texture2D);
 
             if (UEParameterName == TEXT("CPA2ColorTable"))
@@ -253,6 +259,7 @@ bool UAxFTranslator::Translate(UInterchangeBaseNodeContainer &BaseNodeContainer)
                 AxFData.UsedFeatures.AddUnique(EAxFFeature::Heightmap);
             }
 
+            // Flakes
             if (UEParameterName == TEXT("GlintDensity") || UEParameterName == TEXT("GlintUVScale") ||
 				UEParameterName == TEXT("GlintRoughness") || UEParameterName == TEXT("GlintBrightness"))
             {
diff --git a/Engine/Plugins/Experimental/InterchangeAxF/Runtime/Source/InterchangeAxF/Private/InterchangeAxF.cpp b/Engine/Plugins/Experimental/InterchangeAxF/Runtime/Source/InterchangeAxF/Private/InterchangeAxF.cpp
--- a/Engine/Plugins/Experimental/InterchangeAxF/Runtime/Source/InterchangeAxF/Private/InterchangeAxF.cpp
+++ b/Engine/Plugins/Experimental/InterchangeAxF/Runtime/Source/InterchangeAxF/Private/InterchangeAxF.cpp
@@ -49,6 +49,16 @@ void FInterchangeAxFModule::StartupModule()
 				}
 			}
 
+			if (!AssetImportSettings.ShowImportDialogOverride.Contains(EInterchangeTranslatorAssetType::Materials))
+			{
+				FInterchangePerTranslatorDialogOverride AxFTranslatorImportDialogOverride{};
+				AxFTranslatorImportDialogOverride.Translator = TSoftClassPtr<UInterchangeTranslatorBase>(UAxFTranslator::StaticClass());
+				AxFTranslatorImportDialogOverride.bShowImportDialog = false;
+				AxFTranslatorImportDialogOverride.bShowReimportDialog = false;
+				FInterchangeDialogOverride& DialogOverride = AssetImportSettings.ShowImportDialogOverride.FindOrAdd(EInterchangeTranslatorAssetType::Materials); 
+				DialogOverride.PerTranslatorImportDialogOverride.Add(AxFTranslatorImportDialogOverride);
+			}
+
 			if(!bPipelinesWereAdded)
 			{
 				UE_LOG(LogInterchangeAxF, Warning, TEXT("Cannot find Interchange default Materials' stack to add pipelines. Interchange AxF material import may not work properly."));
diff --git a/Engine/Plugins/Experimental/InterchangeAxF/Runtime/Source/InterchangeAxF/Public/AxFInterchangePipeline.h b/Engine/Plugins/Experimental/InterchangeAxF/Runtime/Source/InterchangeAxF/Public/AxFInterchangePipeline.h
--- a/Engine/Plugins/Experimental/InterchangeAxF/Runtime/Source/InterchangeAxF/Public/AxFInterchangePipeline.h
+++ b/Engine/Plugins/Experimental/InterchangeAxF/Runtime/Source/InterchangeAxF/Public/AxFInterchangePipeline.h
@@ -27,8 +27,7 @@ struct FAxFTextureMap
 
 	UPROPERTY()
 	TObjectPtr<UTexture2D> Texture;
-
-	FVector2d UV;
+	FVector2f SizeMM;
 };
 
 UCLASS(BlueprintType)
@@ -47,11 +46,11 @@ public:
 	/* Set the reimport strategy. */
 	UPROPERTY(
 		EditAnywhere, BlueprintReadWrite, Category = "Common",
-		meta = (AdjustPipelineAndRefreshDetailOnChange = "True"))
+		meta = (AdjustPipelineAndRefreshDetailOnChange = "True", PipelineInternalEditionData = "True"))
 	EReimportStrategyFlags ReimportStrategy = EReimportStrategyFlags::ApplyNoProperties;
 
     UPROPERTY(EditAnywhere, BlueprintReadWrite, Category = "AxF")
-    bool bUseTriplanarMappingByDefault = false;
+    bool bUseTriplanarMappingByDefault = true;
 
 protected:
     virtual void ExecutePipeline(
