MaterialX:
- Helper function for standard surface and openpbr in case of subsurface and transmission
- Connect the subsurface color to the transmission_scatter
- Empirical value on the transmission_depth instead of the transmission_color
- Add a Refraction IOR to the Standard Suface and OpenPBR translucent material functions to be plugged on the material
- Allow correct modulation of the transmission even in the case of a white transmission_color

#virtualized
#Synced-CL 50150178

diff --git a/Engine/Plugins/Interchange/Runtime/Source/Import/Private/MaterialX/MaterialXUtils/MaterialXOpenPBRSurfaceShader.cpp b/Engine/Plugins/Interchange/Runtime/Source/Import/Private/MaterialX/MaterialXUtils/MaterialXOpenPBRSurfaceShader.cpp
--- a/Engine/Plugins/Interchange/Runtime/Source/Import/Private/MaterialX/MaterialXUtils/MaterialXOpenPBRSurfaceShader.cpp
+++ b/Engine/Plugins/Interchange/Runtime/Source/Import/Private/MaterialX/MaterialXUtils/MaterialXOpenPBRSurfaceShader.cpp
@@ -25,7 +25,6 @@ UInterchangeBaseNode* FMaterialXOpenPBRSurfaceShader::Translate(MaterialX::NodeP
 	using namespace UE::Interchange::Materials;
 	using namespace UE::Interchange::MaterialX;
 	using namespace UE::Interchange::Materials::Standard::Nodes;
-
 	this->SurfaceShaderNode = OpenPBRSurfaceNode;
 
 	UInterchangeShaderNode* OpenPBRSurfaceShaderNode = FMaterialXSurfaceShaderAbstract::Translate(EInterchangeMaterialXShaders::OpenPBRSurface);
@@ -52,113 +51,12 @@ UInterchangeBaseNode* FMaterialXOpenPBRSurfaceShader::Translate(MaterialX::NodeP
 		{
 			MTLX_LOG("MaterialXOpenPBRSurfaceShader", "OpenPBR material rendering might be wrong. Please select Substrate Adaptive GBuffer format in the Project settings.");
 		}
-
-		auto CreateShaderNodeInternal = [this]<class T>(mx::NodePtr OpenPBRSurfaceNode, const FString & PathOrShaderType, const FString & DisplayLabel)
-		{
-			T* Node;
-
-			const FString NodeUID = UInterchangeShaderNode::MakeNodeUid(GetUniqueName(OpenPBRSurfaceNode) + DisplayLabel, FStringView{});
-
-			if (Node = const_cast<T*>(Cast<T>(NodeContainer.GetNode(NodeUID))); !Node)
-			{
-				Node = NewObject<T>(&NodeContainer);
-				NodeContainer.SetupNode(Node, NodeUID, DisplayLabel, EInterchangeNodeContainerType::TranslatedAsset);
-				if constexpr (std::is_same_v<T, UInterchangeFunctionCallShaderNode>)
-				{
-					Node->SetCustomMaterialFunction(PathOrShaderType);
-				}
-				else
-				{
-					Node->SetCustomShaderType(PathOrShaderType);
-				}
-			}
-
-			return Node;
-		};
-
-		if (UInterchangeShaderPortsAPI::HasInput(OpenPBRSurfaceShaderNode, OpenPBRSurface::Parameters::TransmissionWeight) &&
-			UInterchangeShaderPortsAPI::HasInput(OpenPBRSurfaceShaderNode, OpenPBRSurface::Parameters::SubsurfaceWeight) &&
-			UInterchangeShaderPortsAPI::HasInput(OpenPBRSurfaceShaderNode, OpenPBRSurface::Parameters::BaseColor))
-		{
-			// 1. Set the Material as Opaque with Mask Mode
-			OpenPBRSurfaceShaderNode->AddInt32Attribute(Attributes::EnumType, IndexSurfaceShaders);
-			OpenPBRSurfaceShaderNode->AddInt32Attribute(Attributes::EnumValue, int32(EInterchangeMaterialXShaders::OpenPBRSurface));
-			ShaderGraphNode->SetCustomBlendMode(EBlendMode::BLEND_Masked);
-
-			//2. Disconnect the transmission inputs
-			FString TransmissionNodeUid, TransmissionOutputName;
-			UInterchangeShaderPortsAPI::DisconnectInputFromOutputNode(OpenPBRSurfaceShaderNode, OpenPBRSurface::Parameters::TransmissionWeight.ToString(), TransmissionNodeUid, TransmissionOutputName);
-			UInterchangeShaderPortsAPI::DisconnectInput(OpenPBRSurfaceShaderNode, OpenPBRSurface::Parameters::TransmissionColor.ToString());
-			UInterchangeShaderPortsAPI::DisconnectInput(OpenPBRSurfaceShaderNode, OpenPBRSurface::Parameters::TransmissionDepth.ToString());
-			UInterchangeShaderPortsAPI::DisconnectInput(OpenPBRSurfaceShaderNode, OpenPBRSurface::Parameters::TransmissionDispersionAbbeNumber.ToString());
-			UInterchangeShaderPortsAPI::DisconnectInput(OpenPBRSurfaceShaderNode, OpenPBRSurface::Parameters::TransmissionDispersionScale.ToString());
-			UInterchangeShaderPortsAPI::DisconnectInput(OpenPBRSurfaceShaderNode, OpenPBRSurface::Parameters::TransmissionScatterAnisotropy.ToString());
-			UInterchangeShaderPortsAPI::DisconnectInput(OpenPBRSurfaceShaderNode, OpenPBRSurface::Parameters::TransmissionScatter.ToString());
-
-			//3. Create DitherMask node
-			UInterchangeShaderNode* DitherMaskNode = CreateShaderNodeInternal.operator()<UInterchangeFunctionCallShaderNode>(OpenPBRSurfaceNode, UE::MaterialFunctions::Path::DitherMask, TEXT("DitherMaskNode"));
-
-			//4. Connect the surface shader and the transmission to the MF_DitherMask
-			if (const UInterchangeShaderNode* TransmissionNode = Cast<UInterchangeShaderNode>(NodeContainer.GetNode(TransmissionNodeUid)))
-			{
-				UInterchangeShaderPortsAPI::ConnectOuputToInputByName(DitherMaskNode, TEXT("BSDF"), OpenPBRSurfaceShaderNode->GetUniqueID(), OpenPBRSurface::SubstrateMaterial::Outputs::FrontMaterial.ToString());
-				UInterchangeShaderPortsAPI::ConnectOuputToInputByName(DitherMaskNode, TEXT("transmission"), TransmissionNode->GetUniqueID(), TransmissionOutputName);
-			}
-
-			UInterchangeShaderPortsAPI::ConnectOuputToInputByName(ShaderGraphNode, SubstrateMaterial::Parameters::FrontMaterial.ToString(), DitherMaskNode->GetUniqueID(), TEXT("Dither"));
-			UInterchangeShaderPortsAPI::ConnectOuputToInputByName(ShaderGraphNode, SubstrateMaterial::Parameters::OpacityMask.ToString(), DitherMaskNode->GetUniqueID(), SubstrateMaterial::Parameters::OpacityMask.ToString());
-		}
-		else if (UInterchangeShaderPortsAPI::HasInput(OpenPBRSurfaceShaderNode, OpenPBRSurface::Parameters::TransmissionWeight) &&
-				 UInterchangeShaderPortsAPI::HasInput(OpenPBRSurfaceShaderNode, OpenPBRSurface::Parameters::SubsurfaceWeight))
-		{
-			ShaderGraphNode->SetCustomBlendMode(EBlendMode::BLEND_TranslucentColoredTransmittance);
-			auto ReplaceInputConnection = [this, OpenPBRSurfaceShaderNode](const FString& InputToRemove, const FString& InputName)
-				{
-					UInterchangeShaderPortsAPI::DisconnectInput(OpenPBRSurfaceShaderNode, InputToRemove);
-
-					FString NodeUid, NodeOutputName;
-					UInterchangeShaderPortsAPI::GetInputConnection(OpenPBRSurfaceShaderNode, InputName, NodeUid, NodeOutputName);
-					if (const UInterchangeShaderNode* ShaderNode = Cast<UInterchangeShaderNode>(NodeContainer.GetNode(NodeUid)))
-					{
-						UInterchangeShaderPortsAPI::ConnectOuputToInputByName(OpenPBRSurfaceShaderNode, InputToRemove, ShaderNode->GetUniqueID(), NodeOutputName);
-					}
-				};
-
-			//1. subsurface_color input now replaces the base_color
-			if (UInterchangeShaderPortsAPI::HasInput(OpenPBRSurfaceShaderNode, OpenPBRSurface::Parameters::SubsurfaceColor))
-			{
-				ReplaceInputConnection(OpenPBRSurface::Parameters::BaseColor.ToString(), OpenPBRSurface::Parameters::SubsurfaceColor.ToString());
-			}
-
-			//2. replace base by subsurface
-			ReplaceInputConnection(OpenPBRSurface::Parameters::BaseWeight.ToString(), OpenPBRSurface::Parameters::SubsurfaceWeight.ToString());
-
-			//3. scale the subsurface_radius by 600 (empirical value that gives better results, need further testing though)
-			FString SubsurfaceRadiusNodeUid, SubsurfaceRadiusNodeOutputName;
-			if (UInterchangeShaderPortsAPI::GetInputConnection(OpenPBRSurfaceShaderNode, OpenPBRSurface::Parameters::SubsurfaceRadius.ToString(), SubsurfaceRadiusNodeUid, SubsurfaceRadiusNodeOutputName))
-			{
-				UInterchangeShaderNode* ScaleSubsurfaceRadiusNode = CreateShaderNodeInternal.operator()<UInterchangeShaderNode>(OpenPBRSurfaceNode, Multiply::Name.ToString(), TEXT("ScaleSubsurfaceRadius"));
-				ScaleSubsurfaceRadiusNode->AddFloatAttribute(UInterchangeShaderPortsAPI::MakeInputValueKey(Multiply::Inputs::B.ToString()), 600);
-				UInterchangeShaderPortsAPI::ConnectOuputToInputByName(ScaleSubsurfaceRadiusNode, Multiply::Inputs::A.ToString(), SubsurfaceRadiusNodeUid, SubsurfaceRadiusNodeOutputName);
-
-				// 4. Multiply the scaled SubsurfaceRadius by the TransmissionColor if any
-				FString TransmissionColorNodeUid, TransmissionColorNodeOutputName;
-				if (UInterchangeShaderPortsAPI::DisconnectInputFromOutputNode(OpenPBRSurfaceShaderNode, OpenPBRSurface::Parameters::TransmissionColor.ToString(), TransmissionColorNodeUid, TransmissionColorNodeOutputName))
-				{
-					UInterchangeShaderNode* SubsurfaceTransmissionColorNode = CreateShaderNodeInternal.operator()<UInterchangeShaderNode>(OpenPBRSurfaceNode, Multiply::Name.ToString(), TEXT("SubsurfaceRadiusTransmissionColorMult"));
-					UInterchangeShaderPortsAPI::ConnectOuputToInputByName(SubsurfaceTransmissionColorNode, Multiply::Inputs::A.ToString(), TransmissionColorNodeUid, TransmissionColorNodeOutputName);
-					UInterchangeShaderPortsAPI::ConnectDefaultOuputToInput(SubsurfaceTransmissionColorNode, Multiply::Inputs::B.ToString(), ScaleSubsurfaceRadiusNode->GetUniqueID());
-					UInterchangeShaderPortsAPI::ConnectDefaultOuputToInput(OpenPBRSurfaceShaderNode, OpenPBRSurface::Parameters::TransmissionColor.ToString(), SubsurfaceTransmissionColorNode->GetUniqueID());
-				}
-				else
-				{
-					UInterchangeShaderPortsAPI::ConnectDefaultOuputToInput(OpenPBRSurfaceShaderNode, OpenPBRSurface::Parameters::TransmissionColor.ToString(), ScaleSubsurfaceRadiusNode->GetUniqueID());
-				}
-			}
-
-			UInterchangeShaderPortsAPI::ConnectOuputToInputByName(ShaderGraphNode, SubstrateMaterial::Parameters::FrontMaterial.ToString(), OpenPBRSurfaceShaderNode->GetUniqueID(), OpenPBRSurface::SubstrateMaterial::Outputs::FrontMaterial.ToString());
-		}
-		else
+		
+		if(!ConnectSubsurfaceTransmissionToSurfaceShader(OpenPBRSurfaceShaderNode, EInterchangeMaterialXShaders::OpenPBRSurface
+		   , OpenPBRSurface::Parameters::BaseWeight, OpenPBRSurface::Parameters::BaseColor
+		   , OpenPBRSurface::Parameters::TransmissionWeight, OpenPBRSurface::Parameters::TransmissionColor, OpenPBRSurface::Parameters::TransmissionDepth, OpenPBRSurface::Parameters::TransmissionDispersionAbbeNumber, OpenPBRSurface::Parameters::TransmissionDispersionScale, OpenPBRSurface::Parameters::TransmissionScatter, OpenPBRSurface::Parameters::TransmissionScatterAnisotropy
+		   , OpenPBRSurface::Parameters::SubsurfaceWeight, OpenPBRSurface::Parameters::SubsurfaceColor, OpenPBRSurface::Parameters::SubsurfaceRadiusScale
+		   , OpenPBRSurface::SubstrateMaterial::Outputs::FrontMaterial, OpenPBRSurface::SubstrateMaterial::Outputs::FrontMaterial))
 		{
 			UInterchangeShaderPortsAPI::ConnectOuputToInputByName(ShaderGraphNode, SubstrateMaterial::Parameters::FrontMaterial.ToString(), OpenPBRSurfaceShaderNode->GetUniqueID(), OpenPBRSurface::SubstrateMaterial::Outputs::FrontMaterial.ToString());
 
diff --git a/Engine/Plugins/Interchange/Runtime/Source/Import/Private/MaterialX/MaterialXUtils/MaterialXStandardSurfaceShader.cpp b/Engine/Plugins/Interchange/Runtime/Source/Import/Private/MaterialX/MaterialXUtils/MaterialXStandardSurfaceShader.cpp
--- a/Engine/Plugins/Interchange/Runtime/Source/Import/Private/MaterialX/MaterialXUtils/MaterialXStandardSurfaceShader.cpp
+++ b/Engine/Plugins/Interchange/Runtime/Source/Import/Private/MaterialX/MaterialXUtils/MaterialXStandardSurfaceShader.cpp
@@ -52,124 +52,26 @@ UInterchangeBaseNode* FMaterialXStandardSurfaceShader::Translate(mx::NodePtr Sta
 			MTLX_LOG("MaterialXStandardSurfaceShader", "Standard Surface material rendering might be wrong. Please select Substrate Adaptive GBuffer format in the Project settings.");
 		}
 
-		auto CreateShaderNodeInternal = [this]<class T>(mx::NodePtr StandardSurfaceNode, const FString& PathOrShaderType, const FString& DisplayLabel)
+		if (!ConnectSubsurfaceTransmissionToSurfaceShader(StandardSurfaceShaderNode, EInterchangeMaterialXShaders::StandardSurface
+			, StandardSurface::Parameters::Base, StandardSurface::Parameters::BaseColor
+			, StandardSurface::Parameters::Transmission, StandardSurface::Parameters::TransmissionColor, StandardSurface::Parameters::TransmissionDepth, StandardSurface::Parameters::TransmissionExtraRoughness, StandardSurface::Parameters::TransmissionDispersion, StandardSurface::Parameters::TransmissionScatter, StandardSurface::Parameters::TransmissionScatterAnisotropy
+			, StandardSurface::Parameters::Subsurface, StandardSurface::Parameters::SubsurfaceColor, StandardSurface::Parameters::SubsurfaceRadius
+			, StandardSurface::SubstrateMaterial::Outputs::Opaque, StandardSurface::SubstrateMaterial::Outputs::Translucent))
 		{
-			T* Node;
-
-			const FString NodeUID = UInterchangeShaderNode::MakeNodeUid(GetUniqueName(StandardSurfaceNode) + DisplayLabel, FStringView{});
-
-			if (Node = const_cast<T*>(Cast<T>(NodeContainer.GetNode(NodeUID))); !Node)
-			{
-				Node = NewObject<T>(&NodeContainer);
-				NodeContainer.SetupNode(Node, NodeUID, DisplayLabel, EInterchangeNodeContainerType::TranslatedAsset);
-				if constexpr (std::is_same_v<T, UInterchangeFunctionCallShaderNode>)
-				{
-					Node->SetCustomMaterialFunction(PathOrShaderType);
-				}
-				else
-				{
-					Node->SetCustomShaderType(PathOrShaderType);
-				}
-			}
-
-			return Node;
-		};
-
-		if (UInterchangeShaderPortsAPI::HasInput(StandardSurfaceShaderNode, StandardSurface::Parameters::Transmission) &&
-			UInterchangeShaderPortsAPI::HasInput(StandardSurfaceShaderNode, StandardSurface::Parameters::Subsurface) &&
-			UInterchangeShaderPortsAPI::HasInput(StandardSurfaceShaderNode, StandardSurface::Parameters::BaseColor))
-		{
-			// 1. Set the Material as Opaque with Mask Mode
-			StandardSurfaceShaderNode->AddInt32Attribute(Attributes::EnumType, IndexSurfaceShaders);
-			StandardSurfaceShaderNode->AddInt32Attribute(Attributes::EnumValue, int32(EInterchangeMaterialXShaders::StandardSurface));
-			ShaderGraphNode->SetCustomBlendMode(EBlendMode::BLEND_Masked);
-
-			//2. Disconnect the transmission inputs
-			FString TransmissionNodeUid, TransmissionOutputName;
-			UInterchangeShaderPortsAPI::DisconnectInputFromOutputNode(StandardSurfaceShaderNode, StandardSurface::Parameters::Transmission.ToString(), TransmissionNodeUid, TransmissionOutputName);
-			UInterchangeShaderPortsAPI::DisconnectInput(StandardSurfaceShaderNode, StandardSurface::Parameters::TransmissionColor.ToString());
-			UInterchangeShaderPortsAPI::DisconnectInput(StandardSurfaceShaderNode, StandardSurface::Parameters::TransmissionDepth.ToString());
-			UInterchangeShaderPortsAPI::DisconnectInput(StandardSurfaceShaderNode, StandardSurface::Parameters::TransmissionDispersion.ToString());
-			UInterchangeShaderPortsAPI::DisconnectInput(StandardSurfaceShaderNode, StandardSurface::Parameters::TransmissionExtraRoughness.ToString());
-			UInterchangeShaderPortsAPI::DisconnectInput(StandardSurfaceShaderNode, StandardSurface::Parameters::TransmissionScatter.ToString());
-			UInterchangeShaderPortsAPI::DisconnectInput(StandardSurfaceShaderNode, StandardSurface::Parameters::TransmissionScatterAnisotropy.ToString());
-
-			//3. Create DitherMask node
-			UInterchangeShaderNode* DitherMaskNode = CreateShaderNodeInternal.operator()<UInterchangeFunctionCallShaderNode>(StandardSurfaceNode, UE::MaterialFunctions::Path::DitherMask, TEXT("DitherMaskNode"));
-
-			//4. Connect the surface shader and the transmission to the MF_DitherMask
-			if (const UInterchangeShaderNode* TransmissionNode = Cast<UInterchangeShaderNode>(NodeContainer.GetNode(TransmissionNodeUid)))
-			{
-				UInterchangeShaderPortsAPI::ConnectOuputToInputByName(DitherMaskNode, TEXT("BSDF"), StandardSurfaceShaderNode->GetUniqueID(), StandardSurface::SubstrateMaterial::Outputs::Opaque.ToString());
-				UInterchangeShaderPortsAPI::ConnectOuputToInputByName(DitherMaskNode, TEXT("transmission"), TransmissionNode->GetUniqueID(), TransmissionOutputName);
-			}
-
-			UInterchangeShaderPortsAPI::ConnectOuputToInputByName(ShaderGraphNode, SubstrateMaterial::Parameters::FrontMaterial.ToString(), DitherMaskNode->GetUniqueID(), TEXT("Dither"));
-			UInterchangeShaderPortsAPI::ConnectOuputToInputByName(ShaderGraphNode, SubstrateMaterial::Parameters::OpacityMask.ToString(), DitherMaskNode->GetUniqueID(), SubstrateMaterial::Parameters::OpacityMask.ToString());
-		}
-		else if (UInterchangeShaderPortsAPI::HasInput(StandardSurfaceShaderNode, StandardSurface::Parameters::Transmission) &&
-				 UInterchangeShaderPortsAPI::HasInput(StandardSurfaceShaderNode, StandardSurface::Parameters::Subsurface))
-		{
-			ShaderGraphNode->SetCustomBlendMode(EBlendMode::BLEND_TranslucentColoredTransmittance);
-			auto ReplaceInputConnection = [this, StandardSurfaceShaderNode](const FString& InputToRemove, const FString& InputName)
-			{
-				UInterchangeShaderPortsAPI::DisconnectInput(StandardSurfaceShaderNode, InputToRemove);
-
-				FString NodeUid, NodeOutputName;
-				UInterchangeShaderPortsAPI::GetInputConnection(StandardSurfaceShaderNode, InputName, NodeUid, NodeOutputName);
-				if (const UInterchangeShaderNode* ShaderNode = Cast<UInterchangeShaderNode>(NodeContainer.GetNode(NodeUid)))
-				{
-					UInterchangeShaderPortsAPI::ConnectOuputToInputByName(StandardSurfaceShaderNode, InputToRemove, ShaderNode->GetUniqueID(), NodeOutputName);
-				}
-			};
-
-			//1. subsurface_color input now replaces the base_color
-			if(UInterchangeShaderPortsAPI::HasInput(StandardSurfaceShaderNode, StandardSurface::Parameters::SubsurfaceColor))
+			if (UInterchangeShaderPortsAPI::HasInput(StandardSurfaceShaderNode, StandardSurface::Parameters::Transmission))
 			{
-				ReplaceInputConnection(StandardSurface::Parameters::BaseColor.ToString(), StandardSurface::Parameters::SubsurfaceColor.ToString());
+				ShaderGraphNode->SetCustomBlendMode(EBlendMode::BLEND_TranslucentColoredTransmittance);
+				UInterchangeShaderPortsAPI::ConnectOuputToInputByName(ShaderGraphNode, SubstrateMaterial::Parameters::FrontMaterial.ToString(), StandardSurfaceShaderNode->GetUniqueID(), StandardSurface::SubstrateMaterial::Outputs::Translucent.ToString());
 			}
-
-			//2. replace base by subsurface
-			ReplaceInputConnection(StandardSurface::Parameters::Base.ToString(), StandardSurface::Parameters::Subsurface.ToString());
-
-			//3. scale the subsurface_radius by 600 (empirical value that gives better results, need further testing though)
-			FString SubsurfaceRadiusNodeUid, SubsurfaceRadiusNodeOutputName;
-			if (UInterchangeShaderPortsAPI::GetInputConnection(StandardSurfaceShaderNode, StandardSurface::Parameters::SubsurfaceRadius.ToString(), SubsurfaceRadiusNodeUid, SubsurfaceRadiusNodeOutputName))
+			else
 			{
-				UInterchangeShaderNode* ScaleSubsurfaceRadiusNode = CreateShaderNodeInternal.operator()<UInterchangeShaderNode>(StandardSurfaceNode, Multiply::Name.ToString(), TEXT("ScaleSubsurfaceRadius"));
-				ScaleSubsurfaceRadiusNode->AddFloatAttribute(UInterchangeShaderPortsAPI::MakeInputValueKey(Multiply::Inputs::B.ToString()), 600);
-				UInterchangeShaderPortsAPI::ConnectOuputToInputByName(ScaleSubsurfaceRadiusNode, Multiply::Inputs::A.ToString(), SubsurfaceRadiusNodeUid, SubsurfaceRadiusNodeOutputName);
-
-				// 4. Multiply the scaled SubsurfaceRadius by the TransmissionColor if any
-				FString TransmissionColorNodeUid, TransmissionColorNodeOutputName;
-				if(UInterchangeShaderPortsAPI::DisconnectInputFromOutputNode(StandardSurfaceShaderNode, StandardSurface::Parameters::TransmissionColor.ToString(), TransmissionColorNodeUid, TransmissionColorNodeOutputName))
-				{
-					UInterchangeShaderNode* SubsurfaceTransmissionColorNode = CreateShaderNodeInternal.operator()<UInterchangeShaderNode>(StandardSurfaceNode, Multiply::Name.ToString(), TEXT("SubsurfaceRadiusTransmissionColorMult"));
-					UInterchangeShaderPortsAPI::ConnectOuputToInputByName(SubsurfaceTransmissionColorNode, Multiply::Inputs::A.ToString(), TransmissionColorNodeUid, TransmissionColorNodeOutputName);
-					UInterchangeShaderPortsAPI::ConnectDefaultOuputToInput(SubsurfaceTransmissionColorNode, Multiply::Inputs::B.ToString(), ScaleSubsurfaceRadiusNode->GetUniqueID());
-					UInterchangeShaderPortsAPI::ConnectDefaultOuputToInput(StandardSurfaceShaderNode, StandardSurface::Parameters::TransmissionColor.ToString(), SubsurfaceTransmissionColorNode->GetUniqueID());
-				}
-				else
+				UInterchangeShaderPortsAPI::ConnectOuputToInputByName(ShaderGraphNode, SubstrateMaterial::Parameters::FrontMaterial.ToString(), StandardSurfaceShaderNode->GetUniqueID(), StandardSurface::SubstrateMaterial::Outputs::Opaque.ToString());
+				if (UInterchangeShaderPortsAPI::HasInput(StandardSurfaceShaderNode, StandardSurface::Parameters::Opacity))
 				{
-					UInterchangeShaderPortsAPI::ConnectDefaultOuputToInput(StandardSurfaceShaderNode, StandardSurface::Parameters::TransmissionColor.ToString(), ScaleSubsurfaceRadiusNode->GetUniqueID());
+					UInterchangeShaderPortsAPI::ConnectOuputToInputByName(ShaderGraphNode, SubstrateMaterial::Parameters::OpacityMask.ToString(), StandardSurfaceShaderNode->GetUniqueID(), StandardSurface::SubstrateMaterial::Outputs::Opacity.ToString());
+					ShaderGraphNode->SetCustomBlendMode(EBlendMode::BLEND_Masked);
 				}
 			}
-
-			UInterchangeShaderPortsAPI::ConnectOuputToInputByName(ShaderGraphNode, SubstrateMaterial::Parameters::FrontMaterial.ToString(), StandardSurfaceShaderNode->GetUniqueID(), StandardSurface::SubstrateMaterial::Outputs::Translucent.ToString());
-		}
-		else if(UInterchangeShaderPortsAPI::HasInput(StandardSurfaceShaderNode, StandardSurface::Parameters::Transmission))
-		{
-			ShaderGraphNode->SetCustomBlendMode(EBlendMode::BLEND_TranslucentColoredTransmittance);
-			UInterchangeShaderPortsAPI::ConnectOuputToInputByName(ShaderGraphNode, SubstrateMaterial::Parameters::FrontMaterial.ToString(), StandardSurfaceShaderNode->GetUniqueID(), StandardSurface::SubstrateMaterial::Outputs::Translucent.ToString());
-		}
-		else
-		{
-			UInterchangeShaderPortsAPI::ConnectOuputToInputByName(ShaderGraphNode, SubstrateMaterial::Parameters::FrontMaterial.ToString(), StandardSurfaceShaderNode->GetUniqueID(), StandardSurface::SubstrateMaterial::Outputs::Opaque.ToString());
-			if(UInterchangeShaderPortsAPI::HasInput(StandardSurfaceShaderNode, StandardSurface::Parameters::Opacity))
-			{
-				UInterchangeShaderPortsAPI::ConnectOuputToInputByName(ShaderGraphNode, SubstrateMaterial::Parameters::OpacityMask.ToString(), StandardSurfaceShaderNode->GetUniqueID(), StandardSurface::SubstrateMaterial::Outputs::Opacity.ToString());
-				ShaderGraphNode->SetCustomBlendMode(EBlendMode::BLEND_Masked);
-			}
 		}
 	}
 	else
diff --git a/Engine/Plugins/Interchange/Runtime/Source/Import/Private/MaterialX/MaterialXUtils/MaterialXSurfaceShaderAbstract.cpp b/Engine/Plugins/Interchange/Runtime/Source/Import/Private/MaterialX/MaterialXUtils/MaterialXSurfaceShaderAbstract.cpp
--- a/Engine/Plugins/Interchange/Runtime/Source/Import/Private/MaterialX/MaterialXUtils/MaterialXSurfaceShaderAbstract.cpp
+++ b/Engine/Plugins/Interchange/Runtime/Source/Import/Private/MaterialX/MaterialXUtils/MaterialXSurfaceShaderAbstract.cpp
@@ -1455,6 +1455,120 @@ void FMaterialXSurfaceShaderAbstract::ConnectLuminanceInputToOutput(const FConne
 	UInterchangeShaderPortsAPI::ConnectDefaultOuputToInput(Connect.ParentShaderNode, Connect.InputChannelName, LuminanceNode->GetUniqueID());
 }
 
+bool FMaterialXSurfaceShaderAbstract::ConnectSubsurfaceTransmissionToSurfaceShader(UInterchangeShaderNode* OpenPbrStandarSurfaceShaderNode, EInterchangeMaterialXShaders SurfaceShaderType
+																				   , const FName& BaseWeight, const FName& BaseColor
+																				   , const FName& TransmissionWeight, const FName& TransmissionColor, const FName& TransmissionDepth, const FName& TransmissionDispersionAbbeNumber, const FName& TransmissionDispersionScale, const FName& TransmissionScatter, const FName& TransmissionScatterAnisotropy
+																				   , const FName& SubsurfaceWeight, const FName& SubsurfaceColor, const FName& SubsurfaceRadiusScale
+																				   , const FName& FrontMaterialOpaque, const FName& FrontMaterialTranslucent)
+{
+	using namespace UE::Interchange::Materials;
+	using namespace UE::Interchange::MaterialX;
+	using namespace UE::Interchange::Materials::Standard::Nodes;
+
+	bool bSubsurfaceTransmission = false;
+
+	UInterchangeShaderPortsAPI::ConnectOuputToInputByName(ShaderGraphNode, Common::Parameters::Refraction.ToString(), OpenPbrStandarSurfaceShaderNode->GetUniqueID(), TEXT("Refraction (IOR)"));
+
+	if (UInterchangeShaderPortsAPI::HasInput(OpenPbrStandarSurfaceShaderNode, TransmissionWeight) &&
+		UInterchangeShaderPortsAPI::HasInput(OpenPbrStandarSurfaceShaderNode, SubsurfaceWeight) &&
+		UInterchangeShaderPortsAPI::HasInput(OpenPbrStandarSurfaceShaderNode, BaseColor))
+	{
+		bSubsurfaceTransmission = true;
+		// 1. Set the Material as Opaque with Mask Mode
+		OpenPbrStandarSurfaceShaderNode->AddInt32Attribute(Attributes::EnumType, IndexSurfaceShaders);
+		OpenPbrStandarSurfaceShaderNode->AddInt32Attribute(Attributes::EnumValue, int32(SurfaceShaderType));
+		ShaderGraphNode->SetCustomBlendMode(EBlendMode::BLEND_Masked);
+
+		//2. Disconnect the transmission inputs
+		FString TransmissionNodeUid, TransmissionOutputName;
+		UInterchangeShaderPortsAPI::DisconnectInputFromOutputNode(OpenPbrStandarSurfaceShaderNode, TransmissionWeight.ToString(), TransmissionNodeUid, TransmissionOutputName);
+		UInterchangeShaderPortsAPI::DisconnectInput(OpenPbrStandarSurfaceShaderNode, TransmissionColor.ToString());
+		UInterchangeShaderPortsAPI::DisconnectInput(OpenPbrStandarSurfaceShaderNode, TransmissionDepth.ToString());
+		UInterchangeShaderPortsAPI::DisconnectInput(OpenPbrStandarSurfaceShaderNode, TransmissionDispersionAbbeNumber.ToString());
+		UInterchangeShaderPortsAPI::DisconnectInput(OpenPbrStandarSurfaceShaderNode, TransmissionDispersionScale.ToString());
+		UInterchangeShaderPortsAPI::DisconnectInput(OpenPbrStandarSurfaceShaderNode, TransmissionScatterAnisotropy.ToString());
+		UInterchangeShaderPortsAPI::DisconnectInput(OpenPbrStandarSurfaceShaderNode, TransmissionScatter.ToString());
+
+		//3. Create DitherMask node			
+		UInterchangeShaderNode* DitherMaskNode = CreateFunctionCallShaderNode(SurfaceShaderNode, TEXT("DitherMaskNode"), UE::MaterialFunctions::Path::DitherMask);
+
+		//4. Connect the surface shader and the transmission to the MF_DitherMask
+		if (const UInterchangeShaderNode* TransmissionNode = Cast<UInterchangeShaderNode>(NodeContainer.GetNode(TransmissionNodeUid)))
+		{
+			UInterchangeShaderPortsAPI::ConnectOuputToInputByName(DitherMaskNode, TEXT("BSDF"), OpenPbrStandarSurfaceShaderNode->GetUniqueID(), FrontMaterialOpaque.ToString());
+			UInterchangeShaderPortsAPI::ConnectOuputToInputByName(DitherMaskNode, TEXT("transmission"), TransmissionNode->GetUniqueID(), TransmissionOutputName);
+		}
+
+		UInterchangeShaderPortsAPI::ConnectOuputToInputByName(ShaderGraphNode, SubstrateMaterial::Parameters::FrontMaterial.ToString(), DitherMaskNode->GetUniqueID(), TEXT("Dither"));
+		UInterchangeShaderPortsAPI::ConnectOuputToInputByName(ShaderGraphNode, SubstrateMaterial::Parameters::OpacityMask.ToString(), DitherMaskNode->GetUniqueID(), SubstrateMaterial::Parameters::OpacityMask.ToString());
+	}
+	else if (UInterchangeShaderPortsAPI::HasInput(OpenPbrStandarSurfaceShaderNode, TransmissionWeight) &&
+			 UInterchangeShaderPortsAPI::HasInput(OpenPbrStandarSurfaceShaderNode, SubsurfaceWeight))
+	{
+		bSubsurfaceTransmission = true;
+		ShaderGraphNode->SetCustomBlendMode(EBlendMode::BLEND_TranslucentColoredTransmittance);
+		auto ReplaceInputConnection = [this, OpenPbrStandarSurfaceShaderNode](const FString& InputToRemove, const FString& InputName)
+			{
+				UInterchangeShaderPortsAPI::DisconnectInput(OpenPbrStandarSurfaceShaderNode, InputToRemove);
+
+				FString NodeUid, NodeOutputName;
+				UInterchangeShaderPortsAPI::GetInputConnection(OpenPbrStandarSurfaceShaderNode, InputName, NodeUid, NodeOutputName);
+				if (const UInterchangeShaderNode* ShaderNode = Cast<UInterchangeShaderNode>(NodeContainer.GetNode(NodeUid)))
+				{
+					UInterchangeShaderPortsAPI::ConnectOuputToInputByName(OpenPbrStandarSurfaceShaderNode, InputToRemove, ShaderNode->GetUniqueID(), NodeOutputName);
+				}
+			};
+
+		//1. subsurface_color input now replaces the base_color
+		if (UInterchangeShaderPortsAPI::HasInput(OpenPbrStandarSurfaceShaderNode, SubsurfaceColor))
+		{
+			ReplaceInputConnection(BaseColor.ToString(), SubsurfaceColor.ToString());
+		}
+
+		//2. replace base by subsurface
+		ReplaceInputConnection(BaseWeight.ToString(), SubsurfaceWeight.ToString());
+
+		//3. set an empirical value to the TransmissionDepth if there isn't any input (otherwise just keep it as is).
+		//   see egg_jade in [REDACTED_PROFILE_URL]/MaterialEggs
+		//   and masonJarWater in [REDACTED_PROFILE_URL]/OpenPBRShaderPlayground
+		if (!UInterchangeShaderPortsAPI::HasInput(OpenPbrStandarSurfaceShaderNode, TransmissionDepth))
+		{
+			constexpr float TransmissionDepthValue = 3.f;
+			OpenPbrStandarSurfaceShaderNode->AddFloatAttribute(UInterchangeShaderPortsAPI::MakeInputValueKey(TransmissionDepth.ToString()), TransmissionDepthValue);
+		}
+
+		//4. Connect the SubsurfaceColor to the TransmissionScatter
+		FString SubsurfaceColorNodeUid, SubsurfaceColorNodeOutputName;
+		if (UInterchangeShaderPortsAPI::GetInputConnection(OpenPbrStandarSurfaceShaderNode, SubsurfaceColor.ToString(), SubsurfaceColorNodeUid, SubsurfaceColorNodeOutputName))
+		{
+			UInterchangeShaderPortsAPI::ConnectOuputToInputByName(OpenPbrStandarSurfaceShaderNode, TransmissionScatter.ToString(), SubsurfaceColorNodeUid, SubsurfaceColorNodeOutputName);
+		}
+
+		//5. Connect the SubsurfaceRadius for Standard Surface, SubsurfaceRadiusScale for OpenPBR to the TransmissionColor
+		FString SubsurfaceRadiusNodeUid, SubsurfaceRadiusNodeOutputName;
+		if (UInterchangeShaderPortsAPI::GetInputConnection(OpenPbrStandarSurfaceShaderNode, SubsurfaceRadiusScale.ToString(), SubsurfaceRadiusNodeUid, SubsurfaceRadiusNodeOutputName))
+		{
+			// 6. Multiply the SubsurfaceRadiusScale by the TransmissionColor if any otherwise just plug it to the TransmissionColor
+			FString TransmissionColorNodeUid, TransmissionColorNodeOutputName;
+			if (UInterchangeShaderPortsAPI::DisconnectInputFromOutputNode(OpenPbrStandarSurfaceShaderNode, TransmissionColor.ToString(), TransmissionColorNodeUid, TransmissionColorNodeOutputName))
+			{
+				UInterchangeShaderNode* SubsurfaceTransmissionColorNode = CreateShaderNode(SurfaceShaderNode, TEXT("SubsurfaceRadiusTransmissionColorMult"), Multiply::Name.ToString());
+				UInterchangeShaderPortsAPI::ConnectOuputToInputByName(SubsurfaceTransmissionColorNode, Multiply::Inputs::A.ToString(), TransmissionColorNodeUid, TransmissionColorNodeOutputName);
+				UInterchangeShaderPortsAPI::ConnectOuputToInputByName(SubsurfaceTransmissionColorNode, Multiply::Inputs::B.ToString(), SubsurfaceRadiusNodeUid, SubsurfaceRadiusNodeOutputName);
+				UInterchangeShaderPortsAPI::ConnectDefaultOuputToInput(OpenPbrStandarSurfaceShaderNode, TransmissionColor.ToString(), SubsurfaceTransmissionColorNode->GetUniqueID());
+			}
+			else
+			{
+				UInterchangeShaderPortsAPI::ConnectOuputToInputByName(OpenPbrStandarSurfaceShaderNode, TransmissionColor.ToString(), SubsurfaceRadiusNodeUid, SubsurfaceRadiusNodeOutputName);
+			}
+		}
+
+		UInterchangeShaderPortsAPI::ConnectOuputToInputByName(ShaderGraphNode, SubstrateMaterial::Parameters::FrontMaterial.ToString(), OpenPbrStandarSurfaceShaderNode->GetUniqueID(), FrontMaterialTranslucent.ToString());
+	}
+
+	return bSubsurfaceTransmission;
+}
+
 UInterchangeShaderNode* FMaterialXSurfaceShaderAbstract::CreateMaskShaderNode(uint8 RGBA, mx::ElementPtr Element, const FString& NodeName, const FString& OutputName)
 {
 	bool bR = (0b1000 & RGBA) >> 3;
diff --git a/Engine/Plugins/Interchange/Runtime/Source/Import/Public/MaterialX/MaterialXUtils/MaterialXSurfaceShaderAbstract.h b/Engine/Plugins/Interchange/Runtime/Source/Import/Public/MaterialX/MaterialXUtils/MaterialXSurfaceShaderAbstract.h
--- a/Engine/Plugins/Interchange/Runtime/Source/Import/Public/MaterialX/MaterialXUtils/MaterialXSurfaceShaderAbstract.h
+++ b/Engine/Plugins/Interchange/Runtime/Source/Import/Public/MaterialX/MaterialXUtils/MaterialXSurfaceShaderAbstract.h
@@ -341,9 +341,16 @@ protected:
 
 	/** <luminance>*/
 	void ConnectLuminanceInputToOutput(const FConnectNode& Connect);
-
 	/** End Connect MaterialX nodes*/
 
+	/** Function used only by both Standard Surface and OpenPBR to have both Subsurface and Transmission*/
+	[[nodiscard]]
+	bool ConnectSubsurfaceTransmissionToSurfaceShader(UInterchangeShaderNode* OpenPbrStandarSurfaceShaderNode, EInterchangeMaterialXShaders SurfaceShaderType													  
+													  , const FName& BaseWeight, const FName& BaseColor
+													  , const FName& TransmissionWeight, const FName& TransmissionColor, const FName& TransmissionDepth, const FName& TransmissionDispersionAbbeNumber, const FName& TransmissionDispersionScale, const FName& TransmissionScatter, const FName& TransmissionScatterAnisotropy
+													  , const FName& SubsurfaceWeight, const FName& SubsurfaceColor, const FName& SubsurfaceRadiusScale
+													  , const FName & FrontMaterialOpaque, const FName & FrontMaterialTranslucent);
+
 	/**
 	 * Create a ComponentMask shader node.
 	 *
diff --git a/Engine/Plugins/Interchange/Runtime/Source/Nodes/Public/InterchangeMaterialDefinitions.h b/Engine/Plugins/Interchange/Runtime/Source/Nodes/Public/InterchangeMaterialDefinitions.h
--- a/Engine/Plugins/Interchange/Runtime/Source/Nodes/Public/InterchangeMaterialDefinitions.h
+++ b/Engine/Plugins/Interchange/Runtime/Source/Nodes/Public/InterchangeMaterialDefinitions.h
@@ -646,6 +646,7 @@ namespace Interchange
 					const FName Opaque = TEXT("Substrate StandardSurface Opaque");
 					const FName Translucent = TEXT("Substrate StandardSurface Translucent");
 					const FName Opacity = TEXT("Geometry Opacity");
+					const FName Refraction = TEXT("Refraction (IOR)");
 				}
 			}
 		}
diff --git a/Engine/Plugins/Interchange/Runtime/Source/Pipelines/Private/InterchangeGenericMaterialPipeline.cpp b/Engine/Plugins/Interchange/Runtime/Source/Pipelines/Private/InterchangeGenericMaterialPipeline.cpp
--- a/Engine/Plugins/Interchange/Runtime/Source/Pipelines/Private/InterchangeGenericMaterialPipeline.cpp
+++ b/Engine/Plugins/Interchange/Runtime/Source/Pipelines/Private/InterchangeGenericMaterialPipeline.cpp
@@ -4026,36 +4026,33 @@ bool UInterchangeGenericMaterialPipeline::HandleSubstrate(const UInterchangeShad
 	using namespace UE::Interchange::Materials;
 	bool bShadingModelHandled = false;
 
-	if(UInterchangeShaderPortsAPI::HasInput(ShaderGraphNode, SubstrateMaterial::Parameters::FrontMaterial))
+	if (UInterchangeShaderPortsAPI::HasInput(ShaderGraphNode, SubstrateMaterial::Parameters::FrontMaterial))
 	{
 		TTuple<UInterchangeMaterialExpressionFactoryNode*, FString> FrontMaterialFactoryNode =
 			CreateMaterialExpressionForInput(MaterialFactoryNode, ShaderGraphNode, SubstrateMaterial::Parameters::FrontMaterial.ToString(), MaterialFactoryNode->GetUniqueID());
-		ensure(FrontMaterialFactoryNode.Get<0>());
-
-		if(FrontMaterialFactoryNode.Get<0>())
+		
+		if (ensure(FrontMaterialFactoryNode.Get<0>()))
 		{
 			UInterchangeShaderPortsAPI::ConnectOuputToInputByName(MaterialFactoryNode, SubstrateMaterial::Parameters::FrontMaterial.ToString(), FrontMaterialFactoryNode.Get<0>()->GetUniqueID(), FrontMaterialFactoryNode.Get<1>());
 		}
 
-		if(UInterchangeShaderPortsAPI::HasInput(ShaderGraphNode, SubstrateMaterial::Parameters::OpacityMask))
+		if (UInterchangeShaderPortsAPI::HasInput(ShaderGraphNode, SubstrateMaterial::Parameters::OpacityMask))
 		{
 			TTuple<UInterchangeMaterialExpressionFactoryNode*, FString> OpacityMaskFactoryNode =
-				CreateMaterialExpressionForInput(MaterialFactoryNode, ShaderGraphNode, SubstrateMaterial::Parameters::OpacityMask.ToString(), MaterialFactoryNode->GetUniqueID());
-			ensure(OpacityMaskFactoryNode.Get<0>());
+				CreateMaterialExpressionForInput(MaterialFactoryNode, ShaderGraphNode, SubstrateMaterial::Parameters::OpacityMask.ToString(), MaterialFactoryNode->GetUniqueID());			
 
-			if(OpacityMaskFactoryNode.Get<0>())
+			if (ensure(OpacityMaskFactoryNode.Get<0>()))
 			{
 				UInterchangeShaderPortsAPI::ConnectOuputToInputByName(MaterialFactoryNode, SubstrateMaterial::Parameters::OpacityMask.ToString(), OpacityMaskFactoryNode.Get<0>()->GetUniqueID(), OpacityMaskFactoryNode.Get<1>());
 			}
 		}
 
-		if(UInterchangeShaderPortsAPI::HasInput(ShaderGraphNode, SubstrateMaterial::Parameters::Displacement))
+		if (UInterchangeShaderPortsAPI::HasInput(ShaderGraphNode, SubstrateMaterial::Parameters::Displacement))
 		{
 			TTuple<UInterchangeMaterialExpressionFactoryNode*, FString> DisplacementFactoryNode =
-				CreateMaterialExpressionForInput(MaterialFactoryNode, ShaderGraphNode, SubstrateMaterial::Parameters::Displacement.ToString(), MaterialFactoryNode->GetUniqueID());
-			ensure(DisplacementFactoryNode.Get<0>());
+				CreateMaterialExpressionForInput(MaterialFactoryNode, ShaderGraphNode, SubstrateMaterial::Parameters::Displacement.ToString(), MaterialFactoryNode->GetUniqueID());			
 
-			if(DisplacementFactoryNode.Get<0>())
+			if (ensure(DisplacementFactoryNode.Get<0>()))
 			{
 				UInterchangeShaderPortsAPI::ConnectOuputToInputByName(MaterialFactoryNode, SubstrateMaterial::Parameters::Displacement.ToString(), DisplacementFactoryNode.Get<0>()->GetUniqueID(), DisplacementFactoryNode.Get<1>());
 			}
@@ -4071,22 +4068,32 @@ bool UInterchangeGenericMaterialPipeline::HandleSubstrate(const UInterchangeShad
 			}
 		}
 
-		if(UInterchangeShaderPortsAPI::HasInput(ShaderGraphNode, SubstrateMaterial::Parameters::Occlusion))
+		if (UInterchangeShaderPortsAPI::HasInput(ShaderGraphNode, SubstrateMaterial::Parameters::Occlusion))
 		{
 			TTuple<UInterchangeMaterialExpressionFactoryNode*, FString> OcclusionFactoryNode =
 				CreateMaterialExpressionForInput(MaterialFactoryNode, ShaderGraphNode, SubstrateMaterial::Parameters::Occlusion.ToString(), MaterialFactoryNode->GetUniqueID());
-			ensure(OcclusionFactoryNode.Get<0>());
 
-			if(OcclusionFactoryNode.Get<0>())
+			if (ensure(OcclusionFactoryNode.Get<0>()))
 			{
 				UInterchangeShaderPortsAPI::ConnectOuputToInputByName(MaterialFactoryNode, SubstrateMaterial::Parameters::Occlusion.ToString(), OcclusionFactoryNode.Get<0>()->GetUniqueID(), OcclusionFactoryNode.Get<1>());
 			}
 		}
 
-		if(EBlendMode BlendMode; ShaderGraphNode->GetCustomBlendMode(reinterpret_cast<int&>(BlendMode)))
+		if (UInterchangeShaderPortsAPI::HasInput(ShaderGraphNode, Common::Parameters::Refraction))
+		{
+			TTuple<UInterchangeMaterialExpressionFactoryNode*, FString> RefractionFactoryNode =
+				CreateMaterialExpressionForInput(MaterialFactoryNode, ShaderGraphNode, Common::Parameters::Refraction.ToString(), MaterialFactoryNode->GetUniqueID());
+
+			if (ensure(RefractionFactoryNode.Get<0>()))
+			{
+				UInterchangeShaderPortsAPI::ConnectOuputToInputByName(MaterialFactoryNode, Common::Parameters::Refraction.ToString(), RefractionFactoryNode.Get<0>()->GetUniqueID(), RefractionFactoryNode.Get<1>());
+			}
+		}
+
+		if (EBlendMode BlendMode; ShaderGraphNode->GetCustomBlendMode(reinterpret_cast<int&>(BlendMode)))
 		{
 			MaterialFactoryNode->SetCustomBlendMode(BlendMode);
-			if(BlendMode == BLEND_TranslucentColoredTransmittance)
+			if (BlendMode == BLEND_TranslucentColoredTransmittance)
 			{
 				MaterialFactoryNode->SetCustomTranslucencyLightingMode(ETranslucencyLightingMode::TLM_SurfacePerPixelLighting);
 				MaterialFactoryNode->SetCustomRefractionMethod(ERefractionMode::RM_IndexOfRefraction);
