[nterchange][FBX] ufbx camera/light orientation and camera scaling; Support for ForceFrontXAxis translator option

- camera scaling set to 1 as FBX SDK/Legacy importers do

#Synced-CL 49917634

diff --git a/Engine/Plugins/Interchange/Runtime/Source/Parsers/Fbx/Private/Ufbx/UfbxParser.cpp b/Engine/Plugins/Interchange/Runtime/Source/Parsers/Fbx/Private/Ufbx/UfbxParser.cpp
--- a/Engine/Plugins/Interchange/Runtime/Source/Parsers/Fbx/Private/Ufbx/UfbxParser.cpp
+++ b/Engine/Plugins/Interchange/Runtime/Source/Parsers/Fbx/Private/Ufbx/UfbxParser.cpp
@@ -103,11 +103,31 @@ bool FUfbxParser::LoadFbxFile(const FString& Filename, UInterchangeBaseNodeConta
 
 	LoadOpts.space_conversion = UFBX_SPACE_CONVERSION_ADJUST_TRANSFORMS;
 
+
+	if (bForceFrontXAxis)
+	{
+		// When X for front 
+		LoadOpts.target_axes = ufbx_coordinate_axes{
+	UFBX_COORDINATE_AXIS_NEGATIVE_Y, UFBX_COORDINATE_AXIS_POSITIVE_Z, UFBX_COORDINATE_AXIS_POSITIVE_X,
+		};
+	}
+	else	
+	{
+		LoadOpts.target_axes = ufbx_axes_left_handed_z_up;
+	}
 	LoadOpts.handedness_conversion_axis = UFBX_MIRROR_AXIS_Y;
-	LoadOpts.target_axes = ufbx_axes_left_handed_z_up;
+
 	LoadOpts.target_unit_meters = 0.01; // Centimeters
 	LoadOpts.reverse_winding = true;
 
+	LoadOpts.target_camera_axes.front = UFBX_COORDINATE_AXIS_NEGATIVE_X;
+	LoadOpts.target_camera_axes.up = UFBX_COORDINATE_AXIS_POSITIVE_Z;
+	LoadOpts.target_camera_axes.right = UFBX_COORDINATE_AXIS_NEGATIVE_Y;
+
+	LoadOpts.target_light_axes.front = UFBX_COORDINATE_AXIS_NEGATIVE_X;
+	LoadOpts.target_light_axes.up = UFBX_COORDINATE_AXIS_NEGATIVE_Z;
+	LoadOpts.target_light_axes.right = UFBX_COORDINATE_AXIS_POSITIVE_Y;
+
 	ufbx_error Error;
 	Scene = ufbx_load_file(TCHAR_TO_UTF8(*Filename), &LoadOpts, &Error);
 
diff --git a/Engine/Plugins/Interchange/Runtime/Source/Parsers/Fbx/Private/Ufbx/UfbxParser.h b/Engine/Plugins/Interchange/Runtime/Source/Parsers/Fbx/Private/Ufbx/UfbxParser.h
--- a/Engine/Plugins/Interchange/Runtime/Source/Parsers/Fbx/Private/Ufbx/UfbxParser.h
+++ b/Engine/Plugins/Interchange/Runtime/Source/Parsers/Fbx/Private/Ufbx/UfbxParser.h
@@ -94,8 +94,10 @@ namespace UE::Interchange::Private
 		virtual bool IsThreadSafe() override { return true; }
 
 		virtual void SetResultContainer(UInterchangeResultsContainer* Result) override;
+
 		virtual void SetConvertSettings(const bool InbConvertScene, const bool InbForceFrontXAxis, const bool InbConvertSceneUnit, const bool InbKeepFbxNamespace, const bool InbConsiderClusterBeforePoseForMeshBindPose)  override
 		{
+			bForceFrontXAxis = InbForceFrontXAxis;
 			bKeepFbxNamespace = InbKeepFbxNamespace;
 			bConsiderClusterBeforePoseForMeshBindPose = InbConsiderClusterBeforePoseForMeshBindPose;
 		}
@@ -441,6 +443,7 @@ namespace UE::Interchange::Private
 			bool bIsJoint;
 		};
 
+		bool bForceFrontXAxis = false;
 		bool bKeepFbxNamespace = false;
 		bool bConsiderClusterBeforePoseForMeshBindPose = false;
 
diff --git a/Engine/Plugins/Interchange/Runtime/Source/Parsers/Fbx/Private/Ufbx/UfbxScene.cpp b/Engine/Plugins/Interchange/Runtime/Source/Parsers/Fbx/Private/Ufbx/UfbxScene.cpp
--- a/Engine/Plugins/Interchange/Runtime/Source/Parsers/Fbx/Private/Ufbx/UfbxScene.cpp
+++ b/Engine/Plugins/Interchange/Runtime/Source/Parsers/Fbx/Private/Ufbx/UfbxScene.cpp
@@ -747,7 +747,12 @@ namespace UE::Interchange::Private
 			UInterchangeSceneNode* SceneNode;
 
 			const ufbx_node* Node = Info.Node;
-			const FTransform Transform = Convert::ConvertTransform(Node->local_transform);
+			FTransform Transform = Convert::ConvertTransform(Node->local_transform);
+
+			if (Node->camera)
+			{
+				Transform.SetScale3D(FVector::OneVector);
+			}
 
 			if (Node->bone || IsInSkeleton(*Node))
 			{
