[Interchange][FBX] ufbx - recognize bones without skeleton attribute

- bones may be simple nodes, just used in skinning(i.e. found in skin clusters)

#Synced-CL 49917634

diff --git a/Engine/Plugins/Interchange/Runtime/Source/Parsers/Fbx/Private/Ufbx/UfbxScene.cpp b/Engine/Plugins/Interchange/Runtime/Source/Parsers/Fbx/Private/Ufbx/UfbxScene.cpp
--- a/Engine/Plugins/Interchange/Runtime/Source/Parsers/Fbx/Private/Ufbx/UfbxScene.cpp
+++ b/Engine/Plugins/Interchange/Runtime/Source/Parsers/Fbx/Private/Ufbx/UfbxScene.cpp
@@ -552,7 +552,7 @@ namespace UE::Interchange::Private
 		// Find root bone of each skeleton
 		for (const ufbx_node* Node : Parser.Scene->nodes)
 		{
-			if (Node->bone)
+			if (Node->bone || BoneToMeshIdToGeometryToBone.Contains(Node))
 			{
 				if (const ufbx_node* SkeletonRoot = SceneUtils::FindSkeletonRoot(Parser, Node))
 				{
