MaterialX:
- Add support for hextiledimage and hextilednormalmap nodes
- Add different utility material functions needed for the hextiling (Schlick Gain, gradient normals blending (Mikkelsen), hash function).
- Add Rotate3d material function making graphs easier to read

#virtualized
#Synced-CL 50150178

diff --git a/Engine/Plugins/Interchange/Runtime/Source/Common/Public/MaterialX/InterchangeMaterialXDefinitions.h b/Engine/Plugins/Interchange/Runtime/Source/Common/Public/MaterialX/InterchangeMaterialXDefinitions.h
--- a/Engine/Plugins/Interchange/Runtime/Source/Common/Public/MaterialX/InterchangeMaterialXDefinitions.h
+++ b/Engine/Plugins/Interchange/Runtime/Source/Common/Public/MaterialX/InterchangeMaterialXDefinitions.h
@@ -320,6 +320,8 @@ namespace UE
 			inline constexpr const TCHAR* MxColorCorrect = TEXT("/InterchangeAssets/Functions/MX_ColorCorrect.MX_ColorCorrect");
 			inline constexpr const TCHAR* MxDeonHairAbsorptionFromMelanin = TEXT("/InterchangeAssets/Functions/MX_DeonHairAbsorptionFromMelanin.MX_DeonHairAbsorptionFromMelanin");
 			inline constexpr const TCHAR* MxGoochShade = TEXT("/InterchangeAssets/Functions/MX_GoochShade.MX_GoochShade");
+			inline constexpr const TCHAR* MxHexTiledImage = TEXT("/InterchangeAssets/Functions/MX_HexTiledImage.MX_HexTiledImage");
+			inline constexpr const TCHAR* MxHexTiledNormalMap= TEXT("/InterchangeAssets/Functions/MX_HexTiledNormalMap.MX_HexTiledNormalMap");
 			inline constexpr const TCHAR* MxHsvAdjust = TEXT("/InterchangeAssets/Functions/MX_HSVAdjust.MX_HSVAdjust");
 			inline constexpr const TCHAR* MxLine = TEXT("/InterchangeAssets/Functions/MX_Line.MX_Line");
 			inline constexpr const TCHAR* MxLinearAdobeRGBToRec709 = TEXT("/InterchangeAssets/Functions/MX_LinearAdobeRGBToRec709.MX_LinearAdobeRGBToRec709");
@@ -331,6 +333,7 @@ namespace UE
 			inline constexpr const TCHAR* MxRampGradient = TEXT("/InterchangeAssets/Functions/MX_RampGradient.MX_RampGradient");
 			inline constexpr const TCHAR* MxRandomFloat = TEXT("/InterchangeAssets/Functions/MX_RandomFloat.MX_RandomFloat");
 			inline constexpr const TCHAR* MxRandomColor = TEXT("/InterchangeAssets/Functions/MX_RandomColor.MX_RandomColor");
+			inline constexpr const TCHAR* MxRotate3D = TEXT("/InterchangeAssets/Functions/MX_Rotate3D.MX_Rotate3D");
 			inline constexpr const TCHAR* MxRoughnesAnisotropy = TEXT("/InterchangeAssets/Functions/MX_Roughness_Anisotropy.MX_Roughness_Anisotropy");
 			inline constexpr const TCHAR* MxRoughnessDual = TEXT("/InterchangeAssets/Functions/MX_Roughness_Dual.MX_Roughness_Dual");
 			inline constexpr const TCHAR* MxSrgbDisplayP3ToRec709 = TEXT("/InterchangeAssets/Functions/MX_SrgbDisplayP3ToRec709.MX_SrgbDisplayP3ToRec709");
@@ -696,6 +699,8 @@ MATERIALX_NAMESPACE_BEGIN
 		inline constexpr const char* Exp = "exp";
 		inline constexpr const char* Floor = "floor";
 		inline constexpr const char* Fract = "fract";
+		inline constexpr const char* HexTiledImage = "hextiledimage";
+		inline constexpr const char* HexTiledNormalMap = "hextilednormalmap";
 		inline constexpr const char* Invert = "invert";
 		inline constexpr const char* InvertMatrix = "invertmatrix";
 		inline constexpr const char* Ln = "ln";
diff --git a/Engine/Plugins/Interchange/Runtime/Source/Import/Private/MaterialX/MaterialXUtils/MaterialXManager.cpp b/Engine/Plugins/Interchange/Runtime/Source/Import/Private/MaterialX/MaterialXUtils/MaterialXManager.cpp
--- a/Engine/Plugins/Interchange/Runtime/Source/Import/Private/MaterialX/MaterialXUtils/MaterialXManager.cpp
+++ b/Engine/Plugins/Interchange/Runtime/Source/Import/Private/MaterialX/MaterialXUtils/MaterialXManager.cpp
@@ -54,7 +54,7 @@ FMaterialXManager::FMaterialXManager()
 	   {{TEXT(""),                            TEXT("inhigh")},      ExpressionInput(UE::Expressions::Inputs::InputHigh)},
 	   {{TEXT(""),                            TEXT("lacunarity")},  ExpressionInput(UE::Expressions::Inputs::Lacunarity)},
 	   {{TEXT(""),                            TEXT("low")},         ExpressionInput(UE::Expressions::Inputs::Min)},
-	   {{TEXT(""),                            TEXT("lumacoeffs")},  ExpressionInput(UE::Expressions::Inputs::LuminanceFactors)}, // for the moment not yet handled by Interchange, because of the attribute being an advanced pin
+	   {{TEXT(""),                            TEXT("lumacoeffs")},  ExpressionInput(UE::Expressions::Inputs::LuminanceFactors)},
 	   {{TEXT(""),                            TEXT("mix")},         ExpressionInput(UE::Expressions::Inputs::Alpha)},
 	   {{TEXT(""),                            TEXT("offset")},      ExpressionInput(UE::Expressions::Inputs::Offset)},
 	   {{TEXT(""),                            TEXT("pivot")},       ExpressionInput(UE::Expressions::Inputs::Pivot)},
@@ -83,8 +83,8 @@ FMaterialXManager::FMaterialXManager()
 	   {{mx::Category::IfGreaterEq,           TEXT("in2")},         ExpressionInput(UE::Expressions::Inputs::ALessThanB)},
 	   {{mx::Category::IfEqual,               TEXT("in1")},         ExpressionInput(UE::Expressions::Inputs::AEqualsB)},
 	   {{mx::Category::IfEqual,               TEXT("in2")},         ExpressionInput(UE::Expressions::Inputs::ALessThanB)},     // another input is added for the case 'greater', see ConnectIfEqual
-	   {{mx::Category::Inside,                TEXT("in")},          ExpressionInput(UE::Expressions::Inputs::A)},			  // Inside is treated as a Multiply node
-	   {{mx::Category::Inside,                TEXT("mask")},        ExpressionInput(UE::Expressions::Inputs::B)},			  // Inside is treated as a Multiply node
+	   {{mx::Category::Inside,                TEXT("in")},          ExpressionInput(UE::Expressions::Inputs::A)},			   // Inside is treated as a Multiply node
+	   {{mx::Category::Inside,                TEXT("mask")},        ExpressionInput(UE::Expressions::Inputs::B)},			   // Inside is treated as a Multiply node
 	   {{mx::Category::Invert,                TEXT("amount")},      ExpressionInput(UE::Expressions::Inputs::A)},
 	   {{mx::Category::Invert,                TEXT("in")},          ExpressionInput(UE::Expressions::Inputs::B)},
 	   {{mx::Category::Mix,                   TEXT("fg")},          ExpressionInput(UE::Expressions::Inputs::B)},
@@ -97,8 +97,8 @@ FMaterialXManager::FMaterialXManager()
 	   {{mx::Category::Normalize,             TEXT("in")},          ExpressionInput(UE::Expressions::Inputs::VectorInput)},
 	   {{mx::Category::NormalMap,             TEXT("in")},          ExpressionInput(UE::Expressions::Inputs::Normal)},
 	   {{mx::Category::NormalMap,             TEXT("scale")},       ExpressionInput(UE::Expressions::Inputs::Flatness)},
-	   {{mx::Category::Outside,               TEXT("in")},          ExpressionInput(UE::Expressions::Inputs::A)},			  // Outside is treated as Multiply node
-	   {{mx::Category::Outside,               TEXT("mask")},        ExpressionInput(UE::Expressions::Inputs::B)},			  // Outside is treated as Multiply node
+	   {{mx::Category::Outside,               TEXT("in")},          ExpressionInput(UE::Expressions::Inputs::A)},			   // Outside is treated as Multiply node
+	   {{mx::Category::Outside,               TEXT("mask")},        ExpressionInput(UE::Expressions::Inputs::B)},			   // Outside is treated as Multiply node
 	   {{mx::Category::Power,                 TEXT("in1")},         ExpressionInput(UE::Expressions::Inputs::Base)},
 	   {{mx::Category::Power,                 TEXT("in2")},         ExpressionInput(UE::Expressions::Inputs::Exponent)},
 	   {{mx::Category::Refract,               TEXT("in")},          ExpressionInput(UE::Expressions::Inputs::RayDirection)},
@@ -106,9 +106,6 @@ FMaterialXManager::FMaterialXManager()
 	   {{mx::Category::Refract,               TEXT("ior")},         ExpressionInput(UE::Expressions::Inputs::RefractiveIndexOrigin)},
 	   {{mx::Category::Rotate2D,              TEXT("in")},          ExpressionInput(UE::Expressions::Inputs::Coordinate)},
 	   {{mx::Category::Rotate2D,              TEXT("amount")},      ExpressionInput(UE::Expressions::Inputs::Time)},
-	   {{mx::Category::Rotate3D,              TEXT("amount")},      ExpressionInput(UE::Expressions::Inputs::RotationAngle)},
-	   {{mx::Category::Rotate3D,              TEXT("axis")},		ExpressionInput(UE::Expressions::Inputs::NormalizedRotationAxis)},
-	   {{mx::Category::Rotate3D,              TEXT("in")},          ExpressionInput(UE::Expressions::Inputs::Position)},
 	   {{mx::Category::Smoothstep,            TEXT("in")},          ExpressionInput(UE::Expressions::Inputs::Value)},
 	   {{mx::Category::Separate2,             TEXT("in")},          ExpressionInput(UE::Expressions::Inputs::Float2)},
 	   {{mx::Category::Separate3,             TEXT("in")},          ExpressionInput(UE::Expressions::Inputs::Float3)},
@@ -257,6 +254,7 @@ FMaterialXManager::FMaterialXManager()
 		// Math																								
 		{mx::Category::Place2D,							FMaterialXMaterialFunction{TInPlaceType<FString>{}, MaterialFunctionPackage(UE::MaterialFunctions::Path::MxPlace2D)}},
 		{mx::Category::Refract,							FMaterialXMaterialFunction{TInPlaceType<FString>{}, MaterialFunctionPackage(UE::MaterialFunctions::Path::Refract)}},
+		{mx::Category::Rotate3D,						FMaterialXMaterialFunction{TInPlaceType<FString>{}, MaterialFunctionPackage(UE::MaterialFunctions::Path::MxRotate3D)}},
 		// ColorTransform																					
 		{mx::Category::ACEScgToLinRec709,				FMaterialXMaterialFunction{TInPlaceType<FString>{}, MaterialFunctionPackage(UE::MaterialFunctions::Path::MxACEScgToRec709)}},
 		{mx::Category::AdobeRgbToLinRec709,				FMaterialXMaterialFunction{TInPlaceType<FString>{}, MaterialFunctionPackage(UE::MaterialFunctions::Path::MxAdobeRGBToRec709)}},
@@ -340,7 +338,6 @@ FMaterialXManager::FMaterialXManager()
 		mx::Category::IfGreaterEq,
 		mx::Category::Max,
 		mx::Category::Min,
-		mx::Category::Rotate3D,
 		mx::Category::Sub,
 	}
 	, NodeInputsToRemove{
diff --git a/Engine/Plugins/Interchange/Runtime/Source/Import/Private/MaterialX/MaterialXUtils/MaterialXSurfaceShaderAbstract.cpp b/Engine/Plugins/Interchange/Runtime/Source/Import/Private/MaterialX/MaterialXUtils/MaterialXSurfaceShaderAbstract.cpp
--- a/Engine/Plugins/Interchange/Runtime/Source/Import/Private/MaterialX/MaterialXUtils/MaterialXSurfaceShaderAbstract.cpp
+++ b/Engine/Plugins/Interchange/Runtime/Source/Import/Private/MaterialX/MaterialXUtils/MaterialXSurfaceShaderAbstract.cpp
@@ -566,38 +566,6 @@ void FMaterialXSurfaceShaderAbstract::ConnectRotate2DInputToOutput(const FConnec
 	UInterchangeShaderPortsAPI::ConnectDefaultOuputToInput(Connect.ParentShaderNode, Connect.InputChannelName, Rotate2DNode->GetUniqueID());
 }
 
-void FMaterialXSurfaceShaderAbstract::ConnectRotate3DInputToOutput(const FConnectNode& Connect)
-{
-	using namespace UE::Interchange::Materials::Standard::Nodes;
-
-	UInterchangeShaderNode* Rotate3DNode = CreateShaderNode(Connect.UpstreamNode, Connect.UpstreamNode->getName().c_str(), RotateAboutAxis::Name.ToString());
-	Rotate3DNode->AddLinearColorAttribute(UInterchangeShaderPortsAPI::MakeInputValueKey(RotateAboutAxis::Inputs::PivotPoint.ToString()), FLinearColor(0., 0., 0));
-
-	mx::InputPtr Input = Connect.UpstreamNode->getInput("in");
-	AddAttributeFromValueOrInterface(Input, UE::Expressions::Inputs::Position, Rotate3DNode);
-	AddAttributeFromValueOrInterface(Connect.UpstreamNode->getInput("axis"), UE::Expressions::Inputs::NormalizedRotationAxis, Rotate3DNode);
-	AddAttributeFromValueOrInterface(Connect.UpstreamNode->getInput("amount"), UE::Expressions::Inputs::RotationAngle, Rotate3DNode);
-
-	// Convert Degrees into radians by setting the period to 360
-	Rotate3DNode->AddFloatAttribute(RotateAboutAxis::Attributes::Period.ToString(), 360.f);
-
-	//RotateAboutAxis returns an offset of the rotated vector, let's add to the vector to rotate it directly
-	FString AddNodeName{ (Connect.UpstreamNode->getName() + "_Add").c_str() };
-	UInterchangeShaderNode* AddNode = CreateShaderNode(Connect.UpstreamNode, AddNodeName, Add::Name.ToString());
-	UInterchangeShaderPortsAPI::ConnectDefaultOuputToInput(AddNode, Add::Inputs::A.ToString(), Rotate3DNode->GetUniqueID());
-
-	//Add another input to connect the Position to the Add node
-	if (!Connect.UpstreamNode->getInput("position_add"))
-	{
-		mx::InputPtr InputPositionAdd = Connect.UpstreamNode->addInput("position_add");
-		InputPositionAdd->copyContentFrom(Input);
-		InputPositionAdd->setAttribute(mx::Attributes::ParentName, TCHAR_TO_ANSI(*AddNodeName));
-		SetAttributeNewName(InputPositionAdd, TCHAR_TO_ANSI(UE::Expressions::Inputs::B));
-	}
-
-	UInterchangeShaderPortsAPI::ConnectDefaultOuputToInput(Connect.ParentShaderNode, Connect.InputChannelName, AddNode->GetUniqueID());
-}
-
 void FMaterialXSurfaceShaderAbstract::ConnectImageInputToOutput(const FConnectNode& Connect)
 {
 	using namespace UE::Interchange::Materials::Standard::Nodes;
@@ -1455,6 +1423,33 @@ void FMaterialXSurfaceShaderAbstract::ConnectLuminanceInputToOutput(const FConne
 	UInterchangeShaderPortsAPI::ConnectDefaultOuputToInput(Connect.ParentShaderNode, Connect.InputChannelName, LuminanceNode->GetUniqueID());
 }
 
+void FMaterialXSurfaceShaderAbstract::ConnectHexTiledInputToOutput(const FConnectNode& Connect)
+{
+	using namespace UE::Interchange::Materials::Standard::Nodes;
+
+	const bool bNormalTexture = Connect.UpstreamNode->getCategory() == mx::Category::HexTiledNormalMap ? true : false;
+	TGuardValue<bool>InputTypeBeingProcessedGuard(bTangentSpaceInput, bNormalTexture);
+
+	if (UInterchangeTextureNode* TextureNode = CreateTextureNode<UInterchangeTexture2DNode>(Connect.UpstreamNode))
+	{
+		const FString HexTileFunctionPath = Connect.UpstreamNode->getCategory() == mx::Category::HexTiledImage ? UE::MaterialFunctions::Path::MxHexTiledImage : UE::MaterialFunctions::Path::MxHexTiledNormalMap;
+		UInterchangeShaderNode* HexTileNode = CreateFunctionCallShaderNode(Connect.UpstreamNode, Connect.UpstreamNode->getName().c_str(), HexTileFunctionPath);
+
+		const FString TextureNodeName = Connect.UpstreamNode->getName().c_str() + FString{ "_texture" };
+		UInterchangeShaderNode* TextureShaderNode = CreateShaderNode(Connect.UpstreamNode, TextureNodeName, TextureObject::Name.ToString());
+		const FString TextureObjectInputKey = UInterchangeShaderPortsAPI::MakeInputParameterKey(TextureObject::Inputs::Texture.ToString());
+		TextureShaderNode->AddStringAttribute(TextureObjectInputKey, TextureNode->GetUniqueID());
+
+		for (mx::InputPtr Input : Connect.UpstreamNode->getInputs())
+		{
+			AddAttributeFromValueOrInterface(Input, GetInputName(Input), HexTileNode);
+		}
+
+		UInterchangeShaderPortsAPI::ConnectDefaultOuputToInput(HexTileNode, TEXT("Texture"), TextureShaderNode->GetUniqueID());
+		UInterchangeShaderPortsAPI::ConnectDefaultOuputToInput(Connect.ParentShaderNode, Connect.InputChannelName, HexTileNode->GetUniqueID());
+	}
+}
+
 bool FMaterialXSurfaceShaderAbstract::ConnectSubsurfaceTransmissionToSurfaceShader(UInterchangeShaderNode* OpenPbrStandarSurfaceShaderNode, EInterchangeMaterialXShaders SurfaceShaderType
 																				   , const FName& BaseWeight, const FName& BaseColor
 																				   , const FName& TransmissionWeight, const FName& TransmissionColor, const FName& TransmissionDepth, const FName& TransmissionDispersionAbbeNumber, const FName& TransmissionDispersionScale, const FName& TransmissionScatter, const FName& TransmissionScatterAnisotropy
@@ -1760,43 +1755,44 @@ FString FMaterialXSurfaceShaderAbstract::GetUniqueName(MaterialX::ElementPtr Ele
 
 void FMaterialXSurfaceShaderAbstract::RegisterConnectNodeOutputToInputDelegates()
 {
-	MatchingConnectNodeDelegates.Add(mx::Category::Constant,		FOnConnectNodeOutputToInput::CreateSP(this, &FMaterialXSurfaceShaderAbstract::ConnectConstantInputToOutput));
-	MatchingConnectNodeDelegates.Add(mx::Category::Extract,			FOnConnectNodeOutputToInput::CreateSP(this, &FMaterialXSurfaceShaderAbstract::ConnectExtractInputToOutput));
-	MatchingConnectNodeDelegates.Add(mx::Category::Dot,				FOnConnectNodeOutputToInput::CreateSP(this, &FMaterialXSurfaceShaderAbstract::ConnectDotInputToOutput));
-	MatchingConnectNodeDelegates.Add(mx::Category::NormalMap,		FOnConnectNodeOutputToInput::CreateSP(this, &FMaterialXSurfaceShaderAbstract::ConnectNormalMapInputToOutput));
-	MatchingConnectNodeDelegates.Add(mx::Category::TransformPoint,	FOnConnectNodeOutputToInput::CreateSP(this, &FMaterialXSurfaceShaderAbstract::ConnectTransformPositionInputToOutput));
-	MatchingConnectNodeDelegates.Add(mx::Category::TransformVector, FOnConnectNodeOutputToInput::CreateSP(this, &FMaterialXSurfaceShaderAbstract::ConnectTransformVectorInputToOutput));
-	MatchingConnectNodeDelegates.Add(mx::Category::TransformNormal, FOnConnectNodeOutputToInput::CreateSP(this, &FMaterialXSurfaceShaderAbstract::ConnectTransformVectorInputToOutput));
-	MatchingConnectNodeDelegates.Add(mx::Category::Rotate2D,		FOnConnectNodeOutputToInput::CreateSP(this, &FMaterialXSurfaceShaderAbstract::ConnectRotate2DInputToOutput));
-	MatchingConnectNodeDelegates.Add(mx::Category::Rotate3D,		FOnConnectNodeOutputToInput::CreateSP(this, &FMaterialXSurfaceShaderAbstract::ConnectRotate3DInputToOutput));
-	MatchingConnectNodeDelegates.Add(mx::Category::Image,			FOnConnectNodeOutputToInput::CreateSP(this, &FMaterialXSurfaceShaderAbstract::ConnectImageInputToOutput));
-	MatchingConnectNodeDelegates.Add(mx::Category::IfGreater,		FOnConnectNodeOutputToInput::CreateSP(this, &FMaterialXSurfaceShaderAbstract::ConnectIfGreaterInputToOutput));
-	MatchingConnectNodeDelegates.Add(mx::Category::IfGreaterEq,		FOnConnectNodeOutputToInput::CreateSP(this, &FMaterialXSurfaceShaderAbstract::ConnectIfGreaterEqInputToOutput));
-	MatchingConnectNodeDelegates.Add(mx::Category::IfEqual,			FOnConnectNodeOutputToInput::CreateSP(this, &FMaterialXSurfaceShaderAbstract::ConnectIfEqualInputToOutput));
-	MatchingConnectNodeDelegates.Add(mx::Category::Outside,			FOnConnectNodeOutputToInput::CreateSP(this, &FMaterialXSurfaceShaderAbstract::ConnectOutsideInputToOutput));
-	MatchingConnectNodeDelegates.Add(mx::Category::Position,		FOnConnectNodeOutputToInput::CreateSP(this, &FMaterialXSurfaceShaderAbstract::ConnectPositionInputToOutput));
-	MatchingConnectNodeDelegates.Add(mx::Category::Normal,			FOnConnectNodeOutputToInput::CreateSP(this, &FMaterialXSurfaceShaderAbstract::ConnectNormalInputToOutput));
-	MatchingConnectNodeDelegates.Add(mx::Category::Tangent,			FOnConnectNodeOutputToInput::CreateSP(this, &FMaterialXSurfaceShaderAbstract::ConnectTangentInputToOutput));
-	MatchingConnectNodeDelegates.Add(mx::Category::Bitangent,		FOnConnectNodeOutputToInput::CreateSP(this, &FMaterialXSurfaceShaderAbstract::ConnectBitangentInputToOutput));
-	MatchingConnectNodeDelegates.Add(mx::Category::Time,			FOnConnectNodeOutputToInput::CreateSP(this, &FMaterialXSurfaceShaderAbstract::ConnectTimeInputToOutput));
-	MatchingConnectNodeDelegates.Add(mx::Category::Noise2D,			FOnConnectNodeOutputToInput::CreateSP(this, &FMaterialXSurfaceShaderAbstract::ConnectNoise2DInputToOutput));
-	MatchingConnectNodeDelegates.Add(mx::Category::Noise3D,			FOnConnectNodeOutputToInput::CreateSP(this, &FMaterialXSurfaceShaderAbstract::ConnectNoise3DInputToOutput));
-	MatchingConnectNodeDelegates.Add(mx::Category::CellNoise2D,     FOnConnectNodeOutputToInput::CreateSP(this, &FMaterialXSurfaceShaderAbstract::ConnectCellNoise2DInputToOutput));
-	MatchingConnectNodeDelegates.Add(mx::Category::CellNoise3D,		FOnConnectNodeOutputToInput::CreateSP(this, &FMaterialXSurfaceShaderAbstract::ConnectCellNoise3DInputToOutput));
-	MatchingConnectNodeDelegates.Add(mx::Category::WorleyNoise2D,	FOnConnectNodeOutputToInput::CreateSP(this, &FMaterialXSurfaceShaderAbstract::ConnectWorleyNoise2DInputToOutput));
-	MatchingConnectNodeDelegates.Add(mx::Category::WorleyNoise3D,	FOnConnectNodeOutputToInput::CreateSP(this, &FMaterialXSurfaceShaderAbstract::ConnectWorleyNoise3DInputToOutput));
-	MatchingConnectNodeDelegates.Add(mx::Category::Blur,			FOnConnectNodeOutputToInput::CreateSP(this, &FMaterialXSurfaceShaderAbstract::ConnectBlurInputToOutput));
-	MatchingConnectNodeDelegates.Add(mx::Category::HeightToNormal,	FOnConnectNodeOutputToInput::CreateSP(this, &FMaterialXSurfaceShaderAbstract::ConnectHeightToNormalInputToOutput));
-	MatchingConnectNodeDelegates.Add(mx::Category::TexCoord,		FOnConnectNodeOutputToInput::CreateSP(this, &FMaterialXSurfaceShaderAbstract::ConnectTexCoordInputToOutput));
-	MatchingConnectNodeDelegates.Add(mx::Category::Switch,			FOnConnectNodeOutputToInput::CreateSP(this, &FMaterialXSurfaceShaderAbstract::ConnectSwitchInputToOutput));
-	MatchingConnectNodeDelegates.Add(mx::Category::Refract,			FOnConnectNodeOutputToInput::CreateSP(this, &FMaterialXSurfaceShaderAbstract::ConnectRefractInputToOutput));
-	MatchingConnectNodeDelegates.Add(mx::Category::ViewDirection,	FOnConnectNodeOutputToInput::CreateSP(this, &FMaterialXSurfaceShaderAbstract::ConnectViewDirectionInputToOutput));
-	MatchingConnectNodeDelegates.Add(mx::Category::And,				FOnConnectNodeOutputToInput::CreateSP(this, &FMaterialXSurfaceShaderAbstract::ConnectAndInputToOutput));
-	MatchingConnectNodeDelegates.Add(mx::Category::Or,				FOnConnectNodeOutputToInput::CreateSP(this, &FMaterialXSurfaceShaderAbstract::ConnectOrInputToOutput));
-	MatchingConnectNodeDelegates.Add(mx::Category::Xor,				FOnConnectNodeOutputToInput::CreateSP(this, &FMaterialXSurfaceShaderAbstract::ConnectXorInputToOutput));
-	MatchingConnectNodeDelegates.Add(mx::Category::Not,				FOnConnectNodeOutputToInput::CreateSP(this, &FMaterialXSurfaceShaderAbstract::ConnectNotInputToOutput));
-	MatchingConnectNodeDelegates.Add(mx::Category::SparseVolume,	FOnConnectNodeOutputToInput::CreateSP(this, &FMaterialXSurfaceShaderAbstract::ConnectSparseVolumeInputToOutput));
-	MatchingConnectNodeDelegates.Add(mx::Category::Luminance,		FOnConnectNodeOutputToInput::CreateSP(this, &FMaterialXSurfaceShaderAbstract::ConnectLuminanceInputToOutput));
+	MatchingConnectNodeDelegates.Add(mx::Category::Constant,		  FOnConnectNodeOutputToInput::CreateSP(this, &FMaterialXSurfaceShaderAbstract::ConnectConstantInputToOutput));
+	MatchingConnectNodeDelegates.Add(mx::Category::Extract,			  FOnConnectNodeOutputToInput::CreateSP(this, &FMaterialXSurfaceShaderAbstract::ConnectExtractInputToOutput));
+	MatchingConnectNodeDelegates.Add(mx::Category::Dot,				  FOnConnectNodeOutputToInput::CreateSP(this, &FMaterialXSurfaceShaderAbstract::ConnectDotInputToOutput));
+	MatchingConnectNodeDelegates.Add(mx::Category::NormalMap,		  FOnConnectNodeOutputToInput::CreateSP(this, &FMaterialXSurfaceShaderAbstract::ConnectNormalMapInputToOutput));
+	MatchingConnectNodeDelegates.Add(mx::Category::TransformPoint,	  FOnConnectNodeOutputToInput::CreateSP(this, &FMaterialXSurfaceShaderAbstract::ConnectTransformPositionInputToOutput));
+	MatchingConnectNodeDelegates.Add(mx::Category::TransformVector,   FOnConnectNodeOutputToInput::CreateSP(this, &FMaterialXSurfaceShaderAbstract::ConnectTransformVectorInputToOutput));
+	MatchingConnectNodeDelegates.Add(mx::Category::TransformNormal,   FOnConnectNodeOutputToInput::CreateSP(this, &FMaterialXSurfaceShaderAbstract::ConnectTransformVectorInputToOutput));
+	MatchingConnectNodeDelegates.Add(mx::Category::Rotate2D,		  FOnConnectNodeOutputToInput::CreateSP(this, &FMaterialXSurfaceShaderAbstract::ConnectRotate2DInputToOutput));
+	MatchingConnectNodeDelegates.Add(mx::Category::Image,			  FOnConnectNodeOutputToInput::CreateSP(this, &FMaterialXSurfaceShaderAbstract::ConnectImageInputToOutput));
+	MatchingConnectNodeDelegates.Add(mx::Category::IfGreater,		  FOnConnectNodeOutputToInput::CreateSP(this, &FMaterialXSurfaceShaderAbstract::ConnectIfGreaterInputToOutput));
+	MatchingConnectNodeDelegates.Add(mx::Category::IfGreaterEq,		  FOnConnectNodeOutputToInput::CreateSP(this, &FMaterialXSurfaceShaderAbstract::ConnectIfGreaterEqInputToOutput));
+	MatchingConnectNodeDelegates.Add(mx::Category::IfEqual,			  FOnConnectNodeOutputToInput::CreateSP(this, &FMaterialXSurfaceShaderAbstract::ConnectIfEqualInputToOutput));
+	MatchingConnectNodeDelegates.Add(mx::Category::Outside,			  FOnConnectNodeOutputToInput::CreateSP(this, &FMaterialXSurfaceShaderAbstract::ConnectOutsideInputToOutput));
+	MatchingConnectNodeDelegates.Add(mx::Category::Position,		  FOnConnectNodeOutputToInput::CreateSP(this, &FMaterialXSurfaceShaderAbstract::ConnectPositionInputToOutput));
+	MatchingConnectNodeDelegates.Add(mx::Category::Normal,			  FOnConnectNodeOutputToInput::CreateSP(this, &FMaterialXSurfaceShaderAbstract::ConnectNormalInputToOutput));
+	MatchingConnectNodeDelegates.Add(mx::Category::Tangent,			  FOnConnectNodeOutputToInput::CreateSP(this, &FMaterialXSurfaceShaderAbstract::ConnectTangentInputToOutput));
+	MatchingConnectNodeDelegates.Add(mx::Category::Bitangent,		  FOnConnectNodeOutputToInput::CreateSP(this, &FMaterialXSurfaceShaderAbstract::ConnectBitangentInputToOutput));
+	MatchingConnectNodeDelegates.Add(mx::Category::Time,			  FOnConnectNodeOutputToInput::CreateSP(this, &FMaterialXSurfaceShaderAbstract::ConnectTimeInputToOutput));
+	MatchingConnectNodeDelegates.Add(mx::Category::Noise2D,			  FOnConnectNodeOutputToInput::CreateSP(this, &FMaterialXSurfaceShaderAbstract::ConnectNoise2DInputToOutput));
+	MatchingConnectNodeDelegates.Add(mx::Category::Noise3D,			  FOnConnectNodeOutputToInput::CreateSP(this, &FMaterialXSurfaceShaderAbstract::ConnectNoise3DInputToOutput));
+	MatchingConnectNodeDelegates.Add(mx::Category::CellNoise2D,       FOnConnectNodeOutputToInput::CreateSP(this, &FMaterialXSurfaceShaderAbstract::ConnectCellNoise2DInputToOutput));
+	MatchingConnectNodeDelegates.Add(mx::Category::CellNoise3D,		  FOnConnectNodeOutputToInput::CreateSP(this, &FMaterialXSurfaceShaderAbstract::ConnectCellNoise3DInputToOutput));
+	MatchingConnectNodeDelegates.Add(mx::Category::WorleyNoise2D,	  FOnConnectNodeOutputToInput::CreateSP(this, &FMaterialXSurfaceShaderAbstract::ConnectWorleyNoise2DInputToOutput));
+	MatchingConnectNodeDelegates.Add(mx::Category::WorleyNoise3D,	  FOnConnectNodeOutputToInput::CreateSP(this, &FMaterialXSurfaceShaderAbstract::ConnectWorleyNoise3DInputToOutput));
+	MatchingConnectNodeDelegates.Add(mx::Category::Blur,			  FOnConnectNodeOutputToInput::CreateSP(this, &FMaterialXSurfaceShaderAbstract::ConnectBlurInputToOutput));
+	MatchingConnectNodeDelegates.Add(mx::Category::HeightToNormal,	  FOnConnectNodeOutputToInput::CreateSP(this, &FMaterialXSurfaceShaderAbstract::ConnectHeightToNormalInputToOutput));
+	MatchingConnectNodeDelegates.Add(mx::Category::TexCoord,		  FOnConnectNodeOutputToInput::CreateSP(this, &FMaterialXSurfaceShaderAbstract::ConnectTexCoordInputToOutput));
+	MatchingConnectNodeDelegates.Add(mx::Category::Switch,			  FOnConnectNodeOutputToInput::CreateSP(this, &FMaterialXSurfaceShaderAbstract::ConnectSwitchInputToOutput));
+	MatchingConnectNodeDelegates.Add(mx::Category::Refract,			  FOnConnectNodeOutputToInput::CreateSP(this, &FMaterialXSurfaceShaderAbstract::ConnectRefractInputToOutput));
+	MatchingConnectNodeDelegates.Add(mx::Category::ViewDirection,	  FOnConnectNodeOutputToInput::CreateSP(this, &FMaterialXSurfaceShaderAbstract::ConnectViewDirectionInputToOutput));
+	MatchingConnectNodeDelegates.Add(mx::Category::And,				  FOnConnectNodeOutputToInput::CreateSP(this, &FMaterialXSurfaceShaderAbstract::ConnectAndInputToOutput));
+	MatchingConnectNodeDelegates.Add(mx::Category::Or,				  FOnConnectNodeOutputToInput::CreateSP(this, &FMaterialXSurfaceShaderAbstract::ConnectOrInputToOutput));
+	MatchingConnectNodeDelegates.Add(mx::Category::Xor,				  FOnConnectNodeOutputToInput::CreateSP(this, &FMaterialXSurfaceShaderAbstract::ConnectXorInputToOutput));
+	MatchingConnectNodeDelegates.Add(mx::Category::Not,				  FOnConnectNodeOutputToInput::CreateSP(this, &FMaterialXSurfaceShaderAbstract::ConnectNotInputToOutput));
+	MatchingConnectNodeDelegates.Add(mx::Category::SparseVolume,	  FOnConnectNodeOutputToInput::CreateSP(this, &FMaterialXSurfaceShaderAbstract::ConnectSparseVolumeInputToOutput));
+	MatchingConnectNodeDelegates.Add(mx::Category::Luminance,		  FOnConnectNodeOutputToInput::CreateSP(this, &FMaterialXSurfaceShaderAbstract::ConnectLuminanceInputToOutput));
+	MatchingConnectNodeDelegates.Add(mx::Category::HexTiledImage,	  FOnConnectNodeOutputToInput::CreateSP(this, &FMaterialXSurfaceShaderAbstract::ConnectHexTiledInputToOutput));
+	MatchingConnectNodeDelegates.Add(mx::Category::HexTiledNormalMap, FOnConnectNodeOutputToInput::CreateSP(this, &FMaterialXSurfaceShaderAbstract::ConnectHexTiledInputToOutput));
 }
 
 void FMaterialXSurfaceShaderAbstract::SetMatchingInputsNames(MaterialX::NodePtr Node) const
diff --git a/Engine/Plugins/Interchange/Runtime/Source/Import/Public/MaterialX/MaterialXUtils/MaterialXSurfaceShaderAbstract.h b/Engine/Plugins/Interchange/Runtime/Source/Import/Public/MaterialX/MaterialXUtils/MaterialXSurfaceShaderAbstract.h
--- a/Engine/Plugins/Interchange/Runtime/Source/Import/Public/MaterialX/MaterialXUtils/MaterialXSurfaceShaderAbstract.h
+++ b/Engine/Plugins/Interchange/Runtime/Source/Import/Public/MaterialX/MaterialXUtils/MaterialXSurfaceShaderAbstract.h
@@ -252,9 +252,6 @@ protected:
 	/** <rotate2d> */
 	void ConnectRotate2DInputToOutput(const FConnectNode& Connect);
 
-	/** <rotate3d> */
-	void ConnectRotate3DInputToOutput(const FConnectNode& Connect);
-
 	/** <image> */
 	void ConnectImageInputToOutput(const FConnectNode& Connect);
 
@@ -341,6 +338,11 @@ protected:
 
 	/** <luminance>*/
 	void ConnectLuminanceInputToOutput(const FConnectNode& Connect);
+
+	/** <hextiledimage>
+	  * <hextilednormalmap>
+	  */
+	void ConnectHexTiledInputToOutput(const FConnectNode& Connect);
 	/** End Connect MaterialX nodes*/
 
 	/** Function used only by both Standard Surface and OpenPBR to have both Subsurface and Transmission*/
