MaterialX:
- Add support for <latlongimage> node
- Refactor MaterialExpressionMaterialXFractal3D to support fractal2d nodes and float, float2, float3 and float4 outputs

#virtualized
#Synced-CL 50150178

diff --git a/Engine/Plugins/Interchange/Runtime/Config/DefaultInterchange.ini b/Engine/Plugins/Interchange/Runtime/Config/DefaultInterchange.ini
--- a/Engine/Plugins/Interchange/Runtime/Config/DefaultInterchange.ini
+++ b/Engine/Plugins/Interchange/Runtime/Config/DefaultInterchange.ini
@@ -43,6 +43,7 @@ bCanSaveAllSections=True
 +EnumRedirects=(OldName="/Script/InterchangePipelines.EInterchangeMaterialXBSDF>", NewName="/Script/InterchangeCommon.EInterchangeMaterialXBSDF>")
 +EnumRedirects=(OldName="/Script/InterchangePipelines.EInterchangeMaterialXEDF>", NewName="/Script/InterchangeCommon.EInterchangeMaterialXEDF>")
 +EnumRedirects=(OldName="/Script/InterchangePipelines.EInterchangeMaterialXVDF>", NewName="/Script/InterchangeCommon.EInterchangeMaterialXVDF>")
++ClassRedirects=(OldName="/Script/InterchangeImport.MaterialExpressionMaterialXFractal3D", NewName="/Script/InterchangeImport.MaterialExpressionMaterialXFractal")
 
 ; Rename InterchangeAnimationTrackSetFactory to InterchangeLevelSequenceFactory
 +ClassRedirects=(OldName="/Script/InterchangeImport.InterchangeAnimationTrackSetFactory", NewName="/Script/InterchangeImport.InterchangeLevelSequenceFactory")
diff --git a/Engine/Plugins/Interchange/Runtime/Source/Common/Public/MaterialX/InterchangeMaterialXDefinitions.h b/Engine/Plugins/Interchange/Runtime/Source/Common/Public/MaterialX/InterchangeMaterialXDefinitions.h
--- a/Engine/Plugins/Interchange/Runtime/Source/Common/Public/MaterialX/InterchangeMaterialXDefinitions.h
+++ b/Engine/Plugins/Interchange/Runtime/Source/Common/Public/MaterialX/InterchangeMaterialXDefinitions.h
@@ -176,7 +176,7 @@ namespace UE
 			inline constexpr const TCHAR* MaterialXDifference = TEXT("MaterialXDifference");
 			inline constexpr const TCHAR* MaterialXDisjointover = TEXT("MaterialXDisjointover");
 			inline constexpr const TCHAR* MaterialXDodge = TEXT("MaterialXDodge");
-			inline constexpr const TCHAR* MaterialXFractal3D = TEXT("MaterialXFractal3D");
+			inline constexpr const TCHAR* MaterialXFractal = TEXT("MaterialXFractal");
 			inline constexpr const TCHAR* MaterialXIn = TEXT("MaterialXIn");
 			inline constexpr const TCHAR* MaterialXLuminance = TEXT("MaterialXLuminance");
 			inline constexpr const TCHAR* MaterialXMask = TEXT("MaterialXMask");
@@ -323,6 +323,7 @@ namespace UE
 			inline constexpr const TCHAR* MxHexTiledImage = TEXT("/InterchangeAssets/Functions/MX_HexTiledImage.MX_HexTiledImage");
 			inline constexpr const TCHAR* MxHexTiledNormalMap= TEXT("/InterchangeAssets/Functions/MX_HexTiledNormalMap.MX_HexTiledNormalMap");
 			inline constexpr const TCHAR* MxHsvAdjust = TEXT("/InterchangeAssets/Functions/MX_HSVAdjust.MX_HSVAdjust");
+			inline constexpr const TCHAR* MxLatLongImage = TEXT("/InterchangeAssets/Functions/MX_LatLongImage.MX_LatLongImage");
 			inline constexpr const TCHAR* MxLine = TEXT("/InterchangeAssets/Functions/MX_Line.MX_Line");
 			inline constexpr const TCHAR* MxLinearAdobeRGBToRec709 = TEXT("/InterchangeAssets/Functions/MX_LinearAdobeRGBToRec709.MX_LinearAdobeRGBToRec709");
 			inline constexpr const TCHAR* MxLinearDisplayP3ToRec709 = TEXT("/InterchangeAssets/Functions/MX_LinearDisplayP3ToRec709.MX_LinearDisplayP3ToRec709");
@@ -699,8 +700,6 @@ MATERIALX_NAMESPACE_BEGIN
 		inline constexpr const char* Exp = "exp";
 		inline constexpr const char* Floor = "floor";
 		inline constexpr const char* Fract = "fract";
-		inline constexpr const char* HexTiledImage = "hextiledimage";
-		inline constexpr const char* HexTiledNormalMap = "hextilednormalmap";
 		inline constexpr const char* Invert = "invert";
 		inline constexpr const char* InvertMatrix = "invertmatrix";
 		inline constexpr const char* Ln = "ln";
@@ -773,6 +772,7 @@ MATERIALX_NAMESPACE_BEGIN
 		inline constexpr const char* CellNoise2D = "cellnoise2d";
 		inline constexpr const char* Checkerboard = "checkerboard";
 		inline constexpr const char* Circle = "circle";
+		inline constexpr const char* Fractal2D = "fractal2d";
 		inline constexpr const char* Line = "line";
 		inline constexpr const char* Noise2D = "noise2d";
 		inline constexpr const char* Ramp4 = "ramp4";
@@ -795,7 +795,10 @@ MATERIALX_NAMESPACE_BEGIN
 		// Organization nodes 
 		inline constexpr const char* Dot = "dot";
 		// Texture nodes 
+		inline constexpr const char* HexTiledImage = "hextiledimage";
+		inline constexpr const char* HexTiledNormalMap = "hextilednormalmap";
 		inline constexpr const char* Image = "image";
+		inline constexpr const char* LatLongImage = "latlongimage";
 		inline constexpr const char* TiledImage = "tiledimage";
 		// Geometric nodes
 		inline constexpr const char* Bitangent = "bitangent";
diff --git a/Engine/Plugins/Interchange/Runtime/Source/Import/Private/MaterialX/MaterialExpressions/MaterialExpressionFractal.cpp b/Engine/Plugins/Interchange/Runtime/Source/Import/Private/MaterialX/MaterialExpressions/MaterialExpressionFractal.cpp
new file mode 100644
--- /dev/null
+++ b/Engine/Plugins/Interchange/Runtime/Source/Import/Private/MaterialX/MaterialExpressions/MaterialExpressionFractal.cpp
@@ -0,0 +1,155 @@
+// Copyright Epic Games, Inc. All Rights Reserved.
+
+#include "MaterialExpressionFractal.h"
+#include "MaterialCompiler.h"
+#include "Materials/MaterialExpressionCustom.h"
+
+#include UE_INLINE_GENERATED_CPP_BY_NAME(MaterialExpressionFractal)
+
+#define LOCTEXT_NAMESPACE "MaterialExpressionMaterialXFractal"
+
+UMaterialExpressionMaterialXFractal::UMaterialExpressionMaterialXFractal(const FObjectInitializer& ObjectInitializer)
+	: Super(ObjectInitializer)
+{
+	// Structure to hold one-time initialization
+	struct FConstructorStatics
+	{
+		FText NAME_MaterialX;
+		FConstructorStatics()
+			: NAME_MaterialX(LOCTEXT("MaterialX", "MaterialX"))
+		{}
+	};
+	static FConstructorStatics ConstructorStatics;
+
+#if WITH_EDITORONLY_DATA
+	Outputs.Reset();
+	Outputs.Add(FExpressionOutput(TEXT("R"), 1, 1, 0, 0, 0));
+	Outputs.Add(FExpressionOutput(TEXT("RG"), 1, 1, 1, 0, 0));
+	Outputs.Add(FExpressionOutput(TEXT("RGB"), 1, 1, 1, 1, 0));
+	Outputs.Add(FExpressionOutput(TEXT("RGBA"), 1, 1, 1, 1, 1));
+	bShowOutputNameOnPin = true;
+	bShowMaskColorsOnPin = false;
+
+	MenuCategories.Add(ConstructorStatics.NAME_MaterialX);
+#endif
+}
+
+#if WITH_EDITOR
+int32 UMaterialExpressionMaterialXFractal::Compile(FMaterialCompiler* Compiler, int32 OutputIndex)
+{
+	int32 IndexPosition = Position.GetTracedInput().Expression ? Position.Compile(Compiler) : Compiler->LocalPosition(EPositionIncludedOffsets::IncludeOffsets, ELocalPositionOrigin::Instance);
+
+	if (IndexPosition == INDEX_NONE)
+	{
+		return Compiler->Errorf(TEXT("Failed to compile Position input."));
+	}
+
+	bFractal2dType = Compiler->GetType(IndexPosition) == EMaterialValueType::MCT_Float2;
+	if (bFractal2dType)
+	{
+		IndexPosition = Compiler->AppendVector(IndexPosition, Compiler->Constant(0.));
+	}
+
+	UMaterialExpressionCustom* MaterialExpressionCustom = NewObject<UMaterialExpressionCustom>();
+	MaterialExpressionCustom->Inputs[0].InputName = TEXT("Position");
+	MaterialExpressionCustom->Inputs.Add({ TEXT("Octaves") });
+	MaterialExpressionCustom->Inputs.Add({ TEXT("Lacunarity") });
+	MaterialExpressionCustom->Inputs.Add({ TEXT("Diminish") });
+	
+	FString Result = TEXT("float result = 0.0;");
+	FString Noise = TEXT("result.x += amplitude * GradientNoise3D_TEX(Position, bTiling, RepeatSize);\n");
+
+	// Decorrelating offsets for multi-channel Perlin noise.
+	// Chosen to avoid lattice alignment and visible correlation in grid-based noise.
+	// Follows common practice used in MaterialX and Perlin-style implementations.
+	switch (OutputIndex)
+	{
+	case 1:
+		MaterialExpressionCustom->OutputType = ECustomMaterialOutputType::CMOT_Float2;
+		Result = TEXT("float2 result = float2(0.0, 0.0);\n");
+		Noise += TEXT("result.y += amplitude * GradientNoise3D_TEX(Position + float3(19,193,") + FString{ bFractal2dType ? TEXT("0") : TEXT("17") } + TEXT("), bTiling, RepeatSize);\n");
+		break;
+	case 2:
+		MaterialExpressionCustom->OutputType = ECustomMaterialOutputType::CMOT_Float3;
+		Result = TEXT("float3 result = float3(0.0, 0.0, 0.0);\n");
+		Noise += TEXT("result.y += amplitude * GradientNoise3D_TEX(Position + float3(19,193,") + FString{ bFractal2dType ? TEXT("0") : TEXT("17") } + TEXT("), bTiling, RepeatSize);\n");
+		Noise += TEXT("result.z += amplitude * GradientNoise3D_TEX(Position + float3(73,27,") + FString{ bFractal2dType ? TEXT("0") : TEXT("101") } + TEXT("), bTiling, RepeatSize);\n");
+		break;
+	case 3:
+		MaterialExpressionCustom->OutputType = ECustomMaterialOutputType::CMOT_Float4;
+		Result = TEXT("float4 result = float4(0.0, 0.0, 0.0, 0.0);\n");
+		Noise += TEXT("result.y += amplitude * GradientNoise3D_TEX(Position + float3(19,193,") + FString{ bFractal2dType ? TEXT("0")  : TEXT("17") } + TEXT("), bTiling, RepeatSize);\n");
+		Noise += TEXT("result.z += amplitude * GradientNoise3D_TEX(Position + float3(73,27,") + FString{ bFractal2dType ? TEXT("0") : TEXT("101") } + TEXT("), bTiling, RepeatSize);\n ");
+		Noise += TEXT("result.w += amplitude * GradientNoise3D_TEX(Position + float3(151,89,") + FString{ bFractal2dType ? TEXT("0") : TEXT("7") } + TEXT("), bTiling, RepeatSize);\n");
+
+		break;
+	default:
+		MaterialExpressionCustom->OutputType = ECustomMaterialOutputType::CMOT_Float1;
+		break;
+	}
+	
+	MaterialExpressionCustom->Code =
+		TEXT(R"(const bool bTiling = false;
+        const float RepeatSize = 512;
+	    )")
+		+ Result +
+TEXT(R"(float amplitude = 1.0;
+		for (int i = 0;  i < Octaves; ++i)
+		{
+         )") 
+		 + Noise +
+TEXT(R"(
+			amplitude *= Diminish;
+			Position *= Lacunarity;
+		}
+		return result;)");
+
+	int32 IndexOctaves = Octaves.GetTracedInput().Expression ? Octaves.Compile(Compiler) : Compiler->Constant(ConstOctaves);
+	int32 IndexLacunarity = Lacunarity.GetTracedInput().Expression ? Lacunarity.Compile(Compiler) : Compiler->Constant(ConstLacunarity);
+	int32 IndexDiminish = Diminish.GetTracedInput().Expression ? Diminish.Compile(Compiler) : Compiler->Constant(ConstDiminish);
+
+	TArray<int32> Inputs{ IndexPosition, IndexOctaves, IndexLacunarity, IndexDiminish };
+	int32 IndexFractal = Compiler->CustomExpression(MaterialExpressionCustom, 0, Inputs);
+
+	int32 IndexAmplitude = Amplitude.GetTracedInput().Expression ? Amplitude.Compile(Compiler) : Compiler->Constant(ConstAmplitude);
+	
+	return Compiler->Mul(IndexFractal, IndexAmplitude);
+}
+
+void UMaterialExpressionMaterialXFractal::GetCaption(TArray<FString>& OutCaptions) const
+{
+	if (bFractal2dType)
+	{
+		OutCaptions.Add(TEXT("MaterialX Fractal2D"));
+	}
+	else
+	{
+		OutCaptions.Add(TEXT("MaterialX Fractal3D"));
+	}
+}
+
+FName UMaterialExpressionMaterialXFractal::GetInputName(int32 InputIndex) const
+{
+	// caveat : bFractal2dType is set during Compile, but the editor doesn't refresh the pins after it.
+	// it will only work at creation-time or if the Material Editor is completely refreshed
+	if (InputIndex == 0 && bFractal2dType)
+	{
+		return TEXT("Coordinates");
+	}
+	return Super::GetInputName(InputIndex);
+}
+
+void UMaterialExpressionMaterialXFractal::GetConnectorToolTip(int32 InputIndex, int32 OutputIndex, TArray<FString>& OutToolTip)
+{
+	if (InputIndex == 0 && bFractal2dType)
+	{
+		OutToolTip.Add(TEXT("The 2D texture coordinate at which the noise is evaluated."));
+	}
+	else
+	{
+		Super::GetConnectorToolTip(InputIndex, OutputIndex, OutToolTip);
+	}
+}
+#endif
+
+#undef LOCTEXT_NAMESPACE 
\ No newline at end of file
diff --git a/Engine/Plugins/Interchange/Runtime/Source/Import/Private/MaterialX/MaterialExpressions/MaterialExpressionFractal.h b/Engine/Plugins/Interchange/Runtime/Source/Import/Private/MaterialX/MaterialExpressions/MaterialExpressionFractal.h
new file mode 100644
--- /dev/null
+++ b/Engine/Plugins/Interchange/Runtime/Source/Import/Private/MaterialX/MaterialExpressions/MaterialExpressionFractal.h
@@ -0,0 +1,82 @@
+// Copyright Epic Games, Inc. All Rights Reserved.
+#pragma once
+
+#include "Materials/MaterialExpression.h"
+
+#include "MaterialExpressionFractal.generated.h"
+
+/**
+ * Zero-centered 2D or 3D Fractal noise in 1, 2, 3 or 4 channels, created by summing several
+ * octaves of 2D or 3D Perlin noise, increasing the frequency and decreasing the amplitude at each octave.
+ * Defaulting to a fractal 3D. Inputting a 2D vector will create a fractal 2D.
+ */
+
+UCLASS(collapsecategories, hidecategories = Object, MinimalAPI, meta = (Private))
+class UMaterialExpressionMaterialXFractal : public UMaterialExpression
+{
+	GENERATED_UCLASS_BODY()
+
+	/** The 3D position at which the noise is evaluated. By default the vector is in local space*/
+	UPROPERTY()
+	FExpressionInput Position;	
+
+	/** Center-to-peak amplitude of the noise (peak-to-peak amplitude is 2x this value).*/
+	UPROPERTY(meta = (RequiredInput = "false", ToolTip = "Defaults to 'ConstAmplitude' if not specified"))
+	FExpressionInput Amplitude;
+
+	/** only used if Amplitude is not hooked up */
+	UPROPERTY(EditAnywhere, Category = MaterialExpressionFractal, meta = (OverridingInputProperty = "Amplitude"))
+	float ConstAmplitude = 1;
+
+	/** The number of octaves of noise to be summed. */
+	UPROPERTY(meta = (RequiredInput = "false", ToolTip = "Defaults to 'ConstOctaves' if not specified"))
+	FExpressionInput Octaves;
+
+	/** only used if Octaves is not hooked up */
+	UPROPERTY(EditAnywhere, Category = MaterialExpressionFractal, meta = (OverridingInputProperty = "Octaves"))
+	int32 ConstOctaves = 3;
+
+	/** The exponential scale between successive octaves of noise. */
+	UPROPERTY(meta = (RequiredInput = "false", ToolTip = "Defaults to 'ConstLacunarity' if not specified"))
+	FExpressionInput Lacunarity;
+
+	/** only used if Lacunarity is not hooked up */
+	UPROPERTY(EditAnywhere, Category = MaterialExpressionFractal, meta = (OverridingInputProperty = "Lacunarity"))
+	float ConstLacunarity = 2;
+
+	/** The rate at which noise amplitude is diminished for each octave. Should be between 0.0 and 1.0 */
+	UPROPERTY(meta = (RequiredInput = "false", ToolTip = "Defaults to 'ConstDiminish' if not specified"))
+	FExpressionInput Diminish;
+
+	/** only used if Diminish is not hooked up */
+	UPROPERTY(EditAnywhere, Category = MaterialExpressionFractal, meta = (OverridingInputProperty = "Diminish"))
+	float ConstDiminish = 0.5;
+
+	UPROPERTY()
+	float Scale_DEPRECATED = 1.f;
+
+	UPROPERTY()
+	bool bTurbulence_DEPRECATED = false;
+
+	UPROPERTY()
+	int32 Levels_DEPRECATED = 6;
+
+	UPROPERTY()
+	float OutputMin_DEPRECATED = 0.f;
+
+	UPROPERTY()
+	float OutputMax_DEPRECATED = 1.f;
+
+	//~ Begin UMaterialExpressionMaterialX Interface
+#if WITH_EDITOR
+	virtual int32 Compile(class FMaterialCompiler* Compiler, int32 OutputIndex) override;
+	virtual void GetCaption(TArray<FString>& OutCaptions) const override;
+	virtual FName GetInputName(int32 InputIndex) const override;
+	virtual void GetConnectorToolTip(int32 InputIndex, int32 OutputIndex, TArray<FString>& OutToolTip) override;
+
+private:
+	bool bFractal2dType = false;
+#endif
+	//~ End UMaterialExpressionMaterialX Interface
+};
+
diff --git a/Engine/Plugins/Interchange/Runtime/Source/Import/Private/MaterialX/MaterialExpressions/MaterialExpressionFractal3D.cpp b/Engine/Plugins/Interchange/Runtime/Source/Import/Private/MaterialX/MaterialExpressions/MaterialExpressionFractal3D.cpp
deleted file mode 100644
diff --git a/Engine/Plugins/Interchange/Runtime/Source/Import/Private/MaterialX/MaterialExpressions/MaterialExpressionFractal3D.h b/Engine/Plugins/Interchange/Runtime/Source/Import/Private/MaterialX/MaterialExpressions/MaterialExpressionFractal3D.h
deleted file mode 100644
diff --git a/Engine/Plugins/Interchange/Runtime/Source/Import/Private/MaterialX/MaterialXUtils/MaterialXManager.cpp b/Engine/Plugins/Interchange/Runtime/Source/Import/Private/MaterialX/MaterialXUtils/MaterialXManager.cpp
--- a/Engine/Plugins/Interchange/Runtime/Source/Import/Private/MaterialX/MaterialXUtils/MaterialXManager.cpp
+++ b/Engine/Plugins/Interchange/Runtime/Source/Import/Private/MaterialX/MaterialXUtils/MaterialXManager.cpp
@@ -76,6 +76,7 @@ FMaterialXManager::FMaterialXManager()
 	   {{TEXT(""),                            TEXT("value2")},      ExpressionInput(UE::Expressions::Inputs::B)},
 	   {{mx::Category::Atan2,                 TEXT("iny")},         ExpressionInput(UE::Expressions::Inputs::Y)},
 	   {{mx::Category::Atan2,                 TEXT("inx")},         ExpressionInput(UE::Expressions::Inputs::X)},
+	   {{mx::Category::Fractal2D,             TEXT("texcoord")},    ExpressionInput(UE::Expressions::Inputs::Position)},
 	   {{mx::Category::HeightToNormal,        TEXT("scale")},       ExpressionInput(*UE::Interchange::Materials::Standard::Nodes::NormalFromHeightMap::Inputs::Intensity.ToString())},
 	   {{mx::Category::IfGreater,             TEXT("in1")},         ExpressionInput(UE::Expressions::Inputs::AGreaterThanB)},
 	   {{mx::Category::IfGreater,             TEXT("in2")},         ExpressionInput(UE::Expressions::Inputs::ALessThanB)},     //another input is added for the case 'equal', see ConnectIfGreater
@@ -188,8 +189,6 @@ FMaterialXManager::FMaterialXManager()
 		{mx::Category::RampTB,       UE::Expressions::Names::MaterialXRampTopBottom},
 		{mx::Category::SplitLR,      UE::Expressions::Names::MaterialXSplitLeftRight},
 		{mx::Category::SplitTB,      UE::Expressions::Names::MaterialXSplitTopBottom},
-		// Procedural3D nodes
-		{mx::Category::Fractal3D,    UE::Expressions::Names::MaterialXFractal3D},
 		// Geometric nodes 
 		{mx::Category::GeomColor,    UE::Expressions::Names::VertexColor},
 		// Adjustment nodes,
@@ -295,9 +294,11 @@ FMaterialXManager::FMaterialXManager()
 		mx::Category::Distance,
 		mx::Category::DeonHairAbsorptionFromMelanin,
 		mx::Category::Extract,
+		mx::Category::Fractal2D,
 		mx::Category::Fractal3D,
 		mx::Category::GoochShade,
 		mx::Category::HsvAdjust,
+		mx::Category::LatLongImage,
 		mx::Category::Line,
 		mx::Category::LinAdobeRgbToLinRec709,
 		mx::Category::LinDisplayP3ToLinRec709,
diff --git a/Engine/Plugins/Interchange/Runtime/Source/Import/Private/MaterialX/MaterialXUtils/MaterialXSurfaceShaderAbstract.cpp b/Engine/Plugins/Interchange/Runtime/Source/Import/Private/MaterialX/MaterialXUtils/MaterialXSurfaceShaderAbstract.cpp
--- a/Engine/Plugins/Interchange/Runtime/Source/Import/Private/MaterialX/MaterialXUtils/MaterialXSurfaceShaderAbstract.cpp
+++ b/Engine/Plugins/Interchange/Runtime/Source/Import/Private/MaterialX/MaterialXUtils/MaterialXSurfaceShaderAbstract.cpp
@@ -1448,6 +1448,59 @@ void FMaterialXSurfaceShaderAbstract::ConnectHexTiledInputToOutput(const FConnec
 		UInterchangeShaderPortsAPI::ConnectDefaultOuputToInput(HexTileNode, TEXT("Texture"), TextureShaderNode->GetUniqueID());
 		UInterchangeShaderPortsAPI::ConnectDefaultOuputToInput(Connect.ParentShaderNode, Connect.InputChannelName, HexTileNode->GetUniqueID());
 	}
+	else
+	{
+		AddAttributeFromValueOrInterface(Connect.UpstreamNode->getInput("default"), Connect.InputChannelName, Connect.ParentShaderNode);
+	}
+}
+
+void FMaterialXSurfaceShaderAbstract::ConnectLatLongImageInputToOutput(const FConnectNode& Connect)
+{
+	using namespace UE::Interchange::Materials::Standard::Nodes;
+
+	if (UInterchangeTexture2DNode* TextureNode = Cast<UInterchangeTexture2DNode>(CreateTextureNode<UInterchangeTexture2DNode>(Connect.UpstreamNode)))
+	{
+		TextureNode->SetForceLongLatCubemap(true);
+		UInterchangeShaderNode* LatLongNode = CreateFunctionCallShaderNode(Connect.UpstreamNode, Connect.UpstreamNode->getName().c_str(), UE::MaterialFunctions::Path::MxLatLongImage);
+
+		const FString TextureNodeName = Connect.UpstreamNode->getName().c_str() + FString{ "_texture" };
+		UInterchangeShaderNode* TextureShaderNode = CreateShaderNode(Connect.UpstreamNode, TextureNodeName, TextureObject::Name.ToString());
+		const FString TextureObjectInputKey = UInterchangeShaderPortsAPI::MakeInputParameterKey(TextureObject::Inputs::Texture.ToString());
+		TextureShaderNode->AddStringAttribute(TextureObjectInputKey, TextureNode->GetUniqueID());
+
+		if (mx::InputPtr Input = Connect.UpstreamNode->getInput("viewdir"))
+		{
+			AddAttributeFromValueOrInterface(Input, GetInputName(Input), LatLongNode);
+		}
+
+		if (mx::InputPtr Input = Connect.UpstreamNode->getInput("rotation"))
+		{
+			AddAttributeFromValueOrInterface(Input, GetInputName(Input), LatLongNode);
+		}
+
+		UInterchangeShaderPortsAPI::ConnectDefaultOuputToInput(LatLongNode, TEXT("Texture"), TextureShaderNode->GetUniqueID());
+		UInterchangeShaderPortsAPI::ConnectDefaultOuputToInput(Connect.ParentShaderNode, Connect.InputChannelName, LatLongNode->GetUniqueID());
+	}
+	else
+	{
+		AddAttributeFromValueOrInterface(Connect.UpstreamNode->getInput("default"), Connect.InputChannelName, Connect.ParentShaderNode);
+	}
+}
+
+void FMaterialXSurfaceShaderAbstract::ConnectFractalInputToOutput(const FConnectNode& Connect)
+{
+	UInterchangeShaderNode* FractalNode = CreateShaderNode(Connect.UpstreamNode, Connect.UpstreamNode->getName().c_str(), UE::Expressions::Names::MaterialXFractal);
+	for (mx::InputPtr Input : Connect.UpstreamNode->getInputs())
+	{
+		AddAttributeFromValueOrInterface(Input, GetInputName(Input), FractalNode);
+	}
+
+	const int32 OutputIndex =  Connect.UpstreamNode->getType() == mx::Type::Vector2 ? 1 
+		: Connect.UpstreamNode->getType() == mx::Type::Vector3 || Connect.UpstreamNode->getType() == mx::Type::Color3 ? 2
+		: Connect.UpstreamNode->getType() == mx::Type::Vector4 || Connect.UpstreamNode->getType() == mx::Type::Color4 ? 3
+		: 0;
+
+	UInterchangeShaderPortsAPI::ConnectOuputToInputByIndex(Connect.ParentShaderNode, Connect.InputChannelName, FractalNode->GetUniqueID(), OutputIndex);
 }
 
 bool FMaterialXSurfaceShaderAbstract::ConnectSubsurfaceTransmissionToSurfaceShader(UInterchangeShaderNode* OpenPbrStandarSurfaceShaderNode, EInterchangeMaterialXShaders SurfaceShaderType
@@ -1793,6 +1846,9 @@ void FMaterialXSurfaceShaderAbstract::RegisterConnectNodeOutputToInputDelegates(
 	MatchingConnectNodeDelegates.Add(mx::Category::Luminance,		  FOnConnectNodeOutputToInput::CreateSP(this, &FMaterialXSurfaceShaderAbstract::ConnectLuminanceInputToOutput));
 	MatchingConnectNodeDelegates.Add(mx::Category::HexTiledImage,	  FOnConnectNodeOutputToInput::CreateSP(this, &FMaterialXSurfaceShaderAbstract::ConnectHexTiledInputToOutput));
 	MatchingConnectNodeDelegates.Add(mx::Category::HexTiledNormalMap, FOnConnectNodeOutputToInput::CreateSP(this, &FMaterialXSurfaceShaderAbstract::ConnectHexTiledInputToOutput));
+	MatchingConnectNodeDelegates.Add(mx::Category::LatLongImage,	  FOnConnectNodeOutputToInput::CreateSP(this, &FMaterialXSurfaceShaderAbstract::ConnectLatLongImageInputToOutput));
+	MatchingConnectNodeDelegates.Add(mx::Category::Fractal2D,		  FOnConnectNodeOutputToInput::CreateSP(this, &FMaterialXSurfaceShaderAbstract::ConnectFractalInputToOutput));
+	MatchingConnectNodeDelegates.Add(mx::Category::Fractal3D,		  FOnConnectNodeOutputToInput::CreateSP(this, &FMaterialXSurfaceShaderAbstract::ConnectFractalInputToOutput));
 }
 
 void FMaterialXSurfaceShaderAbstract::SetMatchingInputsNames(MaterialX::NodePtr Node) const
diff --git a/Engine/Plugins/Interchange/Runtime/Source/Import/Public/MaterialX/MaterialXUtils/MaterialXSurfaceShaderAbstract.h b/Engine/Plugins/Interchange/Runtime/Source/Import/Public/MaterialX/MaterialXUtils/MaterialXSurfaceShaderAbstract.h
--- a/Engine/Plugins/Interchange/Runtime/Source/Import/Public/MaterialX/MaterialXUtils/MaterialXSurfaceShaderAbstract.h
+++ b/Engine/Plugins/Interchange/Runtime/Source/Import/Public/MaterialX/MaterialXUtils/MaterialXSurfaceShaderAbstract.h
@@ -338,11 +338,15 @@ protected:
 
 	/** <luminance>*/
 	void ConnectLuminanceInputToOutput(const FConnectNode& Connect);
-
-	/** <hextiledimage>
-	  * <hextilednormalmap>
-	  */
+	
+	/** <hextiledimage> <hextilednormalmap>  */
 	void ConnectHexTiledInputToOutput(const FConnectNode& Connect);
+
+	/** <latlongimage> */
+	void ConnectLatLongImageInputToOutput(const FConnectNode& Connect);
+
+	/** <fractal2d/3d> */
+	void ConnectFractalInputToOutput(const FConnectNode& Connect);
 	/** End Connect MaterialX nodes*/
 
 	/** Function used only by both Standard Surface and OpenPBR to have both Subsurface and Transmission*/
