Interchange FBX:
- Remove bIsPhong and bHasBump static switch parameters from FBXLegacyPhongSurfaceMaterial
- Create parent material for Lambert and BumpMap cases

#virtualized
#Synced-CL 50150178

diff --git a/Engine/Plugins/Interchange/Runtime/Source/Parsers/Fbx/Private/FbxMaterial.cpp b/Engine/Plugins/Interchange/Runtime/Source/Parsers/Fbx/Private/FbxMaterial.cpp
--- a/Engine/Plugins/Interchange/Runtime/Source/Parsers/Fbx/Private/FbxMaterial.cpp
+++ b/Engine/Plugins/Interchange/Runtime/Source/Parsers/Fbx/Private/FbxMaterial.cpp
@@ -26,7 +26,6 @@ namespace UE
 			UInterchangeMaterialInstanceNode* FFbxMaterial::CreateMaterialInstanceNode(UInterchangeBaseNodeContainer& NodeContainer, const FString& NodeUID, const FString& NodeName)
 			{
 				UInterchangeMaterialInstanceNode* MaterialInstanceNode = NewObject<UInterchangeMaterialInstanceNode>(&NodeContainer);
-				MaterialInstanceNode->SetCustomParent(TEXT("/InterchangeAssets/Materials/FBXLegacyPhongSurfaceMaterial.FBXLegacyPhongSurfaceMaterial"));
 				NodeContainer.SetupNode(MaterialInstanceNode, NodeUID, NodeName, EInterchangeNodeContainerType::TranslatedAsset);
 
 				return MaterialInstanceNode;
@@ -192,6 +191,8 @@ namespace UE
 					return ExistingMaterialInstanceNode;
 				}
 
+				FString ShadingModel;
+				FString BumpMaterial;
 
 				UInterchangeMaterialInstanceNode* MaterialInstanceNode = CreateMaterialInstanceNode(NodeContainer, NodeUid, MaterialName);
 				if (MaterialInstanceNode == nullptr)
@@ -278,7 +279,7 @@ namespace UE
 						bool bHasBump = ConnectInput(Phong::Parameters::Normal, FbxSurfaceMaterial::sBump, FactorBump, DefaultValue, false);
 						if (bHasBump)
 						{
-							MaterialInstanceNode->AddStaticSwitchParameterValue(TEXT("bHasBump"), true);
+							BumpMaterial = TEXT("Bump");
 						}
 						bHasInput |= bHasBump;
 					}
@@ -346,7 +347,9 @@ namespace UE
 					}
 				}
 
-				MaterialInstanceNode->AddStaticSwitchParameterValue(TEXT("bIsPhong"), bIsPhong);
+				ShadingModel = bIsPhong ? TEXT("Phong") : TEXT("Lambert");
+				FString ParentMaterial{ TEXT("FBXLegacy") + ShadingModel + TEXT("SurfaceMaterial") + BumpMaterial };
+				MaterialInstanceNode->SetCustomParent(TEXT("/InterchangeAssets/Materials/") + ParentMaterial + TEXT('.') + ParentMaterial);
 
 				// If no valid property found, create a material anyway
 				TArray<FString> InputNames;
