[Interchange][FBX] ufbx - fixed animation

- skinned tracks that starts later than others in the animation are now sampled correctly
- fixed curve animation - weightmode is not used as ufbx always accomodates this in tangents and this works for all the cases

#Synced-CL 49917634

diff --git a/Engine/Plugins/Interchange/Runtime/Source/Parsers/Fbx/Private/Ufbx/UfbxAnimation.cpp b/Engine/Plugins/Interchange/Runtime/Source/Parsers/Fbx/Private/Ufbx/UfbxAnimation.cpp
--- a/Engine/Plugins/Interchange/Runtime/Source/Parsers/Fbx/Private/Ufbx/UfbxAnimation.cpp
+++ b/Engine/Plugins/Interchange/Runtime/Source/Parsers/Fbx/Private/Ufbx/UfbxAnimation.cpp
@@ -39,7 +39,7 @@ namespace UE::Interchange::Private
 			return Convert::ConvertQuat(V);
 		}
 
-		// Get inverval where the values are actually animated
+		// Get interval where the values are actually animated
 		template<typename ValueList, typename Value>
 		void UpdateAnimationInterval(TInterval<double>& Interval, const ValueList& Keys, const bool& IsConstant, const Value& Reference)
 		{
@@ -105,35 +105,19 @@ namespace UE::Interchange::Private
 					{
 						NewInterpMode = EInterchangeCurveInterpMode::Cubic;
 						EInterchangeCurveTangentMode NewTangentMode = EInterchangeCurveTangentMode::User;
-						EInterchangeCurveTangentWeightMode NewTangentWeightMode = EInterchangeCurveTangentWeightMode::WeightedNone;
 
 						float RightTangent = 0.0f;
 						float LeftTangent = 0.0f;
 						float RightTangentWeight = 0.0f;
 						float LeftTangentWeight = 0.0f;
 
-						if (ComputeTangent(Key.left, ScaleValue, LeftTangent, LeftTangentWeight))
-						{
-							if (KeyIndex > 0)
-							{
-								NewTangentWeightMode = EInterchangeCurveTangentWeightMode::WeightedArrive;
-							}
-						}
+						ComputeTangent(Key.left, ScaleValue, LeftTangent, LeftTangentWeight);
 						
-						if (ComputeTangent(Key.right, ScaleValue, RightTangent, RightTangentWeight))
-						{
-							if (KeyIndex < KeyCount-1)
-							{
-								NewTangentWeightMode = NewTangentWeightMode == EInterchangeCurveTangentWeightMode::WeightedArrive
-								   ? EInterchangeCurveTangentWeightMode::WeightedBoth
-								   : EInterchangeCurveTangentWeightMode::WeightedLeave;
-							}
-						}
+						ComputeTangent(Key.right, ScaleValue, RightTangent, RightTangentWeight);
 
 						// #ufbx_todo: possible to optimize curve by checking value, tangents and tangent weights for equality, in case cubic curve segment is actually constant/linear
-						// 
 						InterchangeCurveKey.TangentMode = NewTangentMode;
-						InterchangeCurveKey.TangentWeightMode = NewTangentWeightMode;
+						InterchangeCurveKey.TangentWeightMode = EInterchangeCurveTangentWeightMode::WeightedNone;
 
 						InterchangeCurveKey.ArriveTangent = LeftTangent;
 						InterchangeCurveKey.LeaveTangent = RightTangent;
@@ -171,12 +155,12 @@ namespace UE::Interchange::Private
 
 	}
 
-	void FUfbxAnimation::FetchSkinnedAnimation(const UE::Interchange::FAnimationPayloadQuery& PayloadQuery, const FSkeletalAnimationPayloadContext& SkeletalAnimationPayload, FAnimationPayloadData& PayloadData)
+	void FUfbxAnimation::FetchSkinnedAnimation(ufbx_scene* Scene, const UE::Interchange::FAnimationPayloadQuery& PayloadQuery, const FSkeletalAnimationPayloadContext& SkeletalAnimationPayload, FAnimationPayloadData& PayloadData)
 	{
 		ufbx_shared_ptr<ufbx_baked_anim> BakedAnim = SkeletalAnimationPayload.BakedAnim;
 		const ufbx_baked_node& BakedNode = BakedAnim->nodes[SkeletalAnimationPayload.BakedNodeIndex];
 
-		double BakeFrequency = PayloadQuery.TimeDescription.BakeFrequency;
+		const double BakeFrequency = PayloadQuery.TimeDescription.BakeFrequency;
 
 		PayloadData.BakeFrequency = BakeFrequency;
 		PayloadData.RangeStartTime = PayloadQuery.TimeDescription.RangeStartSecond;
@@ -186,23 +170,21 @@ namespace UE::Interchange::Private
 		int32 TotalKeyCount = FMath::RoundToInt32(AnimationDuration*BakeFrequency) + 1;
 		PayloadData.Transforms.SetNum(TotalKeyCount);
 
-		// Make transform from animation key index, 
-		auto MakeTransform = [&BakedNode, BakeFrequency](const int32 TransformIndex)
-		{
-			FTransform Transform;
-			FVector Translation = Convert::ConvertVec3(AnimUtils::GetAnimationKeyValue(BakedNode.translation_keys, BakedNode.constant_translation, TransformIndex, BakeFrequency));
-			FQuat Rotation = Convert::ConvertQuat(AnimUtils::GetAnimationKeyValue(BakedNode.rotation_keys, BakedNode.constant_rotation, TransformIndex, BakeFrequency));
-			FVector Scale = Convert::ConvertVec3(AnimUtils::GetAnimationKeyValue(BakedNode.scale_keys, BakedNode.constant_scale, TransformIndex, BakeFrequency));
-
-			Transform.SetTranslation(Translation);
-			Transform.SetRotation(Rotation);
-			Transform.SetScale3D(Scale);
-			return Transform;
-		};
-
+		const double FrameDuration = 1.0 / PayloadData.BakeFrequency;
 		for (int TransformIndex = 0; TransformIndex < TotalKeyCount; ++TransformIndex)
 		{
-			PayloadData.Transforms[TransformIndex] = MakeTransform(TransformIndex);
+			ufbx_transform Transform = ufbx_evaluate_transform_flags(
+				Scene->anim_stacks[SkeletalAnimationPayload.AnimStackIndex]->anim, 
+				Scene->nodes[BakedNode.typed_id], PayloadQuery.TimeDescription.RangeStartSecond + FrameDuration * TransformIndex,
+				0
+				| UFBX_TRANSFORM_FLAG_IGNORE_SCALE_HELPER 
+				| UFBX_TRANSFORM_FLAG_IGNORE_COMPONENTWISE_SCALE
+				| UFBX_TRANSFORM_FLAG_EXPLICIT_INCLUDES
+				| UFBX_TRANSFORM_FLAG_INCLUDE_TRANSLATION
+				| UFBX_TRANSFORM_FLAG_INCLUDE_ROTATION
+				| UFBX_TRANSFORM_FLAG_INCLUDE_SCALE
+			);
+			PayloadData.Transforms[TransformIndex] = Convert::ConvertTransform(Transform);
 		}
 	}
 
diff --git a/Engine/Plugins/Interchange/Runtime/Source/Parsers/Fbx/Private/Ufbx/UfbxAnimation.h b/Engine/Plugins/Interchange/Runtime/Source/Parsers/Fbx/Private/Ufbx/UfbxAnimation.h
--- a/Engine/Plugins/Interchange/Runtime/Source/Parsers/Fbx/Private/Ufbx/UfbxAnimation.h
+++ b/Engine/Plugins/Interchange/Runtime/Source/Parsers/Fbx/Private/Ufbx/UfbxAnimation.h
@@ -23,7 +23,7 @@ namespace UE::Interchange::Private
 	class FUfbxAnimation
 	{
 	public:
-		static void FetchSkinnedAnimation(const UE::Interchange::FAnimationPayloadQuery& PayloadQuery, const FSkeletalAnimationPayloadContext& SkeletalAnimationPayload, FAnimationPayloadData& PayloadData);
+		static void FetchSkinnedAnimation(ufbx_scene* Scene, const UE::Interchange::FAnimationPayloadQuery& PayloadQuery, const FSkeletalAnimationPayloadContext& SkeletalAnimationPayload, FAnimationPayloadData& PayloadData);
 		static bool FetchMorphTargetAnimation(const FUfbxParser& Parser, const FMorphAnimationPayloadContext& Animation, TArray<FInterchangeCurve>& OutCurves);
 		static bool FetchRigidAnimation(const FUfbxParser& Parser, const FRigidAnimationPayloadContext& Animation, TArray<FInterchangeCurve>& OutCurves);
 
diff --git a/Engine/Plugins/Interchange/Runtime/Source/Parsers/Fbx/Private/Ufbx/UfbxParser.cpp b/Engine/Plugins/Interchange/Runtime/Source/Parsers/Fbx/Private/Ufbx/UfbxParser.cpp
--- a/Engine/Plugins/Interchange/Runtime/Source/Parsers/Fbx/Private/Ufbx/UfbxParser.cpp
+++ b/Engine/Plugins/Interchange/Runtime/Source/Parsers/Fbx/Private/Ufbx/UfbxParser.cpp
@@ -652,7 +652,7 @@ bool FUfbxParser::FetchAnimationBakeTransformPayload(const TArray<UE::Interchang
 		{
 			const FSkeletalAnimationPayloadContext& SkeletalAnimationPayloadContext = PayloadContexts.GetAnimation(*PayloadContext);
 			FAnimationPayloadData PayloadData = FAnimationPayloadData(SkeletalAnimationPayloadContext.NodeUid, PayloadQuery.PayloadKey);
-			FUfbxAnimation::FetchSkinnedAnimation(PayloadQuery, SkeletalAnimationPayloadContext, PayloadData);
+			FUfbxAnimation::FetchSkinnedAnimation(Scene, PayloadQuery, SkeletalAnimationPayloadContext, PayloadData);
 
 			FString QueryHashString = PayloadQuery.GetHashString();
 
