Fixed issue with LevelVariantSet not imported

#rnx

#Synced-CL 50452739

diff --git a/Engine/Plugins/Interchange/Runtime/Source/Import/Public/Scene/InterchangeSceneVariantSetsFactory.h b/Engine/Plugins/Interchange/Runtime/Source/Import/Public/Scene/InterchangeSceneVariantSetsFactory.h
--- a/Engine/Plugins/Interchange/Runtime/Source/Import/Public/Scene/InterchangeSceneVariantSetsFactory.h
+++ b/Engine/Plugins/Interchange/Runtime/Source/Import/Public/Scene/InterchangeSceneVariantSetsFactory.h
@@ -26,8 +26,9 @@ public:
 	// Interchange factory base interface begin
 
 	UE_API virtual UClass* GetFactoryClass() const override;
+	virtual EInterchangeFactoryAssetType GetFactoryAssetType() override { return EInterchangeFactoryAssetType::Variants; }
+
 	UE_API virtual FImportAssetResult BeginImportAsset_GameThread(const FImportAssetObjectParams& Arguments) override;
-	
 	UE_API virtual void SetupObject_GameThread(const FSetupObjectParams& Arguments) override;
 
 	// Interchange factory base interface end
diff --git a/Engine/Source/Runtime/Interchange/Core/Public/InterchangeFactoryBase.h b/Engine/Source/Runtime/Interchange/Core/Public/InterchangeFactoryBase.h
--- a/Engine/Source/Runtime/Interchange/Core/Public/InterchangeFactoryBase.h
+++ b/Engine/Source/Runtime/Interchange/Core/Public/InterchangeFactoryBase.h
@@ -31,6 +31,7 @@ enum class EInterchangeFactoryAssetType : uint8
 	Grooms,
 	Sounds,
 	Levels,
+	Variants,
 	Custom
 };
 
