[Interchange][FBX] ufbx - fixed import of morphs

#Synced-CL 49917634

diff --git a/Engine/Plugins/Interchange/Runtime/Source/Parsers/Fbx/Private/Ufbx/UfbxMesh.cpp b/Engine/Plugins/Interchange/Runtime/Source/Parsers/Fbx/Private/Ufbx/UfbxMesh.cpp
--- a/Engine/Plugins/Interchange/Runtime/Source/Parsers/Fbx/Private/Ufbx/UfbxMesh.cpp
+++ b/Engine/Plugins/Interchange/Runtime/Source/Parsers/Fbx/Private/Ufbx/UfbxMesh.cpp
@@ -374,7 +374,7 @@ namespace UE::Interchange::Private
 
 		int32 VertexCount = Mesh.num_vertices;
 
-		TArrayView<ufbx_vec3> Vertices(Mesh.vertices.data, Mesh.vertices.count);
+		TArray<ufbx_vec3> Vertices(Mesh.vertices.data, Mesh.vertices.count);
 		ufbx_add_blend_shape_vertex_offsets(&BlendShape, Vertices.GetData(), VertexCount, 1.0f);
 
 		TVertexAttributesRef<FVector3f> VertexPositions = Attributes.GetVertexPositions();
