MaterialX:
- Add support for texcoord input on <heighttonormal> node
- Replace usage of HeightToNormalSmooth material function by MX_HeightToNormal, since the former computes a normal map based on 3D positions/normals whereas the latter is compliant with MaterialX spec (using texcoord)
- Plug a TextureProperty using TexelSize on the NormalFromHeightmap material function giving smoother normal perturbation

Interchange Generic Material Pipeline:
- Add support for TextureProperty

#virtualized
#Synced-CL 50150178

diff --git a/Engine/Plugins/Interchange/Runtime/Source/Common/Public/MaterialX/InterchangeMaterialXDefinitions.h b/Engine/Plugins/Interchange/Runtime/Source/Common/Public/MaterialX/InterchangeMaterialXDefinitions.h
--- a/Engine/Plugins/Interchange/Runtime/Source/Common/Public/MaterialX/InterchangeMaterialXDefinitions.h
+++ b/Engine/Plugins/Interchange/Runtime/Source/Common/Public/MaterialX/InterchangeMaterialXDefinitions.h
@@ -320,6 +320,7 @@ namespace UE
 			inline constexpr const TCHAR* MxColorCorrect = TEXT("/InterchangeAssets/Functions/MX_ColorCorrect.MX_ColorCorrect");
 			inline constexpr const TCHAR* MxDeonHairAbsorptionFromMelanin = TEXT("/InterchangeAssets/Functions/MX_DeonHairAbsorptionFromMelanin.MX_DeonHairAbsorptionFromMelanin");
 			inline constexpr const TCHAR* MxGoochShade = TEXT("/InterchangeAssets/Functions/MX_GoochShade.MX_GoochShade");
+			inline constexpr const TCHAR* MxHeightToNormalMap = TEXT("/InterchangeAssets/Functions/MX_HeightToNormal.MX_HeightToNormal");
 			inline constexpr const TCHAR* MxHexTiledImage = TEXT("/InterchangeAssets/Functions/MX_HexTiledImage.MX_HexTiledImage");
 			inline constexpr const TCHAR* MxHexTiledNormalMap= TEXT("/InterchangeAssets/Functions/MX_HexTiledNormalMap.MX_HexTiledNormalMap");
 			inline constexpr const TCHAR* MxHsvAdjust = TEXT("/InterchangeAssets/Functions/MX_HSVAdjust.MX_HSVAdjust");
diff --git a/Engine/Plugins/Interchange/Runtime/Source/Import/Private/MaterialX/MaterialXUtils/MaterialXSurfaceShaderAbstract.cpp b/Engine/Plugins/Interchange/Runtime/Source/Import/Private/MaterialX/MaterialXUtils/MaterialXSurfaceShaderAbstract.cpp
--- a/Engine/Plugins/Interchange/Runtime/Source/Import/Private/MaterialX/MaterialXUtils/MaterialXSurfaceShaderAbstract.cpp
+++ b/Engine/Plugins/Interchange/Runtime/Source/Import/Private/MaterialX/MaterialXUtils/MaterialXSurfaceShaderAbstract.cpp
@@ -6,6 +6,7 @@
 #include "InterchangeTexture2DNode.h"
 #include "InterchangeTextureBlurNode.h"
 #include "Materials/MaterialExpressionTransform.h"
+#include "Materials/MaterialExpressionTextureProperty.h"
 #include "Materials/MaterialExpressionVectorNoise.h"
 #include "MaterialX/MaterialXUtils/MaterialXManager.h"
 #include "MaterialX/MaterialExpressions/MaterialExpressionTextureSampleParameterBlur.h"
@@ -1032,14 +1033,17 @@ void FMaterialXSurfaceShaderAbstract::ConnectWorleyNoise3DInputToOutput(const FC
 
 void FMaterialXSurfaceShaderAbstract::ConnectHeightToNormalInputToOutput(const FConnectNode& Connect)
 {
+	// We could always use the material function MX_HeightToNormalMap, 
+	// though it has been proven that sampling textures to compute a normal map is cheaper and gives better results
+
+	using namespace UE::Interchange::Materials::Standard::Nodes;
+
 	if(mx::InputPtr Input = Connect.UpstreamNode->getInput("in"))
 	{
 		// Image node will become this node
 		if(mx::NodePtr ConnectedNode = Input->getConnectedNode();
 		   ConnectedNode && ConnectedNode->getCategory() == mx::Category::Image)
 		{
-			using namespace UE::Interchange::Materials::Standard::Nodes;
-
 			if(UInterchangeTextureNode* TextureNode = CreateTextureNode<UInterchangeTexture2DNode>(ConnectedNode))
 			{
 				UInterchangeShaderNode* HeightMapNode = CreateShaderNode(Connect.UpstreamNode, Connect.UpstreamNode->getName().c_str(), NormalFromHeightMap::Name.ToString());
@@ -1051,6 +1055,15 @@ void FMaterialXSurfaceShaderAbstract::ConnectHeightToNormalInputToOutput(const F
 				TextureShaderNode->AddStringAttribute(TextureObjectInputKey, TextureNode->GetUniqueID());
 				UInterchangeShaderPortsAPI::ConnectDefaultOuputToInput(HeightMapNode, NormalFromHeightMap::Inputs::HeightMap.ToString(), TextureShaderNode->GetUniqueID());
 
+				// Create a Texture Property with texel size to have a better quality when sampling the textures for the normalmap
+				const FString TexturePropertyNodeName = Connect.UpstreamNode->getName().c_str() + FString{ "_textureproperty" };
+				UInterchangeShaderNode* TexturePropertyNode = CreateShaderNode(Connect.UpstreamNode, TexturePropertyNodeName, TextureProperty::Name.ToString());
+				TexturePropertyNode->AddInt32Attribute(TextureProperty::Attributes::Property.ToString(), EMaterialExposedTextureProperty::TMTM_TexelSize);
+				UInterchangeShaderPortsAPI::ConnectDefaultOuputToInput(TexturePropertyNode, TextureProperty::Inputs::Texture.ToString(), TextureShaderNode->GetUniqueID());
+				UInterchangeShaderPortsAPI::ConnectDefaultOuputToInput(HeightMapNode, NormalFromHeightMap::Inputs::Offset.ToString(), TexturePropertyNode->GetUniqueID());
+				
+				AddAttributeFromValueOrInterface(Connect.UpstreamNode->getInput("texcoord"), NormalFromHeightMap::Inputs::Coordinates.ToString(), HeightMapNode);
+
 				// The default value in the material function is 8, but the default value for the node is 1
 				if (!AddAttributeFromValueOrInterface(Connect.UpstreamNode->getInput("scale"), NormalFromHeightMap::Inputs::Intensity.ToString(), HeightMapNode))
 				{
@@ -1065,50 +1078,14 @@ void FMaterialXSurfaceShaderAbstract::ConnectHeightToNormalInputToOutput(const F
 		}
 		else
 		{
-			using namespace UE::Interchange::Materials::Standard::Nodes;
-
-			// if we have a scale the multiply node input A will be the Parent Shader Node otherwise it's just the MF
-			UInterchangeShaderNode* ScaleNode = nullptr;
-			FString HeightToNormalNodeName = Connect.UpstreamNode->getName().c_str();
-			if (mx::InputPtr InputScale = Connect.UpstreamNode->getInput("scale"))
-			{
-				ScaleNode = CreateShaderNode(Connect.UpstreamNode, HeightToNormalNodeName, Multiply::Name.ToString());
-				HeightToNormalNodeName += TEXT("_Smooth"); // the MF won't be the main node to connect to
-				SetAttributeNewName(Connect.UpstreamNode->getInput("in"), "A");
-				AddAttributeFromValueOrInterface(InputScale, Multiply::Inputs::B.ToString(), ScaleNode);
-			}
-			else
-			{
-				SetAttributeNewName(Connect.UpstreamNode->getInput("in"), "Height");
-			}
-
-			// This MF outputs a normal vector in World Space, we need to transform the output in Tangent Space
-			UInterchangeFunctionCallShaderNode* HeightToNormalSmoothShaderNode =
-				CreateFunctionCallShaderNode(Connect.UpstreamNode,
-											 HeightToNormalNodeName,
-											 UE::MaterialFunctions::Path::HeightToNormalSmooth);
+			UInterchangeFunctionCallShaderNode* HeightToNormalMapNode = CreateFunctionCallShaderNode(Connect.UpstreamNode, Connect.UpstreamNode->getName().c_str(), UE::MaterialFunctions::Path::MxHeightToNormalMap);
 			
-			if (ScaleNode)
+			for(mx::InputPtr InputHeightmap : Connect.UpstreamNode->getInputs())
 			{
-				UInterchangeShaderPortsAPI::ConnectDefaultOuputToInput(HeightToNormalSmoothShaderNode, UE::Expressions::Inputs::Height, ScaleNode->GetUniqueID());
+				AddAttributeFromValueOrInterface(InputHeightmap, GetInputName(InputHeightmap), HeightToNormalMapNode);
 			}
 
-			UInterchangeShaderNode* TransformNode = CreateShaderNode(Connect.UpstreamNode, (Connect.UpstreamNode->getName() + "_TS").c_str(), TransformVector::Name.ToString());
-			TransformNode->AddInt32Attribute(TransformVector::Attributes::TransformSourceType.ToString(), EMaterialVectorCoordTransformSource::TRANSFORMSOURCE_World);
-			TransformNode->AddInt32Attribute(TransformVector::Attributes::TransformType.ToString(), EMaterialVectorCoordTransform::TRANSFORM_Tangent);
-
-			// By default positions in MaterialX are in object space so we need to create a local position expression here.
-			// TODO: Handling it as a geompropdef so all nodes refer to the same LocalPosition expression
-			UInterchangeShaderNode* PositionNode = CreateShaderNode(Connect.UpstreamNode, (Connect.UpstreamNode->getName() + "_Position").c_str(), UE::Expressions::Names::LocalPosition);
-			UInterchangeShaderNode* UnitNode = CreateShaderNode(Connect.UpstreamNode, (Connect.UpstreamNode->getName() + "_Unit").c_str(), Multiply::Name.ToString());
-
-			UnitNode->AddFloatAttribute(UInterchangeShaderPortsAPI::MakeInputValueKey(Multiply::Inputs::B.ToString()), 0.01f);
-			UInterchangeShaderPortsAPI::ConnectDefaultOuputToInput(UnitNode, Multiply::Inputs::A.ToString(), PositionNode->GetUniqueID());
-
-			UInterchangeShaderPortsAPI::ConnectDefaultOuputToInput(HeightToNormalSmoothShaderNode, UE::Expressions::Inputs::AbsoluteWorldPosition, UnitNode->GetUniqueID());
-
-			UInterchangeShaderPortsAPI::ConnectDefaultOuputToInput(TransformNode, TransformVector::Inputs::Input.ToString(), HeightToNormalSmoothShaderNode->GetUniqueID());
-			UInterchangeShaderPortsAPI::ConnectDefaultOuputToInput(Connect.ParentShaderNode, Connect.InputChannelName, TransformNode->GetUniqueID());
+			UInterchangeShaderPortsAPI::ConnectDefaultOuputToInput(Connect.ParentShaderNode, Connect.InputChannelName, HeightToNormalMapNode->GetUniqueID());
 		}
 	}
 }
@@ -1717,14 +1694,19 @@ const FString& FMaterialXSurfaceShaderAbstract::GetMatchedInputName(MaterialX::N
 
 FString FMaterialXSurfaceShaderAbstract::GetInputName(MaterialX::InputPtr Input) const
 {
-	if (Input->hasAttribute(mx::Attributes::NewName))
-	{
-		return Input->getAttribute(mx::Attributes::NewName).c_str();
-	}
-	else
+	if(Input)
 	{
-		return Input->getName().c_str();
+		if (Input->hasAttribute(mx::Attributes::NewName))
+		{
+			return Input->getAttribute(mx::Attributes::NewName).c_str();
+		}
+		else
+		{
+			return Input->getName().c_str();
+		}
 	}
+
+	return {};
 }
 
 FString FMaterialXSurfaceShaderAbstract::GetFilePrefix(MaterialX::ElementPtr Element) const
diff --git a/Engine/Plugins/Interchange/Runtime/Source/Nodes/Public/InterchangeMaterialDefinitions.h b/Engine/Plugins/Interchange/Runtime/Source/Nodes/Public/InterchangeMaterialDefinitions.h
--- a/Engine/Plugins/Interchange/Runtime/Source/Nodes/Public/InterchangeMaterialDefinitions.h
+++ b/Engine/Plugins/Interchange/Runtime/Source/Nodes/Public/InterchangeMaterialDefinitions.h
@@ -231,6 +231,21 @@ namespace Interchange
 					}
 				}
 
+				namespace TextureProperty
+				{
+					const FName Name = TEXT("TextureProperty");
+
+					namespace Inputs
+					{
+						const FName Texture = TEXT("TextureUid"); // Type: FString (unique id of a texture node)
+					}
+
+					namespace Attributes
+					{
+						const FName Property = TEXT("Property"); // Type: enum
+					}
+				}
+
 				namespace TextureSample
 				{
 					const FName Name = TEXT("TextureSample");
diff --git a/Engine/Plugins/Interchange/Runtime/Source/Pipelines/Private/InterchangeGenericMaterialPipeline.cpp b/Engine/Plugins/Interchange/Runtime/Source/Pipelines/Private/InterchangeGenericMaterialPipeline.cpp
--- a/Engine/Plugins/Interchange/Runtime/Source/Pipelines/Private/InterchangeGenericMaterialPipeline.cpp
+++ b/Engine/Plugins/Interchange/Runtime/Source/Pipelines/Private/InterchangeGenericMaterialPipeline.cpp
@@ -50,6 +50,7 @@
 #include "Materials/MaterialExpressionTextureCoordinate.h"
 #include "Materials/MaterialExpressionTextureObject.h"
 #include "Materials/MaterialExpressionTextureObjectParameter.h"
+#include "Materials/MaterialExpressionTextureProperty.h"
 #include "Materials/MaterialExpressionTextureSample.h"
 #include "Materials/MaterialExpressionTextureSampleParameter2D.h"
 #include "Materials/MaterialExpressionTextureSampleParameter2DArray.h"
@@ -3083,6 +3084,34 @@ void UInterchangeGenericMaterialPipeline::HandleLuminanceNode(const UInterchange
 	}
 }
 
+void UInterchangeGenericMaterialPipeline::HandleTexturePropertyNode(const UInterchangeShaderNode* ShaderNode, UInterchangeBaseMaterialFactoryNode* MaterialFactoryNode, UInterchangeMaterialExpressionFactoryNode* TexturePropertyFactoryNode)
+{
+	using namespace UE::Interchange::Materials::Standard::Nodes;
+	TexturePropertyFactoryNode->SetCustomExpressionClassName(UMaterialExpressionTextureProperty::StaticClass()->GetName());
+
+	// TextureObject
+	{
+		TTuple<UInterchangeMaterialExpressionFactoryNode*, FString> TextureObjectExpression =
+			CreateMaterialExpressionForInput(MaterialFactoryNode, ShaderNode, TextureProperty::Inputs::Texture.ToString(), TexturePropertyFactoryNode->GetUniqueID());
+
+		if (TextureObjectExpression.Get<0>())
+		{
+			UInterchangeShaderPortsAPI::ConnectOuputToInputByName(TexturePropertyFactoryNode, GET_MEMBER_NAME_CHECKED(UMaterialExpressionTextureProperty, TextureObject).ToString(),
+				TextureObjectExpression.Get<0>()->GetUniqueID(), TextureObjectExpression.Get<1>());
+		}
+	}
+
+	// Property
+	{
+		if (int32 Property; ShaderNode->GetInt32Attribute(TextureProperty::Attributes::Property.ToString(), Property))
+		{
+			const FName PropertyMemberName = GET_MEMBER_NAME_CHECKED(UMaterialExpressionTextureProperty, Property);
+			TexturePropertyFactoryNode->AddInt32Attribute(PropertyMemberName.ToString(), Property);
+			TexturePropertyFactoryNode->AddApplyAndFillDelegates<int32>(PropertyMemberName.ToString(), UMaterialExpressionTextureProperty::StaticClass(), PropertyMemberName);
+		}
+	}
+}
+
 /** Trigonometry expressions are unitless and sets the period to 1 by default to let the user sets their period*/
 void UInterchangeGenericMaterialPipeline::HandleTrigonometryNode(const UInterchangeShaderNode* ShaderNode, UClass* StaticClass, UInterchangeBaseMaterialFactoryNode* MaterialFactoryNode, UInterchangeMaterialExpressionFactoryNode* TrigonometryFactoryNode)
 {
@@ -3255,6 +3284,10 @@ UInterchangeMaterialExpressionFactoryNode* UInterchangeGenericMaterialPipeline::
 	{
 		HandleTextureObjectNode(ShaderNode, MaterialFactoryNode, MaterialExpression);
 	}
+	else if(*ShaderType == Nodes::TextureProperty::Name)
+	{
+		HandleTexturePropertyNode(ShaderNode, MaterialFactoryNode, MaterialExpression);
+	}
 	else if (*ShaderType == Nodes::TextureSample::Name)
 	{
 		HandleTextureSampleNode(ShaderNode, MaterialFactoryNode, MaterialExpression);
diff --git a/Engine/Plugins/Interchange/Runtime/Source/Pipelines/Public/InterchangeGenericMaterialPipeline.h b/Engine/Plugins/Interchange/Runtime/Source/Pipelines/Public/InterchangeGenericMaterialPipeline.h
--- a/Engine/Plugins/Interchange/Runtime/Source/Pipelines/Public/InterchangeGenericMaterialPipeline.h
+++ b/Engine/Plugins/Interchange/Runtime/Source/Pipelines/Public/InterchangeGenericMaterialPipeline.h
@@ -231,6 +231,7 @@ private:
 	UE_API void HandleSwitchNode(const UInterchangeShaderNode* ShaderNode, UInterchangeBaseMaterialFactoryNode* MaterialFactoryNode, UInterchangeMaterialExpressionFactoryNode* SwitchFactoryNode);
 	UE_API void HandleSlabBSDFNode(const UInterchangeShaderNode* ShaderNode, UInterchangeBaseMaterialFactoryNode* MaterialFactoryNode, UInterchangeMaterialExpressionFactoryNode* SlabBSDFFactoryNode);
     UE_API void HandleLuminanceNode(const UInterchangeShaderNode* ShaderNode, UInterchangeBaseMaterialFactoryNode* MaterialFactoryNode, UInterchangeMaterialExpressionFactoryNode* LuminanceFactoryNode);
+    UE_API void HandleTexturePropertyNode(const UInterchangeShaderNode* ShaderNode, UInterchangeBaseMaterialFactoryNode* MaterialFactoryNode, UInterchangeMaterialExpressionFactoryNode* TexturePropertyFactoryNode);
 	UE_API void HandleTrigonometryNode(const UInterchangeShaderNode* ShaderNode, UClass* StaticClass, UInterchangeBaseMaterialFactoryNode* MaterialFactoryNode, UInterchangeMaterialExpressionFactoryNode* TrigonometryFactoryNode);
 
 	UE_API void HandleScalarParameterNode(const UInterchangeShaderNode* ShaderNode, UInterchangeMaterialExpressionFactoryNode* ScalarParameterFactoryNode);
