[Interchange][FBX] fixed crash importing animation

- testcase in UE-364588 (try importing main file and then it asserts on Jump.fbx )

#Synced-CL 50703335

diff --git a/Engine/Plugins/Interchange/Runtime/Source/Pipelines/Private/InterchangeGenericAnimationPipeline.cpp b/Engine/Plugins/Interchange/Runtime/Source/Pipelines/Private/InterchangeGenericAnimationPipeline.cpp
--- a/Engine/Plugins/Interchange/Runtime/Source/Pipelines/Private/InterchangeGenericAnimationPipeline.cpp
+++ b/Engine/Plugins/Interchange/Runtime/Source/Pipelines/Private/InterchangeGenericAnimationPipeline.cpp
@@ -301,6 +301,13 @@ void UInterchangeGenericAnimationPipeline::ExecutePipeline(UInterchangeBaseNodeC
 
 	BaseNodeContainer = InBaseNodeContainer;
 
+	if (!CommonSkeletalMeshesAndAnimationsProperties.IsValid())
+	{
+		return;
+	}
+
+	CommonSkeletalMeshesAndAnimationsProperties->SetResultsContainer(Results);
+
 	if (CommonSkeletalMeshesAndAnimationsProperties->bImportOnlyAnimations)
 	{
 		bImportAnimations = true;
@@ -535,11 +542,6 @@ void UInterchangeGenericAnimationPipeline::ExecutePipeline(UInterchangeBaseNodeC
 		}
 	}
 
-	if (!CommonSkeletalMeshesAndAnimationsProperties.IsValid())
-	{
-		return;
-	}
-
 	if (CommonSkeletalMeshesAndAnimationsProperties->bImportOnlyAnimations && !CommonSkeletalMeshesAndAnimationsProperties->IsSkeletonValid())
 	{
 		UE_LOG(LogInterchangePipeline, Warning, TEXT("UInterchangeGenericAnimationPipeline: Cannot execute pre-import pipeline because importing animation only requires a valid skeleton."));
