[Interchange][FBX] ufbx material handling updated to follow FBX SDK parser changes

- due to splitting shader to Bump/Lambert

#Synced-CL 50703335

diff --git a/Engine/Plugins/Interchange/Runtime/Source/Parsers/Fbx/Private/Ufbx/UfbxMaterial.cpp b/Engine/Plugins/Interchange/Runtime/Source/Parsers/Fbx/Private/Ufbx/UfbxMaterial.cpp
--- a/Engine/Plugins/Interchange/Runtime/Source/Parsers/Fbx/Private/Ufbx/UfbxMaterial.cpp
+++ b/Engine/Plugins/Interchange/Runtime/Source/Parsers/Fbx/Private/Ufbx/UfbxMaterial.cpp
@@ -160,7 +160,6 @@ namespace UE::Interchange::Private
 	UInterchangeMaterialInstanceNode* CreateMaterialInstanceNode(UInterchangeBaseNodeContainer& NodeContainer, const FString& NodeUID, const FString& NodeLabel)
 	{
 		UInterchangeMaterialInstanceNode* MaterialInstanceNode = NewObject<UInterchangeMaterialInstanceNode>(&NodeContainer);
-		MaterialInstanceNode->SetCustomParent(TEXT("/InterchangeAssets/Materials/FBXLegacyPhongSurfaceMaterial.FBXLegacyPhongSurfaceMaterial"));
 		NodeContainer.SetupNode(MaterialInstanceNode, NodeUID, NodeLabel, EInterchangeNodeContainerType::TranslatedAsset);
 
 		return MaterialInstanceNode;
@@ -328,6 +327,8 @@ namespace UE::Interchange::Private
 				bHasInput |= bIsPhong;
 			}
 
+			FString BumpMaterial;
+
 			// Normal
 			if (Material->fbx.normal_map.has_value)
 			{
@@ -337,7 +338,7 @@ namespace UE::Interchange::Private
 				}
 				else if (SetTextureFromMap(Phong::Parameters::Normal, Material->fbx.bump, 1.f))
 				{
-					MaterialInstanceNode->AddStaticSwitchParameterValue(TEXT("bHasBump"), true);
+					BumpMaterial = TEXT("Bump");
 					bHasInput = true;
 				}
 			}
@@ -401,6 +402,10 @@ namespace UE::Interchange::Private
 				}
 			}
 
+			FString ShadingModel = bIsPhong ? TEXT("Phong") : TEXT("Lambert");
+			FString ParentMaterial{ TEXT("FBXLegacy") + ShadingModel + TEXT("SurfaceMaterial") + BumpMaterial };
+			MaterialInstanceNode->SetCustomParent(TEXT("/InterchangeAssets/Materials/") + ParentMaterial + TEXT('.') + ParentMaterial);
+
 			// If no valid property found, create a material anyway
 			TArray<FString> InputNames;
 			if (!bHasInput)
