Interchange: Spawn light actors with the same default intensity units that the engine uses.

#Synced-CL 50488581

diff --git a/Engine/Plugins/Interchange/Runtime/Source/Import/Private/Scene/InterchangeLightActorFactory.cpp b/Engine/Plugins/Interchange/Runtime/Source/Import/Private/Scene/InterchangeLightActorFactory.cpp
--- a/Engine/Plugins/Interchange/Runtime/Source/Import/Private/Scene/InterchangeLightActorFactory.cpp
+++ b/Engine/Plugins/Interchange/Runtime/Source/Import/Private/Scene/InterchangeLightActorFactory.cpp
@@ -1,11 +1,14 @@
 // Copyright Epic Games, Inc. All Rights Reserved.
 
 #include "Scene/InterchangeLightActorFactory.h"
-#include "Scene/InterchangeActorHelper.h"
 #include "Components/LightComponent.h"
-#include "Nodes/InterchangeFactoryBaseNode.h"
-#include "Nodes/InterchangeBaseNodeContainer.h"
+#include "Components/PointLightComponent.h"
+#include "Components/RectLightComponent.h"
+#include "Components/SpotLightComponent.h"
 #include "InterchangeTextureLightProfileFactoryNode.h"
+#include "Nodes/InterchangeBaseNodeContainer.h"
+#include "Nodes/InterchangeFactoryBaseNode.h"
+#include "Scene/InterchangeActorHelper.h"
 
 #include "CoreMinimal.h"
 #include "UObject/Object.h"
@@ -62,6 +65,28 @@ UObject* UInterchangeLightActorFactory::ProcessActor(AActor& SpawnedActor, const
 				}
 			}
 
+			static const auto CVarDefaultLightUnits = IConsoleManager::Get().FindTConsoleVariableDataInt(TEXT("r.DefaultFeature.LightUnits"));
+			ELightUnits DefaultUnits = (ELightUnits)CVarDefaultLightUnits->GetValueOnAnyThread();
+
+			// Replicate the code in UActorFactorySpotLight::PostSpawnActor and analogues so we end up with the same light units that a user would get
+			// on their components when placing a light actor manually onto the level
+			if (USpotLightComponent* SpotLightComponent = Cast<USpotLightComponent>(LightComponent))
+			{
+				SpotLightComponent->Intensity *= UPointLightComponent::GetUnitsConversionFactor(SpotLightComponent->IntensityUnits, DefaultUnits, SpotLightComponent->GetCosHalfConeAngle());
+				SpotLightComponent->IntensityUnits = DefaultUnits;
+			}
+			else if (UPointLightComponent* PointLightComponent = Cast<UPointLightComponent>(LightComponent))
+			{
+				PointLightComponent->Intensity *= UPointLightComponent::GetUnitsConversionFactor(PointLightComponent->IntensityUnits, DefaultUnits, -1.f);
+				PointLightComponent->IntensityUnits = DefaultUnits;
+			}
+			else if (URectLightComponent* RectLightComponent = Cast<URectLightComponent>(LightComponent))
+			{
+				// Passing .5 as CosHalfConeAngle  since URectLightComponent::SetLightBrightness() lumens conversion use only PI
+				RectLightComponent->Intensity *= URectLightComponent::GetUnitsConversionFactor(RectLightComponent->IntensityUnits, DefaultUnits, .5f);
+				RectLightComponent->IntensityUnits = DefaultUnits;
+			}
+
 			return LightComponent;
 		}
 	}
