[Interchange] Silencing LOD pattern missmatch warning in case of Collision pattern found.

#rnx

#Synced-CL 50708145

diff --git a/Engine/Plugins/Interchange/Extensions/OpenUSD/Source/Import/Private/SchemaHandlers/GprimSchemaHandler.cpp b/Engine/Plugins/Interchange/Extensions/OpenUSD/Source/Import/Private/SchemaHandlers/GprimSchemaHandler.cpp
--- a/Engine/Plugins/Interchange/Extensions/OpenUSD/Source/Import/Private/SchemaHandlers/GprimSchemaHandler.cpp
+++ b/Engine/Plugins/Interchange/Extensions/OpenUSD/Source/Import/Private/SchemaHandlers/GprimSchemaHandler.cpp
@@ -28,6 +28,7 @@
 #include "InterchangeResult.h"
 #include "InterchangeResultsContainer.h"
 #include "InterchangeSceneNode.h"
+#include "Mesh/InterchangeMeshHelper.h"
 #include "Nodes/InterchangeBaseNodeContainer.h"
 
 #include "Rendering/SkeletalMeshLODImporterData.h"
@@ -227,19 +228,10 @@ namespace UE::GeomGprimSchemaHandler::Private
 				// USD Importer code
 				bool bSetCustomCollisionType = true;
 				{
-					const static TSet<FString> CollisionPrefixes = {TEXT("UBX"), TEXT("UCX"), TEXT("MCDCX"), TEXT("USP"), TEXT("UCP")};
-
 					// Check if this mesh is an FBX-style collider
+					if (UE::Interchange::Private::MeshHelper::GetMeshCollisionFromName(NodeName) != EInterchangeMeshCollision::None)
 					{
-						int32 FirstUnderscoreIndex = INDEX_NONE;
-						if (NodeName.FindChar(TEXT('_'), FirstUnderscoreIndex) && FirstUnderscoreIndex != INDEX_NONE)
-						{
-							FString Prefix = NodeName.Left(FirstUnderscoreIndex);
-							if (CollisionPrefixes.Contains(Prefix))
-							{
-								bSetCustomCollisionType = false;
-							}
-						}
+						bSetCustomCollisionType = false;
 					}
 
 					// Check if we have any siblings that are also FBX-style colliders pointing at this mesh prim.
diff --git a/Engine/Plugins/Interchange/Runtime/Source/Import/Private/Mesh/InterchangeMeshHelper.cpp b/Engine/Plugins/Interchange/Runtime/Source/Import/Private/Mesh/InterchangeMeshHelper.cpp
--- a/Engine/Plugins/Interchange/Runtime/Source/Import/Private/Mesh/InterchangeMeshHelper.cpp
+++ b/Engine/Plugins/Interchange/Runtime/Source/Import/Private/Mesh/InterchangeMeshHelper.cpp
@@ -38,6 +38,7 @@
 #include "InterchangeImportLog.h"
 #include "InterchangeJointNode.h"
 #include "InterchangeMaterialFactoryNode.h"
+#include "InterchangeMeshDefinitions.h"
 #include "InterchangeMeshNode.h"
 #include "InterchangeSceneNode.h"
 #include "InterchangeStaticMeshFactoryNode.h"
@@ -1690,6 +1691,39 @@ namespace UE::Interchange::Private::MeshHelper
 	}
 #endif
 
+	EInterchangeMeshCollision GetMeshCollisionFromName(const FString& Name)
+	{
+		EInterchangeMeshCollision CollisionType = EInterchangeMeshCollision::None;
+
+		if (Name.StartsWith(TEXT("UBX_")))
+		{
+			CollisionType = EInterchangeMeshCollision::Box;
+		}
+		else if (Name.StartsWith(TEXT("UCX_")) || Name.StartsWith(TEXT("MCDCX_")))
+		{
+			CollisionType = EInterchangeMeshCollision::Convex18DOP;
+		}
+		else if (Name.StartsWith(TEXT("USP_")))
+		{
+			CollisionType = EInterchangeMeshCollision::Sphere;
+		}
+		else if (Name.StartsWith(TEXT("UCP_")))
+		{
+			CollisionType = EInterchangeMeshCollision::Capsule;
+		}
+
+		return CollisionType;
+	}
+
+	bool HasPotentialMeshCollisionPrefix(const FString& Name)
+	{
+		return Name.StartsWith(TEXT("UBX")) ||
+			Name.StartsWith(TEXT("UCX")) ||
+			Name.StartsWith(TEXT("USP")) ||
+			Name.StartsWith(TEXT("UCP")) ||
+			Name.StartsWith(TEXT("MCDCX"));
+	}
+
 #if WITH_EDITOR
 	/** Nanite assembly builder abstract base class (to hide underlying USkeletalMesh/UStaticMesh templated type) */
 
diff --git a/Engine/Plugins/Interchange/Runtime/Source/Import/Public/Mesh/InterchangeMeshHelper.h b/Engine/Plugins/Interchange/Runtime/Source/Import/Public/Mesh/InterchangeMeshHelper.h
--- a/Engine/Plugins/Interchange/Runtime/Source/Import/Public/Mesh/InterchangeMeshHelper.h
+++ b/Engine/Plugins/Interchange/Runtime/Source/Import/Public/Mesh/InterchangeMeshHelper.h
@@ -20,6 +20,7 @@ struct FKAggregateGeom;
 struct FSkeletalMaterial;
 struct FStaticMaterial;
 struct FInterchangeMeshPayLoadKey;
+enum class EInterchangeMeshCollision : uint8;
 
 namespace UE::Interchange
 {
@@ -138,6 +139,12 @@ namespace UE::Interchange::Private::MeshHelper
 	INTERCHANGEIMPORT_API void RemapSkeletalMeshVertexColorToMeshDescription(const USkeletalMesh* SkeletalMesh, const int32 LODIndex, FMeshDescription& MeshDescription);
 #endif
 
+	/* Checks name for UBX_, UCX_, MCDCX_, USP_, UCP_ prefixes.*/
+	INTERCHANGEIMPORT_API EInterchangeMeshCollision GetMeshCollisionFromName(const FString& Name);
+
+	/* Checks if Name has any of the prefixes mentioned in GetMeshCollisionFromName, without the '_', this primarily used for checking for typos. */
+	INTERCHANGEIMPORT_API bool HasPotentialMeshCollisionPrefix(const FString& Name);
+
 #if WITH_EDITOR
 	/** Nanite assembly builder helper for static and skeletal meshes. */
 	class FInterchangeNaniteAssemblyBuilder
diff --git a/Engine/Plugins/Interchange/Runtime/Source/Pipelines/Private/InterchangeLODDataParser.cpp b/Engine/Plugins/Interchange/Runtime/Source/Pipelines/Private/InterchangeLODDataParser.cpp
--- a/Engine/Plugins/Interchange/Runtime/Source/Pipelines/Private/InterchangeLODDataParser.cpp
+++ b/Engine/Plugins/Interchange/Runtime/Source/Pipelines/Private/InterchangeLODDataParser.cpp
@@ -116,24 +116,9 @@ namespace UE::Interchange
 				FString MeshName = UE::Interchange::Private::GetNodeName(PipelineMeshesUtilities, NodeContainer, NodeUid);
 
 				// Determine if the mesh name is a potential collision mesh
+				CollisionType = UE::Interchange::Private::MeshHelper::GetMeshCollisionFromName(MeshName);
 
-				if (MeshName.StartsWith(TEXT("UBX_")))
-				{
-					CollisionType = EInterchangeMeshCollision::Box;
-				}
-				else if (MeshName.StartsWith(TEXT("UCX_")) || MeshName.StartsWith(TEXT("MCDCX_")))
-				{
-					CollisionType = EInterchangeMeshCollision::Convex18DOP;
-				}
-				else if (MeshName.StartsWith(TEXT("USP_")))
-				{
-					CollisionType = EInterchangeMeshCollision::Sphere;
-				}
-				else if (MeshName.StartsWith(TEXT("UCP_")))
-				{
-					CollisionType = EInterchangeMeshCollision::Capsule;
-				}
-				else
+				if (CollisionType == EInterchangeMeshCollision::None)
 				{
 					return { EInterchangeMeshCollision::None, FString() };
 				}
@@ -398,8 +383,7 @@ void FInterchangeMeshLODDataParser<TMeshFactoryType>::ProcessCollision(const FSt
 {
 	//#interchange_CollisionRefactor_TODO: Validate the need for the special use case for mesh name generation: if there is only 1 SceneNode instantiating the mesh, it should be named after/by the Instantiating node.
 	const FString MeshName = UE::Interchange::Private::GetNodeName(PipelineMeshesUtilities, BaseNodeContainer, NodeUid);
-	// Ignore the "_" so we can detect potential typos done by the user.
-	const bool bHasCollisionPrefix = MeshName.StartsWith(TEXT("UBX")) || MeshName.StartsWith(TEXT("UCX")) || MeshName.StartsWith(TEXT("USP")) || MeshName.StartsWith(TEXT("UCP")) || MeshName.StartsWith(TEXT("MCDCX"));
+	const bool HasPotentialMeshCollisionPrefix = UE::Interchange::Private::MeshHelper::HasPotentialMeshCollisionPrefix(MeshName);
 
     //#interchange_LODRefactor_Note: using NodeUidTransformsPair (previously NodeUids) is LOD Index based, we are not actually checking against all NodeUids for the CollisionMeshType.
     TTuple<EInterchangeMeshCollision, FString> MeshType = UE::Interchange::CollisionHelper::GetCollisionMeshType(PipelineMeshesUtilities, BaseNodeContainer, NodeUid, AllNodeUids, GenericMeshPipeline->bImportCollisionAccordingToMeshName);
@@ -433,7 +417,7 @@ void FInterchangeMeshLODDataParser<TMeshFactoryType>::ProcessCollision(const FSt
     {
         case EInterchangeMeshCollision::None:
             // If the user is attempting to import a collision mesh using a naming convention but the mesh doesn't match, warn them that it won't be imported as collision.
-            if (bHasCollisionPrefix)
+            if (HasPotentialMeshCollisionPrefix)
             {
                 LogWarning(FText::Format(NSLOCTEXT("FInterchangeMeshLODDataParser", "NoMatchingRenderMeshForCustomCollision",
                     "No render mesh found for custom collision '{0}'. The name must follow the pattern 'UCX_[ExactRenderMeshName]', and a matching render mesh with that name must exist in the source file."),
diff --git a/Engine/Plugins/Interchange/Runtime/Source/Pipelines/Private/InterchangePipelineMeshesUtilities.cpp b/Engine/Plugins/Interchange/Runtime/Source/Pipelines/Private/InterchangePipelineMeshesUtilities.cpp
--- a/Engine/Plugins/Interchange/Runtime/Source/Pipelines/Private/InterchangePipelineMeshesUtilities.cpp
+++ b/Engine/Plugins/Interchange/Runtime/Source/Pipelines/Private/InterchangePipelineMeshesUtilities.cpp
@@ -8,6 +8,7 @@
 #include "InterchangeMeshNode.h"
 #include "InterchangeSceneComponentNodes.h"
 #include "InterchangeJointNode.h"
+#include "Mesh/InterchangeMeshHelper.h"
 #include "Nodes/InterchangeBaseNode.h"
 #include "Nodes/InterchangeBaseNodeContainer.h"
 #include "Nodes/InterchangeSourceNode.h"
@@ -53,10 +54,12 @@ namespace UE::Private::InterchangeMeshPipeline
 						return;
 					}
 
+					FString SceneNodeName = SceneNode->GetDisplayLabel();
+
 					int32 LODIndex;
 					FString LODContainerName;
 					FString Error;
-					if (UInterchangeMeshLODContainerNode::CheckForLODPattern(SceneNode->GetDisplayLabel(), LODIndex, LODContainerName, Error))
+					if (UInterchangeMeshLODContainerNode::CheckForLODPattern(SceneNodeName, LODIndex, LODContainerName, Error))
 					{
 						FString ParentUid = SceneNode->GetParentUid();
 						const UInterchangeSceneNode* ConstParentSceneNode = Cast<UInterchangeSceneNode>(BaseNodeContainer->GetNode(ParentUid));
@@ -65,7 +68,11 @@ namespace UE::Private::InterchangeMeshPipeline
 					}
 					else if (Error.Len())
 					{
-						UE_LOG(LogInterchangePipelineMeshesUtilities, Warning, TEXT("%s"), *Error);
+						if (UE::Interchange::Private::MeshHelper::GetMeshCollisionFromName(SceneNodeName) == EInterchangeMeshCollision::None)
+						{
+							//Do not warn for presumed Collision nodes: (for ep: UCX_SM_Test_LOD0)
+							UE_LOG(LogInterchangePipelineMeshesUtilities, Warning, TEXT("%s"), *Error);
+						}
 					}
 				}
 			}
