MaterialX:
- Add analytics

#Synced-CL 50150178

diff --git a/Engine/Plugins/Interchange/Runtime/Source/Import/Private/MaterialX/MaterialXUtils/MaterialXBase.cpp b/Engine/Plugins/Interchange/Runtime/Source/Import/Private/MaterialX/MaterialXUtils/MaterialXBase.cpp
--- a/Engine/Plugins/Interchange/Runtime/Source/Import/Private/MaterialX/MaterialXUtils/MaterialXBase.cpp
+++ b/Engine/Plugins/Interchange/Runtime/Source/Import/Private/MaterialX/MaterialXUtils/MaterialXBase.cpp
@@ -47,6 +47,11 @@ void FMaterialXBase::SetTranslator(const UInterchangeTranslatorBase* InTranslato
 	Translator = InTranslator;
 }
 
+void FMaterialXBase::CollectAnalytics(TArray<FAnalyticsEventAttribute>& OutAnalyticsAttributes)
+{
+	OutAnalyticsAttributes.Append(AnalyticsAttributes);
+}
+
 void FMaterialXBase::AddTexCoordToTiledImageNodes(MaterialX::ElementPtr Graph)
 {
 	for(mx::NodePtr Node : Graph->getChildrenOfType<mx::Node>(mx::Category::TiledImage))
diff --git a/Engine/Plugins/Interchange/Runtime/Source/Import/Private/MaterialX/MaterialXUtils/MaterialXDisplacementShader.cpp b/Engine/Plugins/Interchange/Runtime/Source/Import/Private/MaterialX/MaterialXUtils/MaterialXDisplacementShader.cpp
--- a/Engine/Plugins/Interchange/Runtime/Source/Import/Private/MaterialX/MaterialXUtils/MaterialXDisplacementShader.cpp
+++ b/Engine/Plugins/Interchange/Runtime/Source/Import/Private/MaterialX/MaterialXUtils/MaterialXDisplacementShader.cpp
@@ -54,6 +54,8 @@ UInterchangeBaseNode* FMaterialXDisplacementShader::Translate(MaterialX::NodePtr
 	// Outputs
 	UInterchangeShaderPortsAPI::ConnectOuputToInputByName(ShaderGraphNode, Common::Parameters::Displacement.ToString(), DisplacementShaderNode->GetUniqueID(), Common::Parameters::Displacement.ToString());
 
+	AnalyticsAttributes.Emplace(TEXT("ShaderType"), TEXT("Displacement Shader"));
+
 	return DisplacementShaderNode;
 }
 #endif
\ No newline at end of file
diff --git a/Engine/Plugins/Interchange/Runtime/Source/Import/Private/MaterialX/MaterialXUtils/MaterialXManager.cpp b/Engine/Plugins/Interchange/Runtime/Source/Import/Private/MaterialX/MaterialXUtils/MaterialXManager.cpp
--- a/Engine/Plugins/Interchange/Runtime/Source/Import/Private/MaterialX/MaterialXUtils/MaterialXManager.cpp
+++ b/Engine/Plugins/Interchange/Runtime/Source/Import/Private/MaterialX/MaterialXUtils/MaterialXManager.cpp
@@ -23,9 +23,10 @@
 #include "MaterialX/MaterialXUtils/MaterialXVolumeShader.h"
 #include "MaterialX/MaterialXUtils/MaterialXVolumeMaterial.h"
 
-#include "InterchangeMaterialDefinitions.h"
-#include "InterchangeImportLog.h"
+#include "InterchangeAnalyticsHelper.h"
 #include "InterchangeHelper.h"
+#include "InterchangeImportLog.h"
+#include "InterchangeMaterialDefinitions.h"
 
 #define LOCTEXT_NAMESPACE "InterchangeMaterialXManager"
 
@@ -546,6 +547,8 @@ bool FMaterialXManager::Translate(MaterialX::DocumentPtr Document, UInterchangeB
 	bool bIsDocumentValid;
 	UE::Interchange::FScopedLambda ScopedLambda([this](){TextureNodeUids.Empty();});
 
+	TArray<FAnalyticsEventAttribute> AnalyticsAttributes;
+
 	try
 	{
 		// Read the document to be sure that the file is valid (meaning all nodes have their nodedef and typedef well-defined)
@@ -582,6 +585,7 @@ bool FMaterialXManager::Translate(MaterialX::DocumentPtr Document, UInterchangeB
 					{
 						ShaderTranslator->SetTranslator(Translator);
 						ShaderTranslator->Translate(Node);
+						ShaderTranslator->CollectAnalytics(AnalyticsAttributes);
 					}
 				}
 			}
@@ -595,6 +599,7 @@ bool FMaterialXManager::Translate(MaterialX::DocumentPtr Document, UInterchangeB
 			UInterchangeResultError_Generic* Message = Translator->AddMessage<UInterchangeResultError_Generic>();
 			Message->Text = FText::Format(LOCTEXT("MaterialXException", "{0}"),
 										  FText::FromString(Exception.what()));
+			AnalyticsAttributes.Emplace(TEXT("MaterialXFail"), TEXT("[FAILED] MaterialX Document invalid"));
 		}
 		else
 		{
@@ -602,6 +607,11 @@ bool FMaterialXManager::Translate(MaterialX::DocumentPtr Document, UInterchangeB
 		}
 	}
 
+	if (Translator && Translator->AnalyticsHelper)
+	{
+		Translator->AnalyticsHelper->AppendThreadSafe(TEXT("Interchange.Usage.Import.MaterialX"), AnalyticsAttributes);
+	}
+
 	return bIsDocumentValid;
 }
 
diff --git a/Engine/Plugins/Interchange/Runtime/Source/Import/Private/MaterialX/MaterialXUtils/MaterialXMixShader.cpp b/Engine/Plugins/Interchange/Runtime/Source/Import/Private/MaterialX/MaterialXUtils/MaterialXMixShader.cpp
--- a/Engine/Plugins/Interchange/Runtime/Source/Import/Private/MaterialX/MaterialXUtils/MaterialXMixShader.cpp
+++ b/Engine/Plugins/Interchange/Runtime/Source/Import/Private/MaterialX/MaterialXUtils/MaterialXMixShader.cpp
@@ -47,6 +47,8 @@ UInterchangeBaseNode* FMaterialXMixShader::Translate(mx::NodePtr MixNode)
 		MixShaderNode = CreateShaderNode(MixNode, MixNode->getName().c_str(), UE::Expressions::Names::SubstrateHorizontalMixing);
         MixShaderNode = FMaterialXSurfaceShaderAbstract::Translate(MixShaderNode);
         UInterchangeShaderPortsAPI::ConnectDefaultOuputToInput(ShaderGraphNode, SubstrateMaterial::Parameters::FrontMaterial.ToString(), MixShaderNode->GetUniqueID());
+
+		AnalyticsAttributes.Emplace(TEXT("ShaderType"), TEXT("Mix Shader"));
     }
 	else
 	{
diff --git a/Engine/Plugins/Interchange/Runtime/Source/Import/Private/MaterialX/MaterialXUtils/MaterialXOpenPBRSurfaceShader.cpp b/Engine/Plugins/Interchange/Runtime/Source/Import/Private/MaterialX/MaterialXUtils/MaterialXOpenPBRSurfaceShader.cpp
--- a/Engine/Plugins/Interchange/Runtime/Source/Import/Private/MaterialX/MaterialXUtils/MaterialXOpenPBRSurfaceShader.cpp
+++ b/Engine/Plugins/Interchange/Runtime/Source/Import/Private/MaterialX/MaterialXUtils/MaterialXOpenPBRSurfaceShader.cpp
@@ -27,8 +27,8 @@ UInterchangeBaseNode* FMaterialXOpenPBRSurfaceShader::Translate(MaterialX::NodeP
 	using namespace UE::Interchange::Materials::Standard::Nodes;
 	this->SurfaceShaderNode = OpenPBRSurfaceNode;
 
-	UInterchangeShaderNode* OpenPBRSurfaceShaderNode = FMaterialXSurfaceShaderAbstract::Translate(EInterchangeMaterialXShaders::OpenPBRSurface);
-
+	EInterchangeMaterialXShaders OpenPBRShaderType = EInterchangeMaterialXShaders::OpenPBRSurface;
+	UInterchangeShaderNode* OpenPBRSurfaceShaderNode = FMaterialXSurfaceShaderAbstract::Translate(OpenPBRShaderType);
 	//Two sided
 	if(MaterialX::InputPtr Input = GetInput(SurfaceShaderNode, mx::OpenPBRSurface::Input::GeometryThinWalled);
 	   Input->hasValue() && mx::fromValueString<bool>(Input->getValueString()) == true)
@@ -40,23 +40,27 @@ UInterchangeBaseNode* FMaterialXOpenPBRSurfaceShader::Translate(MaterialX::NodeP
 
 	if(UInterchangeShaderPortsAPI::HasInput(OpenPBRSurfaceShaderNode, OpenPBRSurface::Parameters::TransmissionWeight))
 	{
+		OpenPBRShaderType = EInterchangeMaterialXShaders::OpenPBRSurfaceTransmission;
 		OpenPBRSurfaceShaderNode->AddInt32Attribute(Attributes::EnumType, IndexSurfaceShaders);
-		OpenPBRSurfaceShaderNode->AddInt32Attribute(Attributes::EnumValue, int32(EInterchangeMaterialXShaders::OpenPBRSurfaceTransmission));
+		OpenPBRSurfaceShaderNode->AddInt32Attribute(Attributes::EnumValue, int32(OpenPBRShaderType));
 		ShaderGraphNode->SetCustomBlendMode(EBlendMode::BLEND_TranslucentColoredTransmittance);
 	}
 
+	bool bSubsurfaceWithTransmission = false;
 	if(bIsSubstrateEnabled)
 	{
 		if (!bIsSubstrateAdaptiveGBufferEnabled)
 		{
 			MTLX_LOG("MaterialXOpenPBRSurfaceShader", "OpenPBR material rendering might be wrong. Please select Substrate Adaptive GBuffer format in the Project settings.");
 		}
-		
-		if(!ConnectSubsurfaceTransmissionToSurfaceShader(OpenPBRSurfaceShaderNode, EInterchangeMaterialXShaders::OpenPBRSurface
-		   , OpenPBRSurface::Parameters::BaseWeight, OpenPBRSurface::Parameters::BaseColor
-		   , OpenPBRSurface::Parameters::TransmissionWeight, OpenPBRSurface::Parameters::TransmissionColor, OpenPBRSurface::Parameters::TransmissionDepth, OpenPBRSurface::Parameters::TransmissionDispersionAbbeNumber, OpenPBRSurface::Parameters::TransmissionDispersionScale, OpenPBRSurface::Parameters::TransmissionScatter, OpenPBRSurface::Parameters::TransmissionScatterAnisotropy
-		   , OpenPBRSurface::Parameters::SubsurfaceWeight, OpenPBRSurface::Parameters::SubsurfaceColor, OpenPBRSurface::Parameters::SubsurfaceRadiusScale
-		   , OpenPBRSurface::SubstrateMaterial::Outputs::FrontMaterial, OpenPBRSurface::SubstrateMaterial::Outputs::FrontMaterial))
+
+		bSubsurfaceWithTransmission = ConnectSubsurfaceTransmissionToSurfaceShader(OpenPBRSurfaceShaderNode, EInterchangeMaterialXShaders::OpenPBRSurface
+			, OpenPBRSurface::Parameters::BaseWeight, OpenPBRSurface::Parameters::BaseColor
+			, OpenPBRSurface::Parameters::TransmissionWeight, OpenPBRSurface::Parameters::TransmissionColor, OpenPBRSurface::Parameters::TransmissionDepth, OpenPBRSurface::Parameters::TransmissionDispersionAbbeNumber, OpenPBRSurface::Parameters::TransmissionDispersionScale, OpenPBRSurface::Parameters::TransmissionScatter, OpenPBRSurface::Parameters::TransmissionScatterAnisotropy
+			, OpenPBRSurface::Parameters::SubsurfaceWeight, OpenPBRSurface::Parameters::SubsurfaceColor, OpenPBRSurface::Parameters::SubsurfaceRadiusScale
+			, OpenPBRSurface::SubstrateMaterial::Outputs::FrontMaterial, OpenPBRSurface::SubstrateMaterial::Outputs::FrontMaterial);
+
+		if(!bSubsurfaceWithTransmission)
 		{
 			UInterchangeShaderPortsAPI::ConnectOuputToInputByName(ShaderGraphNode, SubstrateMaterial::Parameters::FrontMaterial.ToString(), OpenPBRSurfaceShaderNode->GetUniqueID(), OpenPBRSurface::SubstrateMaterial::Outputs::FrontMaterial.ToString());
 
@@ -103,6 +107,22 @@ UInterchangeBaseNode* FMaterialXOpenPBRSurfaceShader::Translate(MaterialX::NodeP
 		}
 	}
 
+	AnalyticsAttributes.Emplace(TEXT("ShaderType"), TEXT("OpenPBR"));
+	
+	// this analytic can also be set in ConnectSubsurfaceTransmissionToSurfaceShader, we don't want to set it twice
+	if (!bSubsurfaceWithTransmission)
+	{
+		switch (OpenPBRShaderType)
+		{
+		case EInterchangeMaterialXShaders::OpenPBRSurfaceTransmission:
+			AnalyticsAttributes.Emplace(TEXT("OpenPBR"), TEXT("Transmission"));
+			break;
+		default:
+			AnalyticsAttributes.Emplace(TEXT("OpenPBR"), TEXT("Opaque"));
+			break;
+		}
+	}
+
 	return OpenPBRSurfaceShaderNode;
 }
 
diff --git a/Engine/Plugins/Interchange/Runtime/Source/Import/Private/MaterialX/MaterialXUtils/MaterialXStandardSurfaceShader.cpp b/Engine/Plugins/Interchange/Runtime/Source/Import/Private/MaterialX/MaterialXUtils/MaterialXStandardSurfaceShader.cpp
--- a/Engine/Plugins/Interchange/Runtime/Source/Import/Private/MaterialX/MaterialXUtils/MaterialXStandardSurfaceShader.cpp
+++ b/Engine/Plugins/Interchange/Runtime/Source/Import/Private/MaterialX/MaterialXUtils/MaterialXStandardSurfaceShader.cpp
@@ -27,7 +27,8 @@ UInterchangeBaseNode* FMaterialXStandardSurfaceShader::Translate(mx::NodePtr Sta
 	using namespace UE::Interchange::MaterialX;
 	using namespace UE::Interchange::Materials::Standard::Nodes;
 
-	UInterchangeShaderNode* StandardSurfaceShaderNode = FMaterialXSurfaceShaderAbstract::Translate(EInterchangeMaterialXShaders::StandardSurface);
+	EInterchangeMaterialXShaders StandardSurfaceShaderType = EInterchangeMaterialXShaders::StandardSurface;
+	UInterchangeShaderNode* StandardSurfaceShaderNode = FMaterialXSurfaceShaderAbstract::Translate(StandardSurfaceShaderType);
 
 	//Two sided
 	if(MaterialX::InputPtr Input = GetInput(SurfaceShaderNode, mx::StandardSurface::Input::ThinWalled);
@@ -40,11 +41,12 @@ UInterchangeBaseNode* FMaterialXStandardSurfaceShader::Translate(mx::NodePtr Sta
 
 	if(UInterchangeShaderPortsAPI::HasInput(StandardSurfaceShaderNode, StandardSurface::Parameters::Transmission))
 	{
+		StandardSurfaceShaderType = EInterchangeMaterialXShaders::StandardSurfaceTransmission;
 		StandardSurfaceShaderNode->AddInt32Attribute(Attributes::EnumType, IndexSurfaceShaders);
 		StandardSurfaceShaderNode->AddInt32Attribute(Attributes::EnumValue, int32(EInterchangeMaterialXShaders::StandardSurfaceTransmission));
 	}
 
-	// Outputs
+	bool bSubsurfaceWithTransmission = false;
 	if(bIsSubstrateEnabled)
 	{
 		if (!bIsSubstrateAdaptiveGBufferEnabled)
@@ -52,11 +54,13 @@ UInterchangeBaseNode* FMaterialXStandardSurfaceShader::Translate(mx::NodePtr Sta
 			MTLX_LOG("MaterialXStandardSurfaceShader", "Standard Surface material rendering might be wrong. Please select Substrate Adaptive GBuffer format in the Project settings.");
 		}
 
-		if (!ConnectSubsurfaceTransmissionToSurfaceShader(StandardSurfaceShaderNode, EInterchangeMaterialXShaders::StandardSurface
+		bSubsurfaceWithTransmission = ConnectSubsurfaceTransmissionToSurfaceShader(StandardSurfaceShaderNode, EInterchangeMaterialXShaders::StandardSurface
 			, StandardSurface::Parameters::Base, StandardSurface::Parameters::BaseColor
 			, StandardSurface::Parameters::Transmission, StandardSurface::Parameters::TransmissionColor, StandardSurface::Parameters::TransmissionDepth, StandardSurface::Parameters::TransmissionExtraRoughness, StandardSurface::Parameters::TransmissionDispersion, StandardSurface::Parameters::TransmissionScatter, StandardSurface::Parameters::TransmissionScatterAnisotropy
 			, StandardSurface::Parameters::Subsurface, StandardSurface::Parameters::SubsurfaceColor, StandardSurface::Parameters::SubsurfaceRadius
-			, StandardSurface::SubstrateMaterial::Outputs::Opaque, StandardSurface::SubstrateMaterial::Outputs::Translucent))
+			, StandardSurface::SubstrateMaterial::Outputs::Opaque, StandardSurface::SubstrateMaterial::Outputs::Translucent);
+
+		if (!bSubsurfaceWithTransmission)
 		{
 			if (UInterchangeShaderPortsAPI::HasInput(StandardSurfaceShaderNode, StandardSurface::Parameters::Transmission))
 			{
@@ -120,6 +124,22 @@ UInterchangeBaseNode* FMaterialXStandardSurfaceShader::Translate(mx::NodePtr Sta
 		}
 	}
 
+	AnalyticsAttributes.Emplace(TEXT("ShaderType"), TEXT("Standard Surface"));
+
+	// this analytic can also be set in ConnectSubsurfaceTransmissionToSurfaceShader, we don't want to set it twice
+	if (!bSubsurfaceWithTransmission)
+	{
+		switch (StandardSurfaceShaderType)
+		{
+		case EInterchangeMaterialXShaders::StandardSurfaceTransmission:
+			AnalyticsAttributes.Emplace(TEXT("StandardSurface"), TEXT("Transmission"));
+			break;
+		default:
+			AnalyticsAttributes.Emplace(TEXT("StandardSurface"), TEXT("Opaque"));
+			break;
+		}
+	}
+
 	return StandardSurfaceShaderNode;
 }
 
diff --git a/Engine/Plugins/Interchange/Runtime/Source/Import/Private/MaterialX/MaterialXUtils/MaterialXSurfaceMaterial.cpp b/Engine/Plugins/Interchange/Runtime/Source/Import/Private/MaterialX/MaterialXUtils/MaterialXSurfaceMaterial.cpp
--- a/Engine/Plugins/Interchange/Runtime/Source/Import/Private/MaterialX/MaterialXUtils/MaterialXSurfaceMaterial.cpp
+++ b/Engine/Plugins/Interchange/Runtime/Source/Import/Private/MaterialX/MaterialXUtils/MaterialXSurfaceMaterial.cpp
@@ -106,9 +106,11 @@ UInterchangeBaseNode* FMaterialXSurfaceMaterial::Translate(MaterialX::NodePtr Su
 		UInterchangeShaderPortsAPI::ConnectOuputToInputByName(SurfaceShaderNode, InputNormalName, DisplacementShaderNode->GetUniqueID(), PBRMR::Parameters::Normal.ToString());
 	}
 
+	AnalyticsAttributes.Emplace(TEXT("MaterialType"), TEXT("Surface Material"));
 	if(!bHasSurfaceShader)
 	{
 		MTLX_LOG("MaterialXSurfaceMaterial", "<{0}>: {1} \"{2}\" is not supported", SurfaceMaterialNode->getCategory().c_str(), SurfaceMaterialNode->getType().c_str(), SurfaceMaterialNode->getName().c_str());
+		AnalyticsAttributes.Emplace(TEXT("ShaderType"), TEXT("Surface Shader unsupported"));
 	}
 
 	return ShaderGraphNode;
diff --git a/Engine/Plugins/Interchange/Runtime/Source/Import/Private/MaterialX/MaterialXUtils/MaterialXSurfaceShader.cpp b/Engine/Plugins/Interchange/Runtime/Source/Import/Private/MaterialX/MaterialXUtils/MaterialXSurfaceShader.cpp
--- a/Engine/Plugins/Interchange/Runtime/Source/Import/Private/MaterialX/MaterialXUtils/MaterialXSurfaceShader.cpp
+++ b/Engine/Plugins/Interchange/Runtime/Source/Import/Private/MaterialX/MaterialXUtils/MaterialXSurfaceShader.cpp
@@ -42,6 +42,7 @@ UInterchangeBaseNode* FMaterialXSurfaceShader::Translate(MaterialX::NodePtr Surf
 		}
 	}
 
+	AnalyticsAttributes.Emplace(TEXT("ShaderType"), TEXT("Surface (Generic)"));
 	return FunctionSurfaceShaderNode;
 }
 #endif
\ No newline at end of file
diff --git a/Engine/Plugins/Interchange/Runtime/Source/Import/Private/MaterialX/MaterialXUtils/MaterialXSurfaceShaderAbstract.cpp b/Engine/Plugins/Interchange/Runtime/Source/Import/Private/MaterialX/MaterialXUtils/MaterialXSurfaceShaderAbstract.cpp
--- a/Engine/Plugins/Interchange/Runtime/Source/Import/Private/MaterialX/MaterialXUtils/MaterialXSurfaceShaderAbstract.cpp
+++ b/Engine/Plugins/Interchange/Runtime/Source/Import/Private/MaterialX/MaterialXUtils/MaterialXSurfaceShaderAbstract.cpp
@@ -1498,6 +1498,13 @@ bool FMaterialXSurfaceShaderAbstract::ConnectSubsurfaceTransmissionToSurfaceShad
 		UInterchangeShaderPortsAPI::HasInput(OpenPbrStandarSurfaceShaderNode, SubsurfaceWeight) &&
 		UInterchangeShaderPortsAPI::HasInput(OpenPbrStandarSurfaceShaderNode, BaseColor))
 	{
+		FString SurfaceType = TEXT("Standard Surface");
+		if (SurfaceShaderType == EInterchangeMaterialXShaders::OpenPBRSurface)
+		{
+			SurfaceType = TEXT("OpenPBR");
+		}
+		AnalyticsAttributes.Emplace(SurfaceType, "Opaque");
+
 		bSubsurfaceTransmission = true;
 		// 1. Set the Material as Opaque with Mask Mode
 		OpenPbrStandarSurfaceShaderNode->AddInt32Attribute(Attributes::EnumType, IndexSurfaceShaders);
diff --git a/Engine/Plugins/Interchange/Runtime/Source/Import/Private/MaterialX/MaterialXUtils/MaterialXSurfaceUnlitShader.cpp b/Engine/Plugins/Interchange/Runtime/Source/Import/Private/MaterialX/MaterialXUtils/MaterialXSurfaceUnlitShader.cpp
--- a/Engine/Plugins/Interchange/Runtime/Source/Import/Private/MaterialX/MaterialXUtils/MaterialXSurfaceUnlitShader.cpp
+++ b/Engine/Plugins/Interchange/Runtime/Source/Import/Private/MaterialX/MaterialXUtils/MaterialXSurfaceUnlitShader.cpp
@@ -57,6 +57,7 @@ UInterchangeBaseNode* FMaterialXSurfaceUnlitShader::Translate(mx::NodePtr Surfac
 		}
 	}
 
+	AnalyticsAttributes.Emplace(TEXT("ShaderType"), TEXT("Surface Unlit"));
 	return SurfaceUnlitShaderNode;
 }
 #endif
\ No newline at end of file
diff --git a/Engine/Plugins/Interchange/Runtime/Source/Import/Private/MaterialX/MaterialXUtils/MaterialXUsdPreviewSurfaceShader.cpp b/Engine/Plugins/Interchange/Runtime/Source/Import/Private/MaterialX/MaterialXUtils/MaterialXUsdPreviewSurfaceShader.cpp
--- a/Engine/Plugins/Interchange/Runtime/Source/Import/Private/MaterialX/MaterialXUtils/MaterialXUsdPreviewSurfaceShader.cpp
+++ b/Engine/Plugins/Interchange/Runtime/Source/Import/Private/MaterialX/MaterialXUtils/MaterialXUsdPreviewSurfaceShader.cpp
@@ -69,6 +69,8 @@ UInterchangeBaseNode* FMaterialXUsdPreviewSurfaceShader::Translate(MaterialX::No
 		}
 	}
 
+	AnalyticsAttributes.Emplace(TEXT("ShaderType"), TEXT("USD Preview Surface"));
+
 	return UsdPreviewSurfaceShaderNode;
 }
 
diff --git a/Engine/Plugins/Interchange/Runtime/Source/Import/Private/MaterialX/MaterialXUtils/MaterialXVolumeMaterial.cpp b/Engine/Plugins/Interchange/Runtime/Source/Import/Private/MaterialX/MaterialXUtils/MaterialXVolumeMaterial.cpp
--- a/Engine/Plugins/Interchange/Runtime/Source/Import/Private/MaterialX/MaterialXUtils/MaterialXVolumeMaterial.cpp
+++ b/Engine/Plugins/Interchange/Runtime/Source/Import/Private/MaterialX/MaterialXUtils/MaterialXVolumeMaterial.cpp
@@ -51,10 +51,13 @@ UInterchangeBaseNode* FMaterialXVolumeMaterial::Translate(MaterialX::NodePtr Vol
 				ShaderTranslator->SetTranslator(Translator);
 				ShaderTranslator->Translate(VolumeShaderNode);
 				NodeContainer.AddNode(ShaderGraphNode);
+
+				AnalyticsAttributes.Emplace(TEXT("MaterialType"), TEXT("Volume Material"));
 			}
 			else
 			{
 				MTLX_LOG("MaterialXVolumeMaterial", "<{0}>: {1} \"{2}\" is not supported", VolumeShaderNode->getCategory().c_str(), VolumeShaderNode->getType().c_str(), VolumeShaderNode->getName().c_str());
+				AnalyticsAttributes.Emplace(TEXT("ShaderType"), TEXT("Volume Shader unsupported"));
 			}
 		}
 	}
diff --git a/Engine/Plugins/Interchange/Runtime/Source/Import/Private/MaterialX/MaterialXUtils/MaterialXVolumeShader.cpp b/Engine/Plugins/Interchange/Runtime/Source/Import/Private/MaterialX/MaterialXUtils/MaterialXVolumeShader.cpp
--- a/Engine/Plugins/Interchange/Runtime/Source/Import/Private/MaterialX/MaterialXUtils/MaterialXVolumeShader.cpp
+++ b/Engine/Plugins/Interchange/Runtime/Source/Import/Private/MaterialX/MaterialXUtils/MaterialXVolumeShader.cpp
@@ -31,6 +31,8 @@ UInterchangeBaseNode* FMaterialXVolumeShader::Translate(MaterialX::NodePtr Volum
 	// Outputs
 	UInterchangeShaderPortsAPI::ConnectDefaultOuputToInput(ShaderGraphNode, SubstrateMaterial::Parameters::FrontMaterial.ToString(), FunctionVolumeShaderNode->GetUniqueID());
 
+	AnalyticsAttributes.Emplace(TEXT("ShaderType"), TEXT("Volume (Generic)"));
+
 	return FunctionVolumeShaderNode;
 }
 #endif
\ No newline at end of file
diff --git a/Engine/Plugins/Interchange/Runtime/Source/Import/Public/MaterialX/MaterialXUtils/MaterialXBase.h b/Engine/Plugins/Interchange/Runtime/Source/Import/Public/MaterialX/MaterialXUtils/MaterialXBase.h
--- a/Engine/Plugins/Interchange/Runtime/Source/Import/Public/MaterialX/MaterialXUtils/MaterialXBase.h
+++ b/Engine/Plugins/Interchange/Runtime/Source/Import/Public/MaterialX/MaterialXUtils/MaterialXBase.h
@@ -5,6 +5,7 @@
 #if WITH_EDITOR
 
 #include "Import/Private/InterchangeImportLog.h"
+#include "AnalyticsEventAttribute.h"
 #include "InterchangeMaterialDefinitions.h"
 #include "InterchangeTranslatorBase.h"
 #include "Logging/StructuredLog.h"
@@ -68,6 +69,8 @@ public:
 	/** Set the translator, only used for logging messages*/
 	void SetTranslator(const UInterchangeTranslatorBase* InTranslator);
 
+	void CollectAnalytics(TArray<FAnalyticsEventAttribute>& OutAnalyticsAttributes);
+
 protected:
 
 	/**
@@ -164,9 +167,12 @@ protected:
 
 	UInterchangeBaseNodeContainer& NodeContainer;
 
-	/** Used for logging */
+	/** Used for logging from either MaterialX or USD translator*/
 	const UInterchangeTranslatorBase* Translator;
 
+	/** To avoid locking Analytics helper in the translator for each add, at the end of the translat,  every attribute will be concatenated*/
+	TArray<FAnalyticsEventAttribute> AnalyticsAttributes;
+
 	/** It's up to the derive class to initialize it with the proper node definition*/
 	const char* NodeDefinition; 
 
