Fixed FBX SDK import crashes the editor

#Synced-CL 50703335

diff --git a/Engine/Plugins/Interchange/Runtime/Source/Pipelines/Private/InterchangeGenericCommonSkeletalMeshesAndAnimationsProperties.cpp b/Engine/Plugins/Interchange/Runtime/Source/Pipelines/Private/InterchangeGenericCommonSkeletalMeshesAndAnimationsProperties.cpp
--- a/Engine/Plugins/Interchange/Runtime/Source/Pipelines/Private/InterchangeGenericCommonSkeletalMeshesAndAnimationsProperties.cpp
+++ b/Engine/Plugins/Interchange/Runtime/Source/Pipelines/Private/InterchangeGenericCommonSkeletalMeshesAndAnimationsProperties.cpp
@@ -61,7 +61,7 @@ UInterchangeSkeletonFactoryNode* UInterchangeGenericCommonSkeletalMeshesAndAnima
 			UE::Interchange::Private::FSkeletonHelper::RecursiveBoneHasBindPose(BaseNodeContainer, RootJointUid, bHasBoneWithoutBindPose);
 			if (bHasBoneWithoutBindPose)
 			{
-				if (!GIsAutomationTesting)
+				if (!GIsAutomationTesting && Results)
 				{
 					UInterchangeResultWarning_Generic* Message = AddMessage<UInterchangeResultWarning_Generic>();
 					Message->Text = NSLOCTEXT("InterchangeSkeletalMeshFactory", "CreatePayloadTasks_ForceRebindOfSkinWithTimeZeroPose", "Imported skeleton has some invalid bind poses. Skeletal mesh skinning has been rebind using the time zero pose.");
diff --git a/Engine/Plugins/Interchange/Runtime/Source/Pipelines/Private/InterchangeGenericMeshPipeline.cpp b/Engine/Plugins/Interchange/Runtime/Source/Pipelines/Private/InterchangeGenericMeshPipeline.cpp
--- a/Engine/Plugins/Interchange/Runtime/Source/Pipelines/Private/InterchangeGenericMeshPipeline.cpp
+++ b/Engine/Plugins/Interchange/Runtime/Source/Pipelines/Private/InterchangeGenericMeshPipeline.cpp
@@ -483,7 +483,13 @@ void UInterchangeGenericMeshPipeline::ExecutePipeline(UInterchangeBaseNodeContai
 		UE_LOG(LogInterchangePipeline, Warning, TEXT("UInterchangeGenericMeshPipeline: Cannot execute pre-import pipeline because InBaseNodeContrainer is null."));
 		return;
 	}
-	
+
+	if (!CommonSkeletalMeshesAndAnimationsProperties.IsValid())
+	{
+		return;
+	}
+	CommonSkeletalMeshesAndAnimationsProperties->SetResultsContainer(Results);
+
 	BaseNodeContainer = InBaseNodeContainer;
 	SourceDatas.Empty(InSourceDatas.Num());
 	for (const UInterchangeSourceData* SourceData : InSourceDatas)
