Submitting on behalf of JeanLuc.Corenthin

Fixed issue with application of reimport strategy on material slots of mesh components
- Root cause: ActorHelper::ApplySlotMaterialDependencies was not applying reimport strategy
- Solution: Added support for reimport strategy to ActorHelper::ApplySlotMaterialDependencies
- Side effect: UE-349293 was also fixed.

Added Python Automation Tests as well.

#rnx

#Synced-CL 50959190

diff --git a/Engine/Plugins/Interchange/Runtime/Source/Import/Private/Scene/InterchangeActorHelper.cpp b/Engine/Plugins/Interchange/Runtime/Source/Import/Private/Scene/InterchangeActorHelper.cpp
--- a/Engine/Plugins/Interchange/Runtime/Source/Import/Private/Scene/InterchangeActorHelper.cpp
+++ b/Engine/Plugins/Interchange/Runtime/Source/Import/Private/Scene/InterchangeActorHelper.cpp
@@ -4,6 +4,7 @@
 
 #include "Engine/Engine.h"
 #include "Engine/Level.h"
+#include "Engine/World.h"
 #include "GameFramework/WorldSettings.h"
 #include "InterchangeActorFactoryNode.h"
 #include "InterchangeAssetUserData.h"
@@ -12,6 +13,7 @@
 #include "InterchangeMaterialFactoryNode.h"
 #include "InterchangeMeshActorFactoryNode.h"
 #include "InterchangeMeshNode.h"
+#include "InterchangeSceneImportAsset.h"
 #include "InterchangeSceneNode.h"
 #include "Nodes/InterchangeBaseNode.h"
 #include "Nodes/InterchangeBaseNodeContainer.h"
@@ -170,7 +172,7 @@ namespace UE::Interchange::ActorHelper
 				}
 			}
 
-			if (ParentActor || !SpawnedActor->IsAttachedTo(ParentActor))
+			if (ParentActor && !SpawnedActor->IsAttachedTo(ParentActor))
 			{
 				SpawnedActor->AttachToActor(ParentActor, FAttachmentTransformRules::KeepRelativeTransform);
 			}
@@ -208,20 +210,138 @@ namespace UE::Interchange::ActorHelper
 		return AssetTargetNodeIds.IsEmpty() ? nullptr : NodeContainer->GetFactoryNode(AssetTargetNodeIds[0]);
 	}
 
-	void ApplySlotMaterialDependencies(const UInterchangeBaseNodeContainer& NodeContainer, const UInterchangeMeshActorFactoryNode& MeshActorFactoryNode, UMeshComponent& MeshComponent)
+	void ApplySlotMaterialDependencies(const UInterchangeBaseNodeContainer& NodeContainer, const UInterchangeMeshActorFactoryNode& FactoryNode, UMeshComponent& MeshComponent, const UInterchangeMeshActorFactoryNode* PreviousFactoryNode)
 	{
-		// Set material slots from imported materials
+		const EReimportStrategyFlags ReimportStrategyFlags = FactoryNode.GetReimportStrategyFlags();
+		if (PreviousFactoryNode && ReimportStrategyFlags == EReimportStrategyFlags::ApplyNoProperties)
+		{
+			// This is a reimport but factory node should be ignored, return
+			return;
+		}
+
+		// Get material slots from imported factory node
 		TMap<FString, FString> SlotMaterialDependencies;
-		MeshActorFactoryNode.GetSlotMaterialDependencies(SlotMaterialDependencies);
-		for (TPair<FString, FString>& SlotMaterialDependency : SlotMaterialDependencies)
+		FactoryNode.GetSlotMaterialDependencies(SlotMaterialDependencies);
+		if (SlotMaterialDependencies.IsEmpty())
 		{
-			const UInterchangeBaseMaterialFactoryNode* MaterialFactoryNode = Cast<UInterchangeBaseMaterialFactoryNode>(NodeContainer.GetNode(SlotMaterialDependency.Value));
-			if (!MaterialFactoryNode || !MaterialFactoryNode->IsEnabled())
+			if (PreviousFactoryNode && ReimportStrategyFlags == EReimportStrategyFlags::ApplyPipelineProperties)
 			{
-				continue;
+				// Imported factory node does not have material overrides and factory node prevails, remove the existing ones
+				if (MeshComponent.HasOverrideMaterials())
+				{
+					MeshComponent.EmptyOverrideMaterials();
+				}
 			}
-			FSoftObjectPath ReferenceObject;
-			MaterialFactoryNode->GetCustomReferenceObject(ReferenceObject);
+
+			return;
+		}
+
+		auto GetMaterialPath = [&NodeContainer](const FString& MaterialUid) -> FSoftObjectPath
+			{
+				const UInterchangeBaseMaterialFactoryNode* MaterialFactoryNode = Cast<UInterchangeBaseMaterialFactoryNode>(NodeContainer.GetNode(MaterialUid));
+				if (MaterialFactoryNode && MaterialFactoryNode->IsEnabled())
+				{
+					FSoftObjectPath ReferenceObject;
+					MaterialFactoryNode->GetCustomReferenceObject(ReferenceObject);
+					if (ReferenceObject.IsValid())
+					{
+						return ReferenceObject;
+					}
+				}
+
+				return FSoftObjectPath();
+			};
+
+		if (PreviousFactoryNode && ReimportStrategyFlags == EReimportStrategyFlags::ApplyEditorChangedProperties)
+		{
+			using namespace UE::Interchange;
+			// Retrieve the InterchangeSceneImportAsset object from the last import
+			UInterchangeSceneImportAsset* SceneImportAsset = nullptr;
+			if (AActor* Actor = MeshComponent.GetOwner())
+			{
+				if (USceneComponent* RootComponent = Actor->GetRootComponent())
+				{
+					UInterchangeAssetUserData* AssetUserData = RootComponent->GetAssetUserData<UInterchangeAssetUserData>();
+					if (ensure(AssetUserData && AssetUserData->MetaData.Contains(InterchangeReset::Constants::SceneImportAssetPathKey)))
+					{
+						const FSoftObjectPath SceneImportAssetPath(AssetUserData->MetaData[InterchangeReset::Constants::SceneImportAssetPathKey]);
+						SceneImportAsset = Cast<UInterchangeSceneImportAsset>(SceneImportAssetPath.TryLoad());
+					}
+				}
+			}
+
+			if (SceneImportAsset)
+			{
+				TMap<FString, FString> PreviousSlotMaterialDependencies;
+				PreviousFactoryNode->GetSlotMaterialDependencies(PreviousSlotMaterialDependencies);
+
+				for (TPair<FString, FString>& SlotMaterialDependency : SlotMaterialDependencies)
+				{
+					FSoftObjectPath MaterialPath = GetMaterialPath(SlotMaterialDependency.Value);
+					if (!MaterialPath.IsValid())
+					{
+						continue;
+					}
+
+					FName MaterialSlotName = *SlotMaterialDependency.Key;
+					const int32 MaterialSlotIndex = MeshComponent.GetMaterialIndex(MaterialSlotName);
+					if (MaterialSlotIndex == INDEX_NONE)
+					{
+						// Cannot apply override as slot does not exist, continue 
+						continue;
+					}
+
+					if (!PreviousSlotMaterialDependencies.Contains(SlotMaterialDependency.Key))
+					{
+						// New material slot, but it did not existed in last import, skip override
+						continue;
+					}
+
+					FSoftObjectPath PreviousMaterialPath;
+					const FString PreviousMaterialUid = PreviousSlotMaterialDependencies[SlotMaterialDependency.Key];
+					if(const UInterchangeFactoryBaseNode* PreviousMaterialFactoryNode = SceneImportAsset->GetFactoryNodeFromUniqueID(PreviousMaterialUid))
+					{
+						PreviousMaterialFactoryNode->GetCustomReferenceObject(PreviousMaterialPath);
+						
+						if (!PreviousMaterialPath.IsValid())
+						{
+							continue;
+						}
+					}
+
+					const FSoftObjectPath MeshMaterialPath(MeshComponent.GetMaterial(MaterialSlotIndex));
+					if (MeshMaterialPath.IsValid() && MeshMaterialPath == PreviousMaterialPath)
+					{
+						UMaterialInterface* MaterialInterface = Cast<UMaterialInterface>(MaterialPath.TryLoad());
+						if (!MaterialInterface)
+						{
+							MaterialInterface = UMaterial::GetDefaultMaterial(MD_Surface);
+						}
+
+						MeshComponent.SetMaterial(MaterialSlotIndex, MaterialInterface);
+					}
+				}
+
+				return;
+			}
+			// Else, proceed as if PreviousFactoryNode was null
+		}
+
+		// Apply material slots to mesh component
+
+		// It is an import or a reimport where the factory node prevails, apply material overrides
+		// i.e. !PreviousFactoryNode || ReimportStrategyFlags == EReimportStrategyFlags::ApplyPipelineProperties
+		
+		// Discard existing material overrides
+		if (MeshComponent.HasOverrideMaterials())
+		{
+			MeshComponent.EmptyOverrideMaterials();
+		}
+
+		// Apply factory node's overrides
+		for (TPair<FString, FString>& SlotMaterialDependency : SlotMaterialDependencies)
+		{
+			FSoftObjectPath ReferenceObject = GetMaterialPath(SlotMaterialDependency.Value);
 			if (!ReferenceObject.IsValid())
 			{
 				continue;
diff --git a/Engine/Plugins/Interchange/Runtime/Source/Import/Private/Scene/InterchangeSkeletalMeshActorFactory.cpp b/Engine/Plugins/Interchange/Runtime/Source/Import/Private/Scene/InterchangeSkeletalMeshActorFactory.cpp
--- a/Engine/Plugins/Interchange/Runtime/Source/Import/Private/Scene/InterchangeSkeletalMeshActorFactory.cpp
+++ b/Engine/Plugins/Interchange/Runtime/Source/Import/Private/Scene/InterchangeSkeletalMeshActorFactory.cpp
@@ -110,7 +110,7 @@ void UInterchangeSkeletalMeshActorFactory::SetupObject_GameThread(const FSetupOb
 
 					if (const UInterchangeMeshActorFactoryNode* MeshActorFactoryNode = Cast<UInterchangeMeshActorFactoryNode>(Arguments.FactoryNode))
 					{
-						UE::Interchange::ActorHelper::ApplySlotMaterialDependencies(*Arguments.NodeContainer, *MeshActorFactoryNode, *SkeletalMeshComponent);
+						UE::Interchange::ActorHelper::ApplySlotMaterialDependencies(*Arguments.NodeContainer, *MeshActorFactoryNode, *SkeletalMeshComponent, Cast<UInterchangeMeshActorFactoryNode>(PreviousFactoryNode));
 
 						FString AnimationAssetUidToPlay;
 						if (MeshActorFactoryNode->GetCustomAnimationAssetUidToPlay(AnimationAssetUidToPlay))
diff --git a/Engine/Plugins/Interchange/Runtime/Source/Import/Private/Scene/InterchangeStaticMeshActorFactory.cpp b/Engine/Plugins/Interchange/Runtime/Source/Import/Private/Scene/InterchangeStaticMeshActorFactory.cpp
--- a/Engine/Plugins/Interchange/Runtime/Source/Import/Private/Scene/InterchangeStaticMeshActorFactory.cpp
+++ b/Engine/Plugins/Interchange/Runtime/Source/Import/Private/Scene/InterchangeStaticMeshActorFactory.cpp
@@ -191,7 +191,7 @@ void UInterchangeStaticMeshActorFactory::ApplyStaticMeshToComponent(const UInter
 		// Apply the materials from the mesh node to the static mesh component
 		if (MeshFactoryNode)
 		{
-			UE::Interchange::ActorHelper::ApplySlotMaterialDependencies(NodeContainer, *MeshFactoryNode, *StaticMeshComponent);
+			UE::Interchange::ActorHelper::ApplySlotMaterialDependencies(NodeContainer, *MeshFactoryNode, *StaticMeshComponent, Cast<UInterchangeMeshActorFactoryNode>(PreviousFactoryNode));
 		}
 	}
 }
diff --git a/Engine/Plugins/Interchange/Runtime/Source/Import/Public/Scene/InterchangeActorFactory.h b/Engine/Plugins/Interchange/Runtime/Source/Import/Public/Scene/InterchangeActorFactory.h
--- a/Engine/Plugins/Interchange/Runtime/Source/Import/Public/Scene/InterchangeActorFactory.h
+++ b/Engine/Plugins/Interchange/Runtime/Source/Import/Public/Scene/InterchangeActorFactory.h
@@ -31,7 +31,6 @@ public:
 protected:
 	UE_API virtual void ExecuteResetObjectProperties(const UInterchangeBaseNodeContainer* BaseNodeContainer, UInterchangeFactoryBaseNode* FactoryNode, UObject* ImportedObject) override;
 
-private:
 	UE_API virtual UObject* ImportSceneObject_GameThread(const UInterchangeFactoryBase::FImportSceneObjectsParams& CreateSceneObjectsParams) override;
 
 	// Interchange factory base interface end
diff --git a/Engine/Plugins/Interchange/Runtime/Source/Import/Public/Scene/InterchangeActorHelper.h b/Engine/Plugins/Interchange/Runtime/Source/Import/Public/Scene/InterchangeActorHelper.h
--- a/Engine/Plugins/Interchange/Runtime/Source/Import/Public/Scene/InterchangeActorHelper.h
+++ b/Engine/Plugins/Interchange/Runtime/Source/Import/Public/Scene/InterchangeActorHelper.h
@@ -39,8 +39,11 @@ namespace UE::Interchange::ActorHelper
 
 	/**
 	 * Applies material slot dependencies stored in ActorFactoryNode to MeshComponent.
+	 * In case of a reimport, PreviousActorFactoryNode is used to determine which slots should be updated
+	 * according to the requested import strategy associated with ActorFactoryNode, i.e. do nothing, apply imported slots,
+	 * apply imported slots while preserving changes made after the last import or reimport.
 	 */
-	UE_API void ApplySlotMaterialDependencies(const UInterchangeBaseNodeContainer& NodeContainer, const UInterchangeMeshActorFactoryNode& ActorFactoryNode, UMeshComponent& MeshComponent);
+	UE_API void ApplySlotMaterialDependencies(const UInterchangeBaseNodeContainer& NodeContainer, const UInterchangeMeshActorFactoryNode& ActorFactoryNode, UMeshComponent& MeshComponent, const UInterchangeMeshActorFactoryNode* PreviousActorFactoryNode = nullptr);
 
 	/**
 	 * Applies custom attributes from to ObjectToUpdate based on content of CreateSceneObjectsParams
diff --git a/Engine/Plugins/Interchange/Runtime/Source/Import/Public/Scene/InterchangeSkeletalMeshActorFactory.h b/Engine/Plugins/Interchange/Runtime/Source/Import/Public/Scene/InterchangeSkeletalMeshActorFactory.h
--- a/Engine/Plugins/Interchange/Runtime/Source/Import/Public/Scene/InterchangeSkeletalMeshActorFactory.h
+++ b/Engine/Plugins/Interchange/Runtime/Source/Import/Public/Scene/InterchangeSkeletalMeshActorFactory.h
@@ -26,6 +26,13 @@ public:
 
 	UE_API virtual UClass* GetFactoryClass() const override;
 
+	virtual UObject* ImportSceneObject_GameThread(const UInterchangeFactoryBase::FImportSceneObjectsParams& CreateSceneObjectsParams) override
+	{
+		// Cache the previous factory node on reimport
+		PreviousFactoryNode = CreateSceneObjectsParams.bIsReimport ? CreateSceneObjectsParams.ReimportFactoryNode : nullptr;
+		return Super::ImportSceneObject_GameThread(CreateSceneObjectsParams);
+	}
+
 	UE_API virtual void SetupObject_GameThread(const FSetupObjectParams& Arguments) override;
 
 	// Interchange factory base interface end
@@ -40,6 +47,10 @@ protected:
 	virtual void ExecuteResetObjectProperties(const UInterchangeBaseNodeContainer* BaseNodeContainer, UInterchangeFactoryBaseNode* FactoryNode, UObject* ImportedObject) override;
 	// Interchange actor factory interface end
 	//////////////////////////////////////////////////////////////////////////
+	
+	// Cache the previous factory node passed in the FImportSceneObjectsParams struct
+	// PreviousFactoryNode will be used to respect requested reimport's strategy when updating material slots
+	const UInterchangeFactoryBaseNode* PreviousFactoryNode = nullptr;
 };
 
 
diff --git a/Engine/Plugins/Interchange/Runtime/Source/Import/Public/Scene/InterchangeStaticMeshActorFactory.h b/Engine/Plugins/Interchange/Runtime/Source/Import/Public/Scene/InterchangeStaticMeshActorFactory.h
--- a/Engine/Plugins/Interchange/Runtime/Source/Import/Public/Scene/InterchangeStaticMeshActorFactory.h
+++ b/Engine/Plugins/Interchange/Runtime/Source/Import/Public/Scene/InterchangeStaticMeshActorFactory.h
@@ -30,6 +30,13 @@ public:
 
 	UE_API virtual void SetupObject_GameThread(const FSetupObjectParams& Arguments) override;
 
+	virtual UObject* ImportSceneObject_GameThread(const UInterchangeFactoryBase::FImportSceneObjectsParams& CreateSceneObjectsParams) override
+	{
+		// Cache the previous factory node on reimport
+		PreviousFactoryNode = CreateSceneObjectsParams.bIsReimport ? CreateSceneObjectsParams.ReimportFactoryNode : nullptr;
+		return Super::ImportSceneObject_GameThread(CreateSceneObjectsParams);
+	}
+
 	// Interchange factory base interface end
 	//////////////////////////////////////////////////////////////////////////
 
@@ -48,6 +55,10 @@ private:
 
 	// Helper function to assign a static mesh to its component while checking for other settings like navigation bounds and material dependencies.
 	void ApplyStaticMeshToComponent(const UInterchangeFactoryBaseNode* MeshNode, UStaticMeshComponent* StaticMeshComponent, const UInterchangeBaseNodeContainer& NodeContainer, const UInterchangeMeshActorFactoryNode* MeshFactoryNode);
+
+	// Cache the previous factory node passed in the FImportSceneObjectsParams struct
+	// PreviousFactoryNode will be used to respect requested reimport's strategy when updating material slots
+	const UInterchangeFactoryBaseNode* PreviousFactoryNode = nullptr;
 };
 
 
