[Interchange] Fix compile errors due to missing & misplaced template instantiation

#rnx

#Synced-CL 50977340

diff --git a/Engine/Plugins/Interchange/Runtime/Source/Pipelines/Private/InterchangeLODDataParser.cpp b/Engine/Plugins/Interchange/Runtime/Source/Pipelines/Private/InterchangeLODDataParser.cpp
--- a/Engine/Plugins/Interchange/Runtime/Source/Pipelines/Private/InterchangeLODDataParser.cpp
+++ b/Engine/Plugins/Interchange/Runtime/Source/Pipelines/Private/InterchangeLODDataParser.cpp
@@ -251,12 +251,11 @@ namespace UE::Interchange
 
 			return { EInterchangeMeshCollision::None, FString() };
 		}
+		template TTuple<EInterchangeMeshCollision, FString> CollisionHelper::GetCollisionMeshType<FString>(
+			TObjectPtr<UInterchangePipelineMeshesUtilities>, const UInterchangeBaseNodeContainer&, const FString&, const TArray<FString>&, bool);
 	}
 }
 
-template struct FInterchangeMeshLODDataParser<UInterchangeStaticMeshFactoryNode>;
-template struct FInterchangeMeshLODDataParser<UInterchangeSkeletalMeshFactoryNode>;
-
 template <class TMeshFactoryType>
 FInterchangeMeshLODDataParser<TMeshFactoryType>::FInterchangeMeshLODDataParser(UInterchangeBaseNodeContainer& InBaseNodeContainer,
 	TMeshFactoryType* InMeshFactoryNode,
@@ -843,4 +842,7 @@ void FInterchangeMeshLODDataParser<TMeshFactoryType>::ProcessNodeUids(const TArr
     
 
 	UE::Interchange::MeshesUtilities::ReorderSlotMaterialDependencies(*MeshFactoryNode, BaseNodeContainer);
-}
\ No newline at end of file
+}
+
+template struct FInterchangeMeshLODDataParser<UInterchangeStaticMeshFactoryNode>;
+template struct FInterchangeMeshLODDataParser<UInterchangeSkeletalMeshFactoryNode>;
