[Interchange][FBX] ufbx - materials fixes

- opacity mask fixed for ufbx to match FBX SDK for Phong/Lambert
- opacity(or opacity mask) source is selected between rgb/alpha based on image assigned to opacity/mask, when image alpha. Materials assets in CL 50813129
- Black Emissive color not assigned anymore
- Properly converted 3-component color properly, without scaling alpha

#Synced-CL 50902575

diff --git a/Engine/Plugins/Interchange/Runtime/Source/Parsers/Fbx/InterchangeFbxParser.Build.cs b/Engine/Plugins/Interchange/Runtime/Source/Parsers/Fbx/InterchangeFbxParser.Build.cs
--- a/Engine/Plugins/Interchange/Runtime/Source/Parsers/Fbx/InterchangeFbxParser.Build.cs
+++ b/Engine/Plugins/Interchange/Runtime/Source/Parsers/Fbx/InterchangeFbxParser.Build.cs
@@ -40,6 +40,8 @@ namespace UnrealBuildTool.Rules
 				new string[]
 				{
 					"DeveloperSettings",
+					"ImageCore",
+					"ImageWrapper",
 					"InterchangeCommonParser",
 					"InterchangeMessages",
 					"InterchangeNodes",
diff --git a/Engine/Plugins/Interchange/Runtime/Source/Parsers/Fbx/Private/FbxMaterial.cpp b/Engine/Plugins/Interchange/Runtime/Source/Parsers/Fbx/Private/FbxMaterial.cpp
--- a/Engine/Plugins/Interchange/Runtime/Source/Parsers/Fbx/Private/FbxMaterial.cpp
+++ b/Engine/Plugins/Interchange/Runtime/Source/Parsers/Fbx/Private/FbxMaterial.cpp
@@ -5,6 +5,7 @@
 #include "FbxAPI.h"
 #include "FbxHelper.h"
 #include "FbxInclude.h"
+#include "IImageWrapperModule.h"
 #include "Fbx/InterchangeFbxMessages.h"
 #include "InterchangeMaterialDefinitions.h"
 #include "InterchangeMaterialInstanceNode.h"
@@ -13,6 +14,7 @@
 #include "InterchangeTexture2DNode.h"
 #include "InterchangeTextureNode.h"
 #include "Misc/Paths.h"
+#include "Modules/ModuleManager.h"
 #include "Nodes/InterchangeBaseNodeContainer.h"
 
 #define LOCTEXT_NAMESPACE "InterchangeFbxMaterial"
@@ -114,63 +116,96 @@ namespace UE
 				return TextureNode;
 			}
 
-			bool FFbxMaterial::ConvertPropertyToMaterialInstanceNode(UInterchangeBaseNodeContainer& NodeContainer, UInterchangeMaterialInstanceNode* MaterialInstanceNode, FbxProperty& Property, float Factor, FName InputName, const TVariant<FLinearColor, float>& DefaultValue, bool bInverse)
+			void FFbxMaterial::ConvertValuePropertyToMaterialInstanceNode(UInterchangeMaterialInstanceNode* MaterialInstanceNode, FbxProperty& Property, float Factor, FName InputName, const TVariant<FLinearColor, float>& DefaultValue, bool bInverse)
 			{
 				using namespace Materials::Standard::Nodes;
 
-				const int32 TextureCount = Property.GetSrcObjectCount<FbxFileTexture>();
-				const EFbxType DataType = Property.GetPropertyDataType().GetType();
-				FString InputToConnectTo = InputName.ToString();
+				const FString InputAttributeKey = InputName.ToString();
 
-				if (TextureCount == 0)
+				const EFbxType DataType = Property.GetPropertyDataType().GetType();
+				if (DataType == eFbxDouble || DataType == eFbxFloat || DataType == eFbxInt)
 				{
-					const FString InputAttributeKey = InputName.ToString();
+					const float PropertyValue = Property.Get<float>() * Factor;
+					MaterialInstanceNode->AddScalarParameterValue(InputAttributeKey, bInverse ? 1.f - PropertyValue : PropertyValue);
+				}
+				else if (DataType == eFbxDouble3 || DataType == eFbxDouble4)
+				{
+					FbxDouble3 Color = DataType == eFbxDouble3 ? Property.Get<FbxDouble3>() : Property.Get<FbxDouble4>();
+					FVector3f FbxValue = FVector3f(Color[0], Color[1], Color[2]) * Factor;
+					FLinearColor PropertyValue = bInverse ? FVector3f::OneVector - FbxValue : FbxValue;
 
-					if (DataType == eFbxDouble || DataType == eFbxFloat || DataType == eFbxInt)
+					if (DefaultValue.IsType<FLinearColor>())
 					{
-						const float PropertyValue = Property.Get<float>() * Factor;
-						MaterialInstanceNode->AddScalarParameterValue(InputAttributeKey, bInverse ? 1.f - PropertyValue : PropertyValue);
+						MaterialInstanceNode->AddVectorParameterValue(InputAttributeKey, PropertyValue);
 					}
-					else if (DataType == eFbxDouble3 || DataType == eFbxDouble4)
+					else if (DefaultValue.IsType<float>())
 					{
-						FbxDouble3 Color = DataType == eFbxDouble3 ? Property.Get<FbxDouble3>() : Property.Get<FbxDouble4>();
-						FVector3f FbxValue = FVector3f(Color[0], Color[1], Color[2]) * Factor;
-						FLinearColor PropertyValue = bInverse ? FVector3f::OneVector - FbxValue : FbxValue;
-
-						if (DefaultValue.IsType<FLinearColor>())
-						{
-							MaterialInstanceNode->AddVectorParameterValue(InputAttributeKey, PropertyValue);
-						}
-						else if (DefaultValue.IsType<float>())
-						{
-							// We're connecting a linear color to a float input. Ideally, we'd go through a desaturate, but for now we'll just take the red channel and ignore the rest.
-							MaterialInstanceNode->AddScalarParameterValue(InputAttributeKey, PropertyValue.R);
-						}
+						// We're connecting a linear color to a float input. Ideally, we'd go through a desaturate, but for now we'll just take the red channel and ignore the rest.
+						MaterialInstanceNode->AddScalarParameterValue(InputAttributeKey, PropertyValue.R);
 					}
-
-					return true;
 				}
+			}
 
-				// Handles max one texture per property.
-				FbxFileTexture* FbxTexture = Property.GetSrcObject<FbxFileTexture>(0);
+			const UInterchangeTexture2DNode* FFbxMaterial::ConvertTexturePropertyToMaterialInstanceNode(UInterchangeBaseNodeContainer& NodeContainer, UInterchangeMaterialInstanceNode* MaterialInstanceNode, FName InputName, FbxFileTexture* FbxTexture)
+			{
 				if (const UInterchangeTexture2DNode* TextureNode = CreateTexture2DNode(FbxTexture, NodeContainer, MaterialInstanceNode))
 				{
 					MaterialInstanceNode->AddTextureParameterValue(InputName.ToString() + TEXT("Map"), TextureNode->GetUniqueID());
 					MaterialInstanceNode->AddScalarParameterValue(InputName.ToString() + TEXT("MapWeight"), 1.f);
+					return TextureNode;
 				}
-				else
+
+				if (!GIsAutomationTesting)
 				{
-					if (!GIsAutomationTesting)
-					{
-						UInterchangeResultTextureDisplay_TextureFileDoNotExist* Message = Parser.AddMessage<UInterchangeResultTextureDisplay_TextureFileDoNotExist>();
-						Message->TextureName = FbxTexture ? UTF8_TO_TCHAR(FbxTexture->GetFileName()) : TEXT("Undefined");
-						Message->MaterialName = MaterialInstanceNode->GetDisplayLabel();
-					}
+					UInterchangeResultTextureDisplay_TextureFileDoNotExist* Message = Parser.AddMessage<UInterchangeResultTextureDisplay_TextureFileDoNotExist>();
+					Message->TextureName = FbxTexture ? UTF8_TO_TCHAR(FbxTexture->GetFileName()) : TEXT("Undefined");
+					Message->MaterialName = MaterialInstanceNode->GetDisplayLabel();
+				}
+				return nullptr;
+			}
 
+			bool FFbxMaterial::ConvertPropertyToMaterialInstanceNode(UInterchangeBaseNodeContainer& NodeContainer, UInterchangeMaterialInstanceNode* MaterialInstanceNode, FbxProperty& Property, float Factor, FName InputName, const TVariant<FLinearColor, float>& DefaultValue, bool bInverse)
+			{
+				if (Property.GetSrcObjectCount<FbxFileTexture>() == 0)
+				{
+					ConvertValuePropertyToMaterialInstanceNode(MaterialInstanceNode, Property, Factor, InputName, DefaultValue, bInverse);
+					return true;
+				}
+
+				// Handles max one texture per property.
+				FbxFileTexture* FbxTexture = Property.GetSrcObject<FbxFileTexture>(0);
+				return ConvertTexturePropertyToMaterialInstanceNode(NodeContainer, MaterialInstanceNode, InputName, FbxTexture) != nullptr;
+			}
+
+			bool DoesTextureHaveAlpha(const FString& TextureFilename)
+			{
+				IImageWrapperModule& ImageWrapperModule = FModuleManager::LoadModuleChecked<IImageWrapperModule>(FName("ImageWrapper"));
+
+				TArray64<uint8> ImageData;
+
+				if (!FFileHelper::LoadFileToArray(ImageData, *TextureFilename) || ImageData.IsEmpty())
+				{
 					return false;
 				}
 
-				return true;
+				switch (ImageWrapperModule.DetectImageFormat(ImageData.GetData(), ImageData.Num()))
+				{
+					case EImageFormat::Invalid:
+						// Early exit for Jpeg, which doesn't support alpha
+					case EImageFormat::JPEG:
+					case EImageFormat::GrayscaleJPEG:
+					case EImageFormat::UEJPEG:
+					case EImageFormat::GrayscaleUEJPEG:
+						return false;
+				}
+
+				FImage LoadedImage;
+				if (ImageWrapperModule.DecompressImage(ImageData.GetData(), ImageData.Num(), LoadedImage))
+				{
+					return FImageCore::DetectAlphaChannel(LoadedImage);
+				}
+
+				return false;
 			}
 
 			const UInterchangeMaterialInstanceNode* FFbxMaterial::AddMaterialInstanceNode(FbxSurfaceMaterial* SurfaceMaterial, UInterchangeBaseNodeContainer& NodeContainer)
@@ -308,8 +343,26 @@ namespace UE
 							}
 
 							MaterialInstanceNode->SetCustomBlendMode(BlendMode);
-							
-							bHasInput |= ConvertPropertyToMaterialInstanceNode(NodeContainer, MaterialInstanceNode, MaterialProperty, Factor, InputName, DefaultValue, true);
+
+							if (MaterialProperty.GetSrcObjectCount<FbxFileTexture>() == 0)
+							{
+								ConvertValuePropertyToMaterialInstanceNode(MaterialInstanceNode, MaterialProperty, Factor, InputName, DefaultValue, true);
+								bHasInput = true;
+							}
+							else
+							{
+								FbxFileTexture* FbxTexture = MaterialProperty.GetSrcObject<FbxFileTexture>(0);
+								if (const UInterchangeTexture2DNode* Texture2DNode = ConvertTexturePropertyToMaterialInstanceNode(NodeContainer, MaterialInstanceNode, InputName, FbxTexture))
+								{
+									bHasInput = true;
+
+									// When alpha is present in the image use it as opacity, instead of RGB
+									if (DoesTextureHaveAlpha(*Texture2DNode->GetPayLoadKey()))
+									{
+										MaterialInstanceNode->AddScalarParameterValue(TEXT("OpacitySourceChannel"), 1);
+									}
+								}
+							}
 						}
 						else
 						{
diff --git a/Engine/Plugins/Interchange/Runtime/Source/Parsers/Fbx/Private/FbxMaterial.h b/Engine/Plugins/Interchange/Runtime/Source/Parsers/Fbx/Private/FbxMaterial.h
--- a/Engine/Plugins/Interchange/Runtime/Source/Parsers/Fbx/Private/FbxMaterial.h
+++ b/Engine/Plugins/Interchange/Runtime/Source/Parsers/Fbx/Private/FbxMaterial.h
@@ -72,7 +72,13 @@ namespace UE
 				const UInterchangeTexture2DNode* CreateTexture2DNode(UInterchangeBaseNodeContainer& NodeContainer, const FString& TextureFilePath);
 				const UInterchangeTexture2DNode* CreateTexture2DNode(FbxFileTexture* FbxTexture, UInterchangeBaseNodeContainer& NodeContainer, UInterchangeMaterialInstanceNode* MaterialInstanceNode);
 				bool ConvertPropertyToMaterialInstanceNode(UInterchangeBaseNodeContainer& NodeContainer, UInterchangeMaterialInstanceNode* MaterialInstanceNode,
-														   FbxProperty& Property, float Factor, FName InputName, const TVariant<FLinearColor, float>& DefaultValue, bool bInverse = false);
+				                                           FbxProperty& Property, float Factor, FName InputName, const TVariant<FLinearColor, float>& DefaultValue, bool bInverse = false);
+				void ConvertValuePropertyToMaterialInstanceNode(UInterchangeMaterialInstanceNode* MaterialInstanceNode,
+					FbxProperty& Property, float Factor, FName InputName,
+					const TVariant<FLinearColor, float>& DefaultValue, bool bInverse);
+				const UInterchangeTexture2DNode* ConvertTexturePropertyToMaterialInstanceNode(UInterchangeBaseNodeContainer& NodeContainer,
+				                                                  UInterchangeMaterialInstanceNode* MaterialInstanceNode,
+				                                                  FName InputName, FbxFileTexture* FbxTexture);
 
 			private:
 				const UInterchangeMaterialInstanceNode* AddMaterialInstanceNode(FbxSurfaceMaterial* SurfaceMaterial, UInterchangeBaseNodeContainer& NodeContainer);
diff --git a/Engine/Plugins/Interchange/Runtime/Source/Parsers/Fbx/Private/Ufbx/UfbxMaterial.cpp b/Engine/Plugins/Interchange/Runtime/Source/Parsers/Fbx/Private/Ufbx/UfbxMaterial.cpp
--- a/Engine/Plugins/Interchange/Runtime/Source/Parsers/Fbx/Private/Ufbx/UfbxMaterial.cpp
+++ b/Engine/Plugins/Interchange/Runtime/Source/Parsers/Fbx/Private/Ufbx/UfbxMaterial.cpp
@@ -201,7 +201,10 @@ namespace UE::Interchange::Private
 				if (Map.value_components >=3)
 				{
 					const FString InputAttributeKey = InputName.ToString();
-					FLinearColor PropertyValue(Convert::ConvertVec4(Map.value_vec4) * Factor);
+
+					FLinearColor PropertyValue = (Map.value_components == 4) 
+					? FLinearColor(Convert::ConvertVec4(Map.value_vec4) * Factor)
+					: FLinearColor(Convert::ConvertVec3(Map.value_vec3) * Factor);
 					if (Filter(PropertyValue))
 					{
 						return MaterialInstanceNode->AddVectorParameterValue(InputAttributeKey, PropertyValue);
@@ -218,19 +221,19 @@ namespace UE::Interchange::Private
 				return MaterialInstanceNode->AddScalarParameterValue(InputAttributeKey, PropertyValue);
 			};
 
-			auto SetTextureFromMap = [&MaterialInstanceNode, &MaterialMapConverter, &bHasUVTransform](FName InputName, const ufbx_material_map& Map, float Factor) -> bool
+			auto SetTextureFromMap = [&MaterialInstanceNode, &MaterialMapConverter, &bHasUVTransform](FName InputName, const ufbx_material_map& Map, float Factor) -> const UInterchangeTexture2DNode*
 			{
 				const ufbx_texture* Texture = Map.texture_enabled ? Map.texture : nullptr;
 				if (!Texture)
 				{
-					return false;	
+					return nullptr;	
 				}
 
 				const UInterchangeTexture2DNode* TextureNode = MaterialMapConverter.CreateTexture2DNode(*Texture);
 
 				if (!TextureNode)
 				{
-					return false;
+					return nullptr;
 				}
 
 				MaterialInstanceNode->AddTextureParameterValue(InputName.ToString() + TEXT("Map"), TextureNode->GetUniqueID());
@@ -257,7 +260,7 @@ namespace UE::Interchange::Private
 				//}
 				// End of future handling of UV transform
 
-				return true;
+				return TextureNode;
 			};
 
 			bool bHasInput = false;
@@ -289,7 +292,9 @@ namespace UE::Interchange::Private
 				const float Factor = GET_VALLUE_FROM_FACTOR_MAP(Material->fbx.emission_factor);
 				const ufbx_material_map& Map = Material->fbx.emission_color;
 
-				bHasInput |= SetTextureFromMap(InputName, Map, Factor) ? true : SetColorFromMap(InputName, Map, Factor);
+				bHasInput |= SetTextureFromMap(InputName, Map, Factor) ? true 
+				: Convert::ConvertVec3(Material->fbx.emission_color.value_vec3).IsNearlyZero() ? false 
+				: SetColorFromMap(InputName, Map, Factor);
 			}
 
 			bool bIsPhong = false;
@@ -364,7 +369,7 @@ namespace UE::Interchange::Private
 
 			// Connect only if transparency is either a texture or different from 0.f
 			// Check fbx.transparency first...
-			if (!bOpaque && Material->fbx.transparency_color.has_value)
+			if (Material->fbx.transparency_color.has_value)
 			{
 				FName InputName = Phong::Parameters::Opacity;
 				const ufbx_material_map& ColorMap = Material->fbx.transparency_color;
@@ -375,29 +380,39 @@ namespace UE::Interchange::Private
 					FString HasOpacity{ TEXT("bHasOpacity") };
 					EBlendMode BlendMode = EBlendMode::BLEND_Translucent;
 
-					// The texture must be hooked to the OpacityMask when transparency_color.value_int is equal to 0
-					if (ColorMap.value_int == 0)
+					// When material is opaque but has texture this means it has opacity map
+					if (bOpaque && ColorMap.value_int == 0)
 					{
 						InputName = Phong::Parameters::OpacityMask;
 						BlendMode = EBlendMode::BLEND_Masked;
-						MaterialInstanceNode->AddStaticSwitchParameterValue(TEXT("bHasOpacityMaskMap"), true);
 						HasOpacity += TEXT("Mask");
 					}
 
 					HasOpacity += TEXT("Map");
-					MaterialInstanceNode->AddStaticSwitchParameterValue(HasOpacity, true);
 					MaterialInstanceNode->SetCustomBlendMode(BlendMode);
-					
-					bHasInput |= SetTextureFromMap(InputName, ColorMap, 1.0f);
+
+					if (const UInterchangeTexture2DNode* Texture2DNode = SetTextureFromMap(InputName, ColorMap, 1.0f))
+					{
+						bHasInput = true;
+
+						// When alpha is present in the image use it as opacity, instead of RGB
+						if (Parser.DoesTextureHaveAlpha(*Texture2DNode->GetPayLoadKey(), *ColorMap.texture))
+						{
+							MaterialInstanceNode->AddScalarParameterValue(TEXT("OpacitySourceChannel"), 1);
+						}
+					}
 				}
 				else
 				{
-					const float TransparencyScalar = Factor * ColorMap.value_real;
-					const float OpacityScalar = 1.f - FMath::Clamp(TransparencyScalar, 0.f, 1.f);
-					if (OpacityScalar < 1.f)
+					if (!bOpaque)
 					{
-						MaterialInstanceNode->SetCustomBlendMode(EBlendMode::BLEND_Translucent);
-						bHasInput |= MaterialInstanceNode->AddScalarParameterValue(InputName.ToString(), OpacityScalar);
+						const float TransparencyScalar = Factor * ColorMap.value_real;
+						const float OpacityScalar = 1.f - FMath::Clamp(TransparencyScalar, 0.f, 1.f);
+						if (OpacityScalar < 1.f)
+						{
+							MaterialInstanceNode->SetCustomBlendMode(EBlendMode::BLEND_Translucent);
+							bHasInput |= MaterialInstanceNode->AddScalarParameterValue(InputName.ToString(), OpacityScalar);
+						}
 					}
 				}
 			}
diff --git a/Engine/Plugins/Interchange/Runtime/Source/Parsers/Fbx/Private/Ufbx/UfbxParser.cpp b/Engine/Plugins/Interchange/Runtime/Source/Parsers/Fbx/Private/Ufbx/UfbxParser.cpp
--- a/Engine/Plugins/Interchange/Runtime/Source/Parsers/Fbx/Private/Ufbx/UfbxParser.cpp
+++ b/Engine/Plugins/Interchange/Runtime/Source/Parsers/Fbx/Private/Ufbx/UfbxParser.cpp
@@ -14,6 +14,7 @@
 #include "Mesh/InterchangeMeshPayload.h"
 #endif
 
+#include "IImageWrapperModule.h"
 #include "InterchangeCameraNode.h"
 #include "Fbx/InterchangeFbxMessages.h"
 #include "InterchangeHelper.h"
@@ -26,6 +27,7 @@
 #include "Serialization/MemoryWriter.h"
 
 #include "ufbx.c"
+#include "Modules/ModuleManager.h"
 #include "Nodes/InterchangeUserDefinedAttribute.h"
 
 
@@ -633,6 +635,47 @@ bool FUfbxParser::FetchTexturePayload(const FString& PayloadKey, TOptional<TArra
 
 #endif
 
+bool FUfbxParser::DoesTextureHaveAlpha(const FString& TextureFilename, const ufbx_texture& Texture)
+{
+	IImageWrapperModule& ImageWrapperModule = FModuleManager::LoadModuleChecked<IImageWrapperModule>(FName("ImageWrapper"));
+
+	TArray64<uint8> ImageData;
+	TConstArrayView64<uint8> ImageDataView;
+
+	if (Texture.content.data && (Texture.content.size > 0))
+	{
+		// embedded texture
+		ImageDataView = TConstArrayView64<uint8>(static_cast<const uint8*>(Texture.content.data), Texture.content.size);
+	}
+	else if (FFileHelper::LoadFileToArray(ImageData, *TextureFilename))
+	{
+		ImageDataView = ImageData;
+	}
+	else
+	{
+		return false;
+	}
+
+	switch (ImageWrapperModule.DetectImageFormat(ImageDataView.GetData(), ImageDataView.Num()))
+	{
+		case EImageFormat::Invalid:
+		// Early exit for Jpeg, which doesn't support alpha
+		case EImageFormat::JPEG:
+		case EImageFormat::GrayscaleJPEG:
+		case EImageFormat::UEJPEG:
+		case EImageFormat::GrayscaleUEJPEG:
+			return false;
+	};
+
+	FImage LoadedImage;
+	if (ImageWrapperModule.DecompressImage(ImageDataView.GetData(), ImageDataView.Num(), LoadedImage))
+	{
+		return FImageCore::DetectAlphaChannel(LoadedImage);
+	}
+
+	return false;
+}
+
 bool FUfbxParser::FetchAnimationBakeTransformPayload(const TArray<UE::Interchange::FAnimationPayloadQuery>& PayloadQueries,
                                                      const FString& ResultFolder, FCriticalSection* ResultPayloadsCriticalSection, TAtomic<int64>& UniqueIdCounter,
                                                      TMap<FString, FString>& ResultPayloads)
diff --git a/Engine/Plugins/Interchange/Runtime/Source/Parsers/Fbx/Private/Ufbx/UfbxParser.h b/Engine/Plugins/Interchange/Runtime/Source/Parsers/Fbx/Private/Ufbx/UfbxParser.h
--- a/Engine/Plugins/Interchange/Runtime/Source/Parsers/Fbx/Private/Ufbx/UfbxParser.h
+++ b/Engine/Plugins/Interchange/Runtime/Source/Parsers/Fbx/Private/Ufbx/UfbxParser.h
@@ -192,6 +192,7 @@ namespace UE::Interchange::Private
 		FString GetCameraName(const ufbx_camera& Camera) const;
 		FString GetLightName(const ufbx_light& Light) const;
 
+		static bool DoesTextureHaveAlpha(const FString& TextureFilename, const ufbx_texture& Texture);
 
 		static void ConvertProperty(const ufbx_prop& Prop, UInterchangeBaseNode* SceneNode, const TOptional<FString>& PayloadKey = TOptional<FString>());
 
