[Interchange] Removed templatization for CollisionHelper's GetCollisionMeshType function.

#rnx

#Synced-CL 50934419

diff --git a/Engine/Plugins/Interchange/Runtime/Source/Pipelines/Private/InterchangeLODDataParser.cpp b/Engine/Plugins/Interchange/Runtime/Source/Pipelines/Private/InterchangeLODDataParser.cpp
--- a/Engine/Plugins/Interchange/Runtime/Source/Pipelines/Private/InterchangeLODDataParser.cpp
+++ b/Engine/Plugins/Interchange/Runtime/Source/Pipelines/Private/InterchangeLODDataParser.cpp
@@ -59,34 +59,16 @@ namespace UE::Interchange
 
 	namespace CollisionHelper
 	{
-		template <class T>
 		TTuple<EInterchangeMeshCollision, FString> GetCollisionMeshType(
 			TObjectPtr<UInterchangePipelineMeshesUtilities> PipelineMeshesUtilities,
 			const UInterchangeBaseNodeContainer& NodeContainer,
 			const FString& NodeUid,
-			const TArray<T>& AllNodeUids,
+			const TArray<FString>& AllNodeUids,
 			bool bImportCollisionAccordingToMeshName
 		)
 		{
 			EInterchangeMeshCollision CollisionType = EInterchangeMeshCollision::None;
 
-			auto InterpretResult = [&CollisionType](const T* CorrespondingMeshUid) -> TTuple<EInterchangeMeshCollision, FString>
-				{
-					if constexpr (std::is_same_v<T, FInterchangeLODMeshData>)
-					{
-						return { CollisionType, CorrespondingMeshUid->NodeUid };
-					}
-					else if constexpr (std::is_same_v<T, FString>)
-					{
-						return { CollisionType, *CorrespondingMeshUid };
-					}
-					else
-					{
-						ensure(false);
-						return { EInterchangeMeshCollision::None, FString() };
-					}
-				};
-
 			// Try to find the mesh node we're actually talking about
 			const UInterchangeBaseNode* BaseNode = NodeContainer.GetNode(NodeUid);
 			const UInterchangeMeshNode* MeshNode = Cast<const UInterchangeMeshNode>(BaseNode);
@@ -133,43 +115,25 @@ namespace UE::Interchange
 				int32 LastUnderscore = INDEX_NONE;
 				verify(MeshName.FindLastChar(TEXT('_'), LastUnderscore));
 
-				auto MatchPredicateString = [&PipelineMeshesUtilities, &NodeContainer](FString Body)
+				auto MatchPredicate = [&PipelineMeshesUtilities, &NodeContainer](FString Body)
 					{
 						// Generate a predicate to be used by the below Finds.
 						return [Body, &PipelineMeshesUtilities, &NodeContainer](const FString& ToCompare) { return Body == UE::Interchange::Private::GetNodeName(PipelineMeshesUtilities, NodeContainer, ToCompare); };
 					};
 
-				auto MatchPredicateLODLODMeshData = [&PipelineMeshesUtilities, &NodeContainer](FString Body)
+				auto FindPyPredicate = [&AllNodeUids, &MatchPredicate](const FString& RenderMeshName) -> const FString*
 					{
-						// Generate a predicate to be used by the below Finds.
-						return [Body, &PipelineMeshesUtilities, &NodeContainer](const FInterchangeLODMeshData& ToCompare) { return Body == UE::Interchange::Private::GetNodeName(PipelineMeshesUtilities, NodeContainer, ToCompare.MeshUid); };
-					};
-
-				auto FindPyPredicate = [&AllNodeUids, &MatchPredicateString, &MatchPredicateLODLODMeshData](const FString& RenderMeshName) -> const T*
-					{
-						if constexpr (std::is_same_v<T, FInterchangeLODMeshData>)
-						{
-							return AllNodeUids.FindByPredicate(MatchPredicateLODLODMeshData(RenderMeshName));
-						}
-						else if constexpr (std::is_same_v<T, FString>)
-						{
-							return AllNodeUids.FindByPredicate(MatchPredicateString(RenderMeshName));
-						}
-						else
-						{
-							ensure(false);
-							return nullptr;
-						}
+						return AllNodeUids.FindByPredicate(MatchPredicate(RenderMeshName));
 					};
 
 				// If we find a mesh named the same as the collision mesh (following the collision prefix), we have a match
 				// e.g. this will match 'UBX_House' with a mesh called 'House'
 				{
 					const FString RenderMeshName = MeshName.RightChop(FirstUnderscore + 1);
-					const T* CorrespondingMeshUid = FindPyPredicate(RenderMeshName);
+					const FString* CorrespondingMeshUid = FindPyPredicate(RenderMeshName);
 					if (CorrespondingMeshUid)
 					{
-						return InterpretResult(CorrespondingMeshUid);
+						return { CollisionType, *CorrespondingMeshUid };
 					}
 				}
 
@@ -178,10 +142,10 @@ namespace UE::Interchange
 				if (FirstUnderscore != LastUnderscore)
 				{
 					const FString RenderMeshName = MeshName.Mid(FirstUnderscore + 1, LastUnderscore - FirstUnderscore - 1);
-					const T* CorrespondingMeshUid = FindPyPredicate(RenderMeshName);
+					const FString* CorrespondingMeshUid = FindPyPredicate(RenderMeshName);
 					if (CorrespondingMeshUid)
 					{
-						return InterpretResult(CorrespondingMeshUid);
+						return { CollisionType, *CorrespondingMeshUid };
 					}
 				}
 
@@ -193,23 +157,9 @@ namespace UE::Interchange
 					// Like above start by just stripping the prefix (e.g. 'UBX_House_01' --> 'House_01')
 					FString RenderMeshName = MeshName.RightChop(FirstUnderscore + 1);
 
-					auto NameCollisionSearchPredicate = [&RenderMeshName, &PipelineMeshesUtilities, &NodeContainer](const T& Other)
+					auto NameCollisionSearchPredicate = [&RenderMeshName, &PipelineMeshesUtilities, &NodeContainer](const FString& Other)
 						{
-							FString NodeUid;
-							if constexpr (std::is_same_v<T, FInterchangeLODMeshData>)
-							{
-								NodeUid = Other.NodeUid;
-							}
-							else if constexpr (std::is_same_v<T, FString>)
-							{
-								NodeUid = Other;
-							}
-							else
-							{
-								ensure(false);
-								return nullptr;
-							}
-
+							FString NodeUid = Other;
 							FString MeshName = UE::Interchange::Private::GetNodeName(PipelineMeshesUtilities, NodeContainer, NodeUid);
 							if (!MeshName.StartsWith(RenderMeshName))
 							{
@@ -232,10 +182,10 @@ namespace UE::Interchange
 							return MeshNameSuffix.Len() == 0;
 						};
 
-					const T* CorrespondingMeshUid = AllNodeUids.FindByPredicate(NameCollisionSearchPredicate);
+					const FString* CorrespondingMeshUid = AllNodeUids.FindByPredicate(NameCollisionSearchPredicate);
 					if (CorrespondingMeshUid)
 					{
-						return InterpretResult(CorrespondingMeshUid);
+						return { CollisionType, *CorrespondingMeshUid };
 					}
 
 					// If we still have nothing, try also stripping numbered suffixes from the render mesh name and repeating the search
@@ -244,15 +194,13 @@ namespace UE::Interchange
 					CorrespondingMeshUid = AllNodeUids.FindByPredicate(NameCollisionSearchPredicate);
 					if (CorrespondingMeshUid)
 					{
-						return InterpretResult(CorrespondingMeshUid);
+						return { CollisionType, *CorrespondingMeshUid };
 					}
 				}
 			}
 
 			return { EInterchangeMeshCollision::None, FString() };
 		}
-		template TTuple<EInterchangeMeshCollision, FString> CollisionHelper::GetCollisionMeshType<FString>(
-			TObjectPtr<UInterchangePipelineMeshesUtilities>, const UInterchangeBaseNodeContainer&, const FString&, const TArray<FString>&, bool);
 	}
 }
 
diff --git a/Engine/Plugins/Interchange/Runtime/Source/Pipelines/Public/InterchangeLODDataParser.h b/Engine/Plugins/Interchange/Runtime/Source/Pipelines/Public/InterchangeLODDataParser.h
--- a/Engine/Plugins/Interchange/Runtime/Source/Pipelines/Public/InterchangeLODDataParser.h
+++ b/Engine/Plugins/Interchange/Runtime/Source/Pipelines/Public/InterchangeLODDataParser.h
@@ -16,12 +16,11 @@ class UInterchangeSkeletalMeshLodDataNode;
 
 namespace UE::Interchange::CollisionHelper
 {
-	template <class T>
 	TTuple<EInterchangeMeshCollision, FString> GetCollisionMeshType(
 		TObjectPtr<UInterchangePipelineMeshesUtilities> PipelineMeshesUtilities,
 		const UInterchangeBaseNodeContainer& NodeContainer,
 		const FString& NodeUid,
-		const TArray<T>& AllNodeUids,
+		const TArray<FString>& AllNodeUids,
 		bool bImportCollisionAccordingToMeshName
 	);
 }
