Fixed ensure when importing specific obj file

#rnx

#Synced-CL 50905747

diff --git a/Engine/Plugins/Interchange/Runtime/Source/Import/Private/Mesh/InterchangeOBJTranslator.cpp b/Engine/Plugins/Interchange/Runtime/Source/Import/Private/Mesh/InterchangeOBJTranslator.cpp
--- a/Engine/Plugins/Interchange/Runtime/Source/Import/Private/Mesh/InterchangeOBJTranslator.cpp
+++ b/Engine/Plugins/Interchange/Runtime/Source/Import/Private/Mesh/InterchangeOBJTranslator.cpp
@@ -25,6 +25,8 @@
 
 #include UE_INLINE_GENERATED_CPP_BY_NAME(InterchangeOBJTranslator)
 
+DEFINE_LOG_CATEGORY_STATIC(LogInterchangeObjImport, Log, All);
+
 #define LOCTEXT_NAMESPACE "InterchangeOBJTranslator"
 
 static bool GInterchangeEnableOBJImport = true;
@@ -805,7 +807,7 @@ namespace ObjParserUtils
 		else if (!UnknownKeywords.Contains(Keyword))
 		{
 			// Unknown keyword, report it once
-			UE_LOG(LogTemp, Display, TEXT("Interchange Obj translator: Unknown keyword: %s"), *Keyword);
+			UE_LOG(LogInterchangeObjImport, Display, TEXT("Interchange Obj translator: Unknown keyword: %s"), *Keyword);
 			UnknownKeywords.Add(Keyword);
 		}
 
@@ -941,7 +943,18 @@ namespace ObjParserUtils
 				}
 			}
 
-			Parameter.Path = FPaths::GetPath(ObjData.ObjFilename) / FString(FileName);
+			if (!FileName.IsEmpty())
+			{
+				Parameter.Path = FPaths::GetPath(ObjData.ObjFilename) / FString(FileName);
+
+				FString NormalizedTexturePath(Parameter.Path);
+				FPaths::NormalizeFilename(NormalizedTexturePath);
+				if (!FPaths::FileExists(NormalizedTexturePath))
+				{
+					Parameter.Path = FString();
+					UE_LOG(LogInterchangeObjImport, Warning, TEXT("Interchange Obj translator: Cannot find texture %s"), *NormalizedTexturePath);
+				}
+			}
 		}
 		else
 		{
@@ -1012,7 +1025,7 @@ namespace ObjParserUtils
 		else
 		{
 			//We support only .mtl material file
-			UE_LOG(LogTemp, Warning, TEXT("Interchange Obj translator: Unsupported material file: %s"), *FileToParse);
+			UE_LOG(LogInterchangeObjImport, Warning, TEXT("Interchange Obj translator: Unsupported material file: %s"), *FileToParse);
 		}
 
 		// After parsing, clear the current material name
@@ -1114,6 +1127,7 @@ namespace ObjTranslatorUtils
 			FString TextureName = MakeTextureNodeName(FPaths::GetBaseFilename(TextureParameter.Path));
 			TextureNode = CreateTexture2DNode(BaseNodeContainer, TextureParameter.Path, TextureName);
 		}
+
 		return TextureNode;
 	}
 
@@ -1123,7 +1137,11 @@ namespace ObjTranslatorUtils
 		if (const UInterchangeTexture2DNode* TextureNode = CreateTexture2DNode(BaseNodeContainer, TexturePath))
 		{
 			const FString InputNameLabel = InputName + TEXT("Map");
-			return MaterialInstanceNode->AddTextureParameterValue(InputNameLabel, TextureNode->GetUniqueID());
+			if (MaterialInstanceNode->AddTextureParameterValue(InputNameLabel, TextureNode->GetUniqueID()))
+			{
+				MaterialInstanceNode->AddScalarParameterValue(InputNameLabel + TEXT("Weight"), 1.0f);
+				return true;
+			}
 		}
 		else if (IsColorInitialized(Color))
 		{
@@ -1139,11 +1157,9 @@ namespace ObjTranslatorUtils
 		bool bResult = false;
 		if (const UInterchangeTexture2DNode* TextureNode = CreateTexture2DNode(BaseNodeContainer, TexturePath))
 		{
-			const FString InputNameLabel = InputName + TEXT("MapWeight");
+			const FString InputNameLabel = InputName + TEXT("Map");
 			bResult |= MaterialInstanceNode->AddTextureParameterValue(InputNameLabel, TextureNode->GetUniqueID());
-
-			const float WeightValue = IsScalarInitialized(Weight) ? Weight : 1.0f;
-            bResult |= MaterialInstanceNode->AddScalarParameterValue(InputName, Weight);
+            bResult |= MaterialInstanceNode->AddScalarParameterValue(InputNameLabel + TEXT("Weight"), 1.0f);
 		}
 		else if (IsScalarInitialized(Weight))
 		{
@@ -1432,7 +1448,7 @@ bool UInterchangeOBJTranslator::Translate(UInterchangeBaseNodeContainer& BaseNod
 
 							float SpecularExponentScale = (IsScalarInitialized(MaterialData.SpecularExponent) ? MaterialData.SpecularExponent : 1.0f) * SpecularExponentToShininessFactor;
 							MaterialInstanceNode->AddTextureParameterValue(ShininessParameterName + TEXT("Map"), SpecularExponentTextureNode->GetUniqueID());
-							MaterialInstanceNode->AddScalarParameterValue(ShininessParameterName, SpecularExponentScale);
+							MaterialInstanceNode->AddScalarParameterValue(ShininessParameterName + TEXT("MapWeight"), SpecularExponentScale);
 						}
 						else if (IsScalarInitialized(MaterialData.SpecularExponent))
 						{
@@ -1459,7 +1475,10 @@ bool UInterchangeOBJTranslator::Translate(UInterchangeBaseNodeContainer& BaseNod
                 if (const UInterchangeTexture2DNode* TransparencyTextureNode = CreateTexture2DNode(BaseNodeContainer, MaterialData.TransparencyTexture))
 				{
 					bOpacity = true;
-                    MaterialInstanceNode->AddTextureParameterValue(OpacityInputName + TEXT("Map"), TransparencyTextureNode->GetUniqueID());
+					if (MaterialInstanceNode->AddTextureParameterValue(OpacityInputName + TEXT("Map"), TransparencyTextureNode->GetUniqueID()))
+					{
+						MaterialInstanceNode->AddScalarParameterValue(OpacityInputName + TEXT("MapWeight"), 1.0f);
+					}
 				}
 				// Opacity texture is weighted by scalar opacity("d")
 				else if (!AddTexturedWeightedInput(BaseNodeContainer, MaterialInstanceNode, OpacityInputName, MaterialData.Opacity, MaterialData.OpacityTexture))
