[Interchange][FBX] ufbx - fixed smoothing

- ufbx computes vertex instances and their normals according to smoothing setup, just needed to recompute hardnesses. Now all works for pe-face smooting too.

- testcase - FortressArtSource\Weapons\Fortress_Weapons\Sniper_Rifle\Dragon_Sniper\Export\SK_Dragon_Sniper.FBX

#Synced-CL 50902575

diff --git a/Engine/Plugins/Interchange/Runtime/Source/Parsers/Fbx/Private/Ufbx/UfbxMesh.cpp b/Engine/Plugins/Interchange/Runtime/Source/Parsers/Fbx/Private/Ufbx/UfbxMesh.cpp
--- a/Engine/Plugins/Interchange/Runtime/Source/Parsers/Fbx/Private/Ufbx/UfbxMesh.cpp
+++ b/Engine/Plugins/Interchange/Runtime/Source/Parsers/Fbx/Private/Ufbx/UfbxMesh.cpp
@@ -496,32 +496,12 @@ namespace UE::Interchange::Private
 		}
 
 		// Provide edge smoothing
-		TEdgeAttributesRef<bool> EdgeHardnesses = Attributes.GetEdgeHardnesses();
-		for (int EdgeIndex = 0; EdgeIndex < Mesh.edges.count; ++EdgeIndex)
-		{
-			const ufbx_edge& Edge = Mesh.edges[EdgeIndex];
-
-			const FVertexID EdgeVertexA = FaceCollector.VertexMapping[Mesh.vertex_indices[Edge.a]];
-			const FVertexID EdgeVertexB = FaceCollector.VertexMapping[Mesh.vertex_indices[Edge.b]];
-
-			FEdgeID EdgeID = MeshDescription.GetVertexPairEdge(EdgeVertexA, EdgeVertexB);
-			// shouldn't edge already be here?
-			if (EdgeID == INDEX_NONE)
-			{
-				EdgeID = MeshDescription.CreateEdge(EdgeVertexA, EdgeVertexB);
-			}
-			if (EdgeIndex < Mesh.edge_smoothing.count)
-			{
-				EdgeHardnesses[EdgeID] = !Mesh.edge_smoothing[EdgeIndex];
-			}
-		}
-
 		HandleMesh(FaceCollector);
 
 		Builder.ResumeMeshDescriptionIndexing();
 #endif
 
-
+		FStaticMeshOperations::DetermineEdgeHardnessesFromVertexInstanceNormals(MeshDescription);
 		FStaticMeshOperations::ApplyTransform(MeshDescription, MeshGlobalTransform, true);
 		return true;
 	}
