MaterialX pipeline:
- Rename RecomputeNormals variable

#Synced-CL 51055552

diff --git a/Engine/Plugins/Interchange/Runtime/Source/Pipelines/Public/InterchangeMaterialXPipeline.h b/Engine/Plugins/Interchange/Runtime/Source/Pipelines/Public/InterchangeMaterialXPipeline.h
--- a/Engine/Plugins/Interchange/Runtime/Source/Pipelines/Public/InterchangeMaterialXPipeline.h
+++ b/Engine/Plugins/Interchange/Runtime/Source/Pipelines/Public/InterchangeMaterialXPipeline.h
@@ -106,10 +106,10 @@ public:
 	UPROPERTY(EditAnywhere, config, Category = "Volume")
 	bool bVolumetricMaterial = true;
 
-	/** Recompute the normals for displacement shaders, this option is best suited with procedural normals (for example shader graphs created specifically for offline-renderers).
+	/** Derive the normals from height for displacement shaders, this option is best suited with procedural normals (for example shader graphs created specifically for offline-renderers).
 	  * Otherwise this option should be disabled if there's already a normal map plugged to the normal input of the surfaceshader.
 	  */
-	UPROPERTY(EditAnywhere, config, Category = "Displacement")
+	UPROPERTY(EditAnywhere, config, Category = "Displacement", meta = (DisplayName = "Normals From Height"))
 	bool bRecomputeNormals = true;
 
 protected:
