[Interchange][FBX] ufbx - light color

#Synced-CL 51089491

diff --git a/Engine/Plugins/Interchange/Runtime/Source/Parsers/Fbx/Private/Ufbx/UfbxScene.cpp b/Engine/Plugins/Interchange/Runtime/Source/Parsers/Fbx/Private/Ufbx/UfbxScene.cpp
--- a/Engine/Plugins/Interchange/Runtime/Source/Parsers/Fbx/Private/Ufbx/UfbxScene.cpp
+++ b/Engine/Plugins/Interchange/Runtime/Source/Parsers/Fbx/Private/Ufbx/UfbxScene.cpp
@@ -461,11 +461,12 @@ namespace UE::Interchange::Private
 				
 					NodeContainer.SetupNode(BaseLightNode, LightUid, LightLabel, EInterchangeNodeContainerType::TranslatedAsset);
 
-					const ufbx_vec4 Color = {{{Light.color.x, Light.color.y, Light.color.z, 1.0}}};
-
 					// ufbx divides intensity by 100, re-multiplied here to match FBX SDK parser
 					BaseLightNode->SetCustomIntensity(static_cast<float>(Light.intensity) * 100.f);
 
+					const ufbx_vec4 Color = { {{Light.color.x, Light.color.y, Light.color.z, 1.0}} };
+					BaseLightNode->SetCustomLightColor(Convert::ConvertColor(Color));
+
 					if (UInterchangeLightNode* InterchangeLightNode = Cast<UInterchangeLightNode>(BaseLightNode))
 					{
 						float CustomAttenuationRadius = Light.decay == UFBX_LIGHT_DECAY_NONE ? FBXSDK_FLOAT_MAX : Light.decay_start;
