Interchange: Fix texture factory deleting existing texture assets when reimporting.

#Synced-CL 50957406

diff --git a/Engine/Plugins/Interchange/Runtime/Source/Import/Private/Texture/InterchangeTextureFactory.cpp b/Engine/Plugins/Interchange/Runtime/Source/Import/Private/Texture/InterchangeTextureFactory.cpp
--- a/Engine/Plugins/Interchange/Runtime/Source/Import/Private/Texture/InterchangeTextureFactory.cpp
+++ b/Engine/Plugins/Interchange/Runtime/Source/Import/Private/Texture/InterchangeTextureFactory.cpp
@@ -1957,6 +1957,7 @@ UInterchangeFactoryBase::FImportAssetResult UInterchangeTextureFactory::BeginImp
 	if (!ExistingAsset)
 	{
 		Texture = NewObject<UTexture>(Arguments.Parent, TextureClass, *Arguments.AssetName, RF_Public | RF_Standalone);
+		bCreatedAsset = true;
 
 		if (UTextureLightProfile* LightProfile = Cast<UTextureLightProfile>(Texture))
 		{
@@ -1971,6 +1972,8 @@ UInterchangeFactoryBase::FImportAssetResult UInterchangeTextureFactory::BeginImp
 	else
 	{
 		Texture = Cast<UTexture>(ExistingAsset);
+		bCreatedAsset = false;
+
 		//We allow override of existing Textures only if the translator is a pure texture translator or the user directly ask to re-import this object
 		if (!bIsReimport && Arguments.Translator->GetSupportedAssetTypes() != EInterchangeTranslatorAssetType::Textures)
 		{
@@ -2195,7 +2198,10 @@ UInterchangeFactoryBase::FImportAssetResult UInterchangeTextureFactory::EndImpor
 	if (UTexture* Texture = Cast<UTexture>(ExistingAsset))
 	{
 		ImportAssetResult.ImportedObject = Texture;
-		if (!ProcessedPayload.IsValid())
+
+		// We need to be careful to try and discard the asset only if we were the ones that originally created it. We don't want 
+		// to discard an asset we just found
+		if (bCreatedAsset && !ProcessedPayload.IsValid())
 		{
 			// The texture is not supported
 			if (!Arguments.ReimportObject)
diff --git a/Engine/Plugins/Interchange/Runtime/Source/Import/Public/Texture/InterchangeTextureFactory.h b/Engine/Plugins/Interchange/Runtime/Source/Import/Public/Texture/InterchangeTextureFactory.h
--- a/Engine/Plugins/Interchange/Runtime/Source/Import/Public/Texture/InterchangeTextureFactory.h
+++ b/Engine/Plugins/Interchange/Runtime/Source/Import/Public/Texture/InterchangeTextureFactory.h
@@ -87,6 +87,9 @@ private:
 
 	//If we import without a pure texture translator, we should not override an existing texture and we must skip the import. See the implementation of BeginImportAssetObject_GameThread function.
 	bool bSkipImport = false;
+
+	// This is set to true if the factory itself called NewObject and created the texture asset referenced by the factory node
+	bool bCreatedAsset = false;
 };
 
 
