Removed unnecessary ensures

#rnx
#okforgithhub public

#Synced-CL 51095113

diff --git a/Engine/Plugins/Interchange/Runtime/Source/Import/Private/Scene/InterchangeLevelFactory.cpp b/Engine/Plugins/Interchange/Runtime/Source/Import/Private/Scene/InterchangeLevelFactory.cpp
--- a/Engine/Plugins/Interchange/Runtime/Source/Import/Private/Scene/InterchangeLevelFactory.cpp
+++ b/Engine/Plugins/Interchange/Runtime/Source/Import/Private/Scene/InterchangeLevelFactory.cpp
@@ -159,12 +159,6 @@ namespace UE::Interchange::Private::InterchangeLevelFactory
 			return nullptr;
 		}
 
-		bool bShouldCreateLevel = false;
-		if (LevelFactoryNode->GetCustomShouldCreateLevel(bShouldCreateLevel))
-		{
-			ensure(bShouldCreateLevel);
-		}
-
 		bool bCreateWorldPartition = false;
 		LevelFactoryNode->GetCustomCreateWorldPartitionLevel(bCreateWorldPartition);
 
diff --git a/Engine/Plugins/Interchange/Runtime/Source/Import/Private/Scene/InterchangeLevelInstanceActorFactory.cpp b/Engine/Plugins/Interchange/Runtime/Source/Import/Private/Scene/InterchangeLevelInstanceActorFactory.cpp
--- a/Engine/Plugins/Interchange/Runtime/Source/Import/Private/Scene/InterchangeLevelInstanceActorFactory.cpp
+++ b/Engine/Plugins/Interchange/Runtime/Source/Import/Private/Scene/InterchangeLevelInstanceActorFactory.cpp
@@ -54,10 +54,7 @@ UObject* UInterchangeLevelInstanceActorFactory::GetObjectToReimport(UObject* Rei
 					FSoftObjectPath ObjectPath;
 					ensure(ActorFactoryNode->GetCustomReferenceObject(ObjectPath));
 
-					UObject* ObjectToReimport = ObjectPath.TryLoad();
-					ensure(ObjectToReimport);
-
-					return ObjectToReimport;
+					return ObjectPath.TryLoad();
 				}
 
 			}
