Interchange: Spawn all light actors with stationary mobility, similar to the Place Actors window.

USD: Spawn SkyLight actors with stationary mobility, like the other light actor types.

#Synced-CL 50957406

diff --git a/Engine/Plugins/Importers/USDImporter/Source/USDSchemas/Private/USDLuxLightTranslator.cpp b/Engine/Plugins/Importers/USDImporter/Source/USDSchemas/Private/USDLuxLightTranslator.cpp
--- a/Engine/Plugins/Importers/USDImporter/Source/USDSchemas/Private/USDLuxLightTranslator.cpp
+++ b/Engine/Plugins/Importers/USDImporter/Source/USDSchemas/Private/USDLuxLightTranslator.cpp
@@ -195,8 +195,6 @@ void FUsdLuxLightTranslator::UpdateComponents(USceneComponent* SceneComponent)
 			SkyLightComponent->Cubemap = TextuxeCube;
 			SkyLightComponent->SourceType = ESkyLightSourceType::SLS_SpecifiedCubemap;
 		}
-
-		SkyLightComponent->Mobility = EComponentMobility::Movable;	  // We won't bake geometry in the sky light so it needs to be movable
 	}
 
 	if (!LightComponent->IsRegistered())
diff --git a/Engine/Plugins/Interchange/Runtime/Source/Pipelines/Private/InterchangeGenericScenesPipeline.cpp b/Engine/Plugins/Interchange/Runtime/Source/Pipelines/Private/InterchangeGenericScenesPipeline.cpp
--- a/Engine/Plugins/Interchange/Runtime/Source/Pipelines/Private/InterchangeGenericScenesPipeline.cpp
+++ b/Engine/Plugins/Interchange/Runtime/Source/Pipelines/Private/InterchangeGenericScenesPipeline.cpp
@@ -863,6 +863,9 @@ void UInterchangeGenericLevelPipeline::SetUpFactoryNode(UInterchangeActorFactory
 	}
 	else if (const UInterchangeBaseLightNode* BaseLightNode = Cast<UInterchangeBaseLightNode>(TranslatedAssetNode))
 	{
+		// Default all lights to stationary as that's also what you get when placing UE lights onto the level manually
+		ActorFactoryNode->SetCustomMobility(EComponentMobility::Stationary);
+
 		if (UInterchangeBaseLightFactoryNode* BaseLightFactoryNode = Cast<UInterchangeBaseLightFactoryNode>(ActorFactoryNode))
 		{
 			if (FLinearColor LightColor; BaseLightNode->GetCustomLightColor(LightColor))
