Interchange: Remove sprite added to every spawned scene component.

#Synced-CL 50957406

diff --git a/Engine/Plugins/Interchange/Runtime/Source/Import/Private/Scene/InterchangeActorHelper.cpp b/Engine/Plugins/Interchange/Runtime/Source/Import/Private/Scene/InterchangeActorHelper.cpp
--- a/Engine/Plugins/Interchange/Runtime/Source/Import/Private/Scene/InterchangeActorHelper.cpp
+++ b/Engine/Plugins/Interchange/Runtime/Source/Import/Private/Scene/InterchangeActorHelper.cpp
@@ -148,9 +148,6 @@ namespace UE::Interchange::ActorHelper
 			if (!SpawnedActor->GetRootComponent())
 			{
 				USceneComponent* RootComponent = NewObject<USceneComponent>(SpawnedActor, USceneComponent::GetDefaultSceneRootVariableName(), RF_Transactional);
-#if WITH_EDITORONLY_DATA
-				RootComponent->bVisualizeComponent = true;
-#endif
 				SpawnedActor->SetRootComponent(RootComponent);
 				SpawnedActor->AddInstanceComponent(RootComponent);
 			}
