Interchange:
- Remove auto detect greyscale image

#Synced-CL 51055552

diff --git a/Engine/Plugins/Interchange/Runtime/Source/Import/Private/Texture/InterchangeImageWrapperTranslator.cpp b/Engine/Plugins/Interchange/Runtime/Source/Import/Private/Texture/InterchangeImageWrapperTranslator.cpp
--- a/Engine/Plugins/Interchange/Runtime/Source/Import/Private/Texture/InterchangeImageWrapperTranslator.cpp
+++ b/Engine/Plugins/Interchange/Runtime/Source/Import/Private/Texture/InterchangeImageWrapperTranslator.cpp
@@ -191,11 +191,6 @@ TOptional<UE::Interchange::FImportImage> UInterchangeImageWrapperTranslator::Get
 			FDecompressedImageOutput DecompressedImage;
 			if (ImageWrapperModule.DecompressImage(Buffer, Length, DecompressedImage))
 			{
-				if (UE::TextureUtilitiesCommon::AutoDetectAndChangeGrayScale(DecompressedImage.MipMapImage))
-				{
-					UE_LOG(LogInterchangeImport, Display, TEXT("Auto-detected grayscale, image changed to G8"));
-				}
-
 				ETextureSourceFormat TextureFormat = FImageCoreUtils::ConvertToTextureSourceFormat(DecompressedImage.MipMapImage.Format);
 				bool bSRGB = DecompressedImage.MipMapImage.GammaSpace != EGammaSpace::Linear;
 				
@@ -248,12 +243,6 @@ TOptional<UE::Interchange::FImportImage> UInterchangeImageWrapperTranslator::Get
 			FImage LoadedImage;
 			if (ImageWrapperModule.DecompressImage(Buffer, Length, LoadedImage))
 			{
-				// Todo interchange: should these payload modification be part of the pipeline, factory or stay there?
-				if (UE::TextureUtilitiesCommon::AutoDetectAndChangeGrayScale(LoadedImage))
-				{
-					UE_LOG(LogInterchangeImport, Display, TEXT("Auto-detected grayscale, image changed to G8"));
-				}
-
 				ETextureSourceFormat TextureFormat = FImageCoreUtils::ConvertToTextureSourceFormat(LoadedImage.Format);
 				bool bSRGB = LoadedImage.GammaSpace != EGammaSpace::Linear;
 
