Interchange - Various clean-up and improvements

* Exposed a `UInterchangeBaseNodeContainer.GetNodesOfType<T>` to simplify some part of the code and prepare for future optimization inside the UInterchangeBaseNodeContainer.
* Using the `UInterchangeBaseNode::Name` to be able to debug in VS the "DisplayNade | UniqueUid" of a node.
* Using `UInterchangeBaseNodeContainer.IterateNodesOfType<T>` instead of `UInterchangeBaseNodeContainer.IterateNodes` when appropriate.
* Displaying both Class Name and Object Name in the GraphInspector Detail panel.
* Fixed an issue with GeometryCacheFactorNode/ActorFactoryNode Uid clash
* Removing some unused code

#tests All Interchange tests passing locally
#rnx

#virtualized
#Synced-CL 51246323

#Synced-CL 51246323

#Synced-CL 51246323

diff --git a/Engine/Plugins/Experimental/Enterprise/DatasmithInterchange/Source/DatasmithInterchange/Private/InterchangeDatasmithLevelPipeline.cpp b/Engine/Plugins/Experimental/Enterprise/DatasmithInterchange/Source/DatasmithInterchange/Private/InterchangeDatasmithLevelPipeline.cpp
--- a/Engine/Plugins/Experimental/Enterprise/DatasmithInterchange/Source/DatasmithInterchange/Private/InterchangeDatasmithLevelPipeline.cpp
+++ b/Engine/Plugins/Experimental/Enterprise/DatasmithInterchange/Source/DatasmithInterchange/Private/InterchangeDatasmithLevelPipeline.cpp
@@ -139,7 +139,7 @@ void UInterchangeDatasmithLevelPipeline::ExecutePostImportPipeline(const UInterc
 	}
 }
 
-UInterchangeActorFactoryNode* UInterchangeDatasmithLevelPipeline::CreateActorFactoryNode(const UInterchangeSceneNode* SceneNode, const UInterchangeBaseNode* TranslatedAssetNode) const
+UInterchangeActorFactoryNode* UInterchangeDatasmithLevelPipeline::CreateActorFactoryNode(const UInterchangeBaseNode* TranslatedAssetNode) const
 {
 	if (!ensure(BaseNodeContainer))
 	{
@@ -152,7 +152,7 @@ UInterchangeActorFactoryNode* UInterchangeDatasmithLevelPipeline::CreateActorFac
 	}
 	else
 	{
-		return Super::CreateActorFactoryNode(SceneNode, TranslatedAssetNode);
+		return Super::CreateActorFactoryNode(TranslatedAssetNode);
 	}
 }
 
diff --git a/Engine/Plugins/Experimental/Enterprise/DatasmithInterchange/Source/DatasmithInterchange/Private/InterchangeDatasmithLevelPipeline.h b/Engine/Plugins/Experimental/Enterprise/DatasmithInterchange/Source/DatasmithInterchange/Private/InterchangeDatasmithLevelPipeline.h
--- a/Engine/Plugins/Experimental/Enterprise/DatasmithInterchange/Source/DatasmithInterchange/Private/InterchangeDatasmithLevelPipeline.h
+++ b/Engine/Plugins/Experimental/Enterprise/DatasmithInterchange/Source/DatasmithInterchange/Private/InterchangeDatasmithLevelPipeline.h
@@ -25,7 +25,7 @@ protected:
 	virtual void ExecutePipeline(UInterchangeBaseNodeContainer* BaseNodeContainer, const TArray<UInterchangeSourceData*>& SourceDatas, const FString& ContentBasePath) override;
 	virtual void ExecutePostImportPipeline(const UInterchangeBaseNodeContainer* BaseNodeContainer, const FString& NodeKey, UObject* CreatedAsset, bool bIsAReimport) override;
 
-	virtual UInterchangeActorFactoryNode* CreateActorFactoryNode(const UInterchangeSceneNode* SceneNode, const UInterchangeBaseNode* TranslatedAssetNode) const override;
+	virtual UInterchangeActorFactoryNode* CreateActorFactoryNode(const UInterchangeBaseNode* TranslatedAssetNode) const override;
 
 	virtual void SetUpFactoryNode(UInterchangeActorFactoryNode* ActorFactoryNode, const UInterchangeSceneNode* SceneNode, const UInterchangeBaseNode* TranslatedAssetNode) const override;
 
diff --git a/Engine/Plugins/Experimental/Enterprise/DatasmithInterchange/Source/DatasmithInterchange/Private/InterchangeDatasmithStaticMeshPipeline.cpp b/Engine/Plugins/Experimental/Enterprise/DatasmithInterchange/Source/DatasmithInterchange/Private/InterchangeDatasmithStaticMeshPipeline.cpp
--- a/Engine/Plugins/Experimental/Enterprise/DatasmithInterchange/Source/DatasmithInterchange/Private/InterchangeDatasmithStaticMeshPipeline.cpp
+++ b/Engine/Plugins/Experimental/Enterprise/DatasmithInterchange/Source/DatasmithInterchange/Private/InterchangeDatasmithStaticMeshPipeline.cpp
@@ -120,16 +120,14 @@ void UInterchangeDatasmithStaticMeshPipeline::ApplyAdditionalData(const UInterch
 
 	const FString& MeshNodeUid = TargetNodes[0];
 
-	UDatasmithInterchangeStaticMeshDataNode* StaticMeshDataNode = nullptr;
-	NodeContainer->IterateNodes([&StaticMeshDataNode](const FString& NodeUid, UInterchangeBaseNode* Node)
-		{
-			if (UDatasmithInterchangeStaticMeshDataNode* TypedNode = Cast<UDatasmithInterchangeStaticMeshDataNode>(Node))
-			{
-				StaticMeshDataNode = TypedNode;
-				return;
-			}
-		});
+	TArray<UDatasmithInterchangeStaticMeshDataNode*> StaticMeshDataNodes;
+	NodeContainer->GetNodesOfType<UDatasmithInterchangeStaticMeshDataNode>(StaticMeshDataNodes);
+	if (StaticMeshDataNodes.Num() == 0)
+	{
+		return;
+	}
 
+	const UDatasmithInterchangeStaticMeshDataNode* StaticMeshDataNode = StaticMeshDataNodes[0];
 	if (!StaticMeshDataNode)
 	{
 		return;
diff --git a/Engine/Plugins/Experimental/Enterprise/DatasmithInterchange/Source/DatasmithInterchange/Private/InterchangeDatasmithUtils.h b/Engine/Plugins/Experimental/Enterprise/DatasmithInterchange/Source/DatasmithInterchange/Private/InterchangeDatasmithUtils.h
--- a/Engine/Plugins/Experimental/Enterprise/DatasmithInterchange/Source/DatasmithInterchange/Private/InterchangeDatasmithUtils.h
+++ b/Engine/Plugins/Experimental/Enterprise/DatasmithInterchange/Source/DatasmithInterchange/Private/InterchangeDatasmithUtils.h
@@ -52,13 +52,7 @@ namespace UE::DatasmithInterchange::NodeUtils
 	{
 		TArray<T*> TypedNodes;
 
-		BaseNodeContainer->IterateNodes([&TypedNodes, &BaseNodeContainer](const FString& NodeUid, UInterchangeBaseNode* Node)
-		{
-			if (T* TypedNode = Cast<T>(Node))
-			{
-				TypedNodes.Add(TypedNode);
-			}
-		});
+		BaseNodeContainer->GetNodesOfType<T>(TypedNodes);
 
 		return TypedNodes;
 	}
diff --git a/Engine/Plugins/Experimental/InterchangeAxF/Runtime/Source/InterchangeAxF/Private/AxFInterchangePipeline.cpp b/Engine/Plugins/Experimental/InterchangeAxF/Runtime/Source/InterchangeAxF/Private/AxFInterchangePipeline.cpp
--- a/Engine/Plugins/Experimental/InterchangeAxF/Runtime/Source/InterchangeAxF/Private/AxFInterchangePipeline.cpp
+++ b/Engine/Plugins/Experimental/InterchangeAxF/Runtime/Source/InterchangeAxF/Private/AxFInterchangePipeline.cpp
@@ -127,10 +127,8 @@ void UAxFInterchangePipeline::ExecutePipeline(UInterchangeBaseNodeContainer* Bas
 
 		FString MaterialInstanceUID = TEXT("MI_") + AxFNode->GetDisplayLabel();
 
-        UInterchangeMaterialInstanceFactoryNode* MaterialInstance =
-            NewObject<UInterchangeMaterialInstanceFactoryNode>(BaseNodeContainer, FName(MaterialInstanceUID));
-        MaterialInstance->InitializeNode(
-            MaterialInstanceUID, AxFNode->GetDisplayLabel(), EInterchangeNodeContainerType::FactoryData);
+        UInterchangeMaterialInstanceFactoryNode* MaterialInstance = NewObject<UInterchangeMaterialInstanceFactoryNode>(BaseNodeContainer, FName(MaterialInstanceUID));
+		BaseNodeContainer->SetupNode(MaterialInstance, MaterialInstanceUID, AxFNode->GetDisplayLabel(), EInterchangeNodeContainerType::FactoryData);
 
 		for (const EAxFFeature Feature : Features)
 		{
@@ -214,7 +212,6 @@ void UAxFInterchangePipeline::ExecutePipeline(UInterchangeBaseNodeContainer* Bas
         MaterialInstance->SetCustomParent(FullBaseMaterialPath);
         MaterialInstance->MarkPackageDirty();
         MaterialInstance->PostEditChange();
-        BaseNodeContainer->AddNode(MaterialInstance);
     }
 
 	BaseNodeContainer->IterateNodesOfType<UInterchangeFactoryBaseNode>(
@@ -238,7 +235,7 @@ void UAxFInterchangePipeline::ExecutePostImportPipeline(
     UMaterialInstanceConstant* Instance = Cast<UMaterialInstanceConstant>(CreatedAsset);
     if (Instance != nullptr && !AxFObjectNodes.IsEmpty())
     {
-        TObjectPtr<UAxFMaterialObjectNode>* FoundData = AxFObjectNodes.FindByPredicate(
+		TObjectPtr<UAxFMaterialObjectNode>* FoundData = AxFObjectNodes.FindByPredicate(
             [Instance](UAxFMaterialObjectNode const* Node) { return Node->GetDisplayLabel() == Instance->GetName(); });
 
         if (FoundData == nullptr)
diff --git a/Engine/Plugins/Experimental/InterchangeAxF/Runtime/Source/InterchangeAxF/Private/AxFTranslator.cpp b/Engine/Plugins/Experimental/InterchangeAxF/Runtime/Source/InterchangeAxF/Private/AxFTranslator.cpp
--- a/Engine/Plugins/Experimental/InterchangeAxF/Runtime/Source/InterchangeAxF/Private/AxFTranslator.cpp
+++ b/Engine/Plugins/Experimental/InterchangeAxF/Runtime/Source/InterchangeAxF/Private/AxFTranslator.cpp
@@ -183,21 +183,19 @@ bool UAxFTranslator::Translate(UInterchangeBaseNodeContainer &BaseNodeContainer)
 
             FString TextureUID = TEXT("T_") + DisplayLabel + TEXT("_") + UEParameterName;
             UInterchangeTexture2DNode *Texture2D = NewObject<UInterchangeTexture2DNode>(&BaseNodeContainer, FName(TextureUID));
-            Texture2D->InitializeNode(TextureUID, TextureUID, EInterchangeNodeContainerType::TranslatedAsset);
+			BaseNodeContainer.SetupNode(Texture2D, TextureUID, TextureUID, EInterchangeNodeContainerType::TranslatedAsset);
             Texture2D->SetPayLoadKey(Path + "|" + Name + "." + Extension + "|" + TexturePath);
             Texture2D->AddFloatAttribute(TEXT("WidthMM"), WidthMM);
             Texture2D->AddFloatAttribute(TEXT("HeightMM"), HeightMM);
-            BaseNodeContainer.AddNode(Texture2D);
 
             if (UEParameterName == TEXT("CPA2ColorTable"))
             {
                 FString SpecProfileUID = TEXT("SP_") + DisplayLabel;
                 UInterchangeSpecularProfileFactoryNode* SpecularProfileNode = NewObject<UInterchangeSpecularProfileFactoryNode>(&BaseNodeContainer, FName(SpecProfileUID));
-                SpecularProfileNode->InitializeNode(SpecProfileUID, SpecProfileUID, EInterchangeNodeContainerType::TranslatedAsset);
+				BaseNodeContainer.SetupNode(SpecularProfileNode, SpecProfileUID, SpecProfileUID, EInterchangeNodeContainerType::TranslatedAsset);
                 SpecularProfileNode->SetCustomFormat(ESpecularProfileFormat::HalfVector);
                 SpecularProfileNode->SetCustomTexture(TextureUID);
                 SpecularProfileNode->SetCustomSubPath(AxFMaterialObjectNode->GetDisplayLabel());
-                BaseNodeContainer.AddNode(SpecularProfileNode);
             }
         }
         else if (bool ValueBool; JSON->TryGetBoolField(ParameterName, ValueBool)) 
diff --git a/Engine/Plugins/Experimental/InterchangeAxF/Runtime/Source/InterchangeAxF/Public/AxFInterchangePipeline.h b/Engine/Plugins/Experimental/InterchangeAxF/Runtime/Source/InterchangeAxF/Public/AxFInterchangePipeline.h
--- a/Engine/Plugins/Experimental/InterchangeAxF/Runtime/Source/InterchangeAxF/Public/AxFInterchangePipeline.h
+++ b/Engine/Plugins/Experimental/InterchangeAxF/Runtime/Source/InterchangeAxF/Public/AxFInterchangePipeline.h
@@ -76,12 +76,12 @@ protected:
     virtual void GetSupportAssetClasses(TArray<UClass*>& PipelineSupportAssetClasses) const override;
 
 private:
-	UPROPERTY()
+	UPROPERTY(Transient)
 	TArray<TObjectPtr<UAxFMaterialObjectNode>> AxFObjectNodes;
 
     // Texture handling
-	UPROPERTY()
-    TArray<TObjectPtr<UInterchangeTexture2DFactoryNode>> Texture2DFactoryNodes;
+	UPROPERTY(Transient)
+	TArray<TObjectPtr<UInterchangeTexture2DFactoryNode>> Texture2DFactoryNodes;
 
 	UPROPERTY()
     TMap<FString, FAxFTextureMap> TexturesMap;
diff --git a/Engine/Plugins/Fab/Source/Fab/Private/Pipelines/InterchangeMegascansPipeline.cpp b/Engine/Plugins/Fab/Source/Fab/Private/Pipelines/InterchangeMegascansPipeline.cpp
--- a/Engine/Plugins/Fab/Source/Fab/Private/Pipelines/InterchangeMegascansPipeline.cpp
+++ b/Engine/Plugins/Fab/Source/Fab/Private/Pipelines/InterchangeMegascansPipeline.cpp
@@ -554,7 +554,7 @@ void UInterchangeMegascansPipeline::SetupStaticMeshParams(const UInterchangeStat
 		const UInterchangeSceneNode* SceneNode = FindNodeOfTypeByName<UInterchangeSceneNode>(StaticMeshFactoryNode->GetDisplayLabel());
 		SetupMeshLod(ParentStaticMeshFactoryNode, SceneNode, LodIndex);
 
-		BaseNodeContainer->ReplaceNode(StaticMeshFactoryUid, nullptr);
+		BaseNodeContainer->RemoveFactoryNode(StaticMeshFactoryUid);
 	}
 }
 
diff --git a/Engine/Plugins/Interchange/Editor/Source/Pipelines/Private/InterchangeEditorPipelineDetails.cpp b/Engine/Plugins/Interchange/Editor/Source/Pipelines/Private/InterchangeEditorPipelineDetails.cpp
--- a/Engine/Plugins/Interchange/Editor/Source/Pipelines/Private/InterchangeEditorPipelineDetails.cpp
+++ b/Engine/Plugins/Interchange/Editor/Source/Pipelines/Private/InterchangeEditorPipelineDetails.cpp
@@ -920,7 +920,11 @@ void FInterchangeBaseNodeDetailsCustomization::CustomizeDetails( IDetailLayoutBu
 		IDetailCategoryBuilder& AttributeCategoryBuilder = DetailBuilder.EditCategory(NodeInformationCategoryName, NodeInformationCategoryText);
 		
 		FText NodeTypeValueText = FText::GetEmpty();
-		FText ClassNameValueText = FText::GetEmpty();
+		FText NodeClassValueText = FText::GetEmpty();
+		FText ObjectClassNameValueText = FText::GetEmpty();
+
+		NodeClassValueText = FText::FromString(InterchangeBaseNode->GetClass()->GetName());
+
 		if (UInterchangeFactoryBaseNode* FactoryNode = Cast<UInterchangeFactoryBaseNode>(InterchangeBaseNode))
 		{
 			const FString ClassName = FactoryNode->GetObjectClass() 
@@ -928,46 +932,42 @@ void FInterchangeBaseNodeDetailsCustomization::CustomizeDetails( IDetailLayoutBu
 										: "None - Unprocessed";
 
 			NodeTypeValueText = LOCTEXT("NodeFactoryInformation_FactoryNodeValueText", "Factory Node");
-			ClassNameValueText = FText::FromString(ClassName);
+			ObjectClassNameValueText = FText::FromString(ClassName);
 		}
 		else
 		{
 			NodeTypeValueText = LOCTEXT("NodeFactoryInformation_TranslatedNodeValueText", "Translated Node");
-			ClassNameValueText = FText::FromString(InterchangeBaseNode->GetClass()->GetName());
 		}
 
+		auto AddNodeInformationRow = [&AttributeCategoryBuilder, &NodeInformationCategoryText](const FText& Name, const FText& Value)
+			{
+				AttributeCategoryBuilder.AddCustomRow(NodeInformationCategoryText)
+					.FilterString(Name)
+					.NameContent()
+					[
+						SNew(STextBlock)
+							.Text(Name)
+							.Font(IDetailLayoutBuilder::GetDetailFont())
+					]
+					.ValueContent()
+					[
+						SNew(STextBlock)
+							.Text(Value)
+							.Font(IDetailLayoutBuilder::GetDetailFont())
+					];
+			};
 
 		const FText NodeTypeNameText = LOCTEXT("NodeFactoryInformationNodeTypeNameText", "Node Type");
-		const FText ClassNameNameText = LOCTEXT("NodeFactoryInformationClassNameNameText", "Class Name");
-		AttributeCategoryBuilder.AddCustomRow(NodeInformationCategoryText)
-			.FilterString(NodeTypeNameText)
-			.NameContent()
-			[
-				SNew(STextBlock)
-					.Text(NodeTypeNameText)
-					.Font(IDetailLayoutBuilder::GetDetailFont())
-			]
-			.ValueContent()
-			[
-				SNew(STextBlock)
-					.Text(NodeTypeValueText)
-					.Font(IDetailLayoutBuilder::GetDetailFont())
-			];
+		AddNodeInformationRow(NodeTypeNameText, NodeTypeValueText);
 
-		AttributeCategoryBuilder.AddCustomRow(NodeInformationCategoryText)
-			.FilterString(ClassNameNameText)
-			.NameContent()
-			[
-				SNew(STextBlock)
-					.Text(ClassNameNameText)
-					.Font(IDetailLayoutBuilder::GetDetailFont())
-			]
-			.ValueContent()
-			[
-				SNew(STextBlock)
-					.Text(ClassNameValueText)
-					.Font(IDetailLayoutBuilder::GetDetailFont())
-			];
+		const FText NodeClassNameText = LOCTEXT("NodeFactoryInformationNodeClassNameText", "Node Class");
+		AddNodeInformationRow(NodeClassNameText, NodeClassValueText);
+
+		if (!ObjectClassNameValueText.IsEmpty())
+		{
+			const FText ObjectClassNameNameText = LOCTEXT("NodeFactoryInformationObjectClassNameNameText", "Object Class Name");
+			AddNodeInformationRow(ObjectClassNameNameText, ObjectClassNameValueText);
+		}
 	}
 
 	TArray< UE::Interchange::FAttributeKey> AttributeKeys;
diff --git a/Engine/Plugins/Interchange/Extensions/OpenUSD/Source/Import/Private/InterchangeUSDPipeline.cpp b/Engine/Plugins/Interchange/Extensions/OpenUSD/Source/Import/Private/InterchangeUSDPipeline.cpp
--- a/Engine/Plugins/Interchange/Extensions/OpenUSD/Source/Import/Private/InterchangeUSDPipeline.cpp
+++ b/Engine/Plugins/Interchange/Extensions/OpenUSD/Source/Import/Private/InterchangeUSDPipeline.cpp
@@ -431,7 +431,7 @@ namespace UE::InterchangeUsdPipeline::Private
 			}
 
 			// Disable the pseudoroot itself (we shouldn't have any asset node for it though)
-			NodeContainer.ReplaceNode(FactoryNodeUid, nullptr);
+			NodeContainer.RemoveFactoryNode(FactoryNodeUid);
 		}
 	}
 
diff --git a/Engine/Plugins/Interchange/Runtime/Source/Import/Private/MaterialX/MaterialXUtils/MaterialXSurfaceMaterial.cpp b/Engine/Plugins/Interchange/Runtime/Source/Import/Private/MaterialX/MaterialXUtils/MaterialXSurfaceMaterial.cpp
--- a/Engine/Plugins/Interchange/Runtime/Source/Import/Private/MaterialX/MaterialXUtils/MaterialXSurfaceMaterial.cpp
+++ b/Engine/Plugins/Interchange/Runtime/Source/Import/Private/MaterialX/MaterialXUtils/MaterialXSurfaceMaterial.cpp
@@ -57,8 +57,7 @@ UInterchangeBaseNode* FMaterialXSurfaceMaterial::Translate(MaterialX::NodePtr Su
 				if (!ShaderGraphNode)
 				{
 					ShaderGraphNode = NewObject<UInterchangeShaderGraphNode>(&NodeContainer);
-					ShaderGraphNode->InitializeNode(ShaderGraphNodeUID, SurfaceMaterialNode->getName().c_str(), EInterchangeNodeContainerType::TranslatedAsset);
-					NodeContainer.AddNode(ShaderGraphNode);
+					NodeContainer.SetupNode(ShaderGraphNode, ShaderGraphNodeUID, SurfaceMaterialNode->getName().c_str(), EInterchangeNodeContainerType::TranslatedAsset);
 					bIsShaderGraphInContainer = false;
 				}
 
diff --git a/Engine/Plugins/Interchange/Runtime/Source/Import/Private/MaterialX/MaterialXUtils/MaterialXVolumeMaterial.cpp b/Engine/Plugins/Interchange/Runtime/Source/Import/Private/MaterialX/MaterialXUtils/MaterialXVolumeMaterial.cpp
--- a/Engine/Plugins/Interchange/Runtime/Source/Import/Private/MaterialX/MaterialXUtils/MaterialXVolumeMaterial.cpp
+++ b/Engine/Plugins/Interchange/Runtime/Source/Import/Private/MaterialX/MaterialXUtils/MaterialXVolumeMaterial.cpp
@@ -45,12 +45,11 @@ UInterchangeBaseNode* FMaterialXVolumeMaterial::Translate(MaterialX::NodePtr Vol
 			if(TSharedPtr<FMaterialXSurfaceShaderAbstract> ShaderTranslator = StaticCastSharedPtr<FMaterialXSurfaceShaderAbstract>(FMaterialXManager::GetInstance().GetShaderTranslator(VolumeShaderNode->getCategory().c_str(), NodeContainer)))
 			{
 				ShaderGraphNode = NewObject<UInterchangeShaderGraphNode>(&NodeContainer);
-				ShaderGraphNode->InitializeNode(ShaderGraphNodeUID, VolumeMaterialNode->getName().c_str(), EInterchangeNodeContainerType::TranslatedAsset);
+				NodeContainer.SetupNode(ShaderGraphNode, ShaderGraphNodeUID, VolumeMaterialNode->getName().c_str(), EInterchangeNodeContainerType::TranslatedAsset);
 				ShaderTranslator->ShaderGraphNode = ShaderGraphNode;
 				ShaderTranslator->SurfaceMaterialName = VolumeShaderNode->getName().c_str();
 				ShaderTranslator->SetTranslator(Translator);
 				ShaderTranslator->Translate(VolumeShaderNode);
-				NodeContainer.AddNode(ShaderGraphNode);
 
 				AnalyticsAttributes.Emplace(TEXT("MaterialType"), TEXT("Volume Material"));
 			}
diff --git a/Engine/Plugins/Interchange/Runtime/Source/Import/Private/Scene/InterchangeCameraActorFactory.cpp b/Engine/Plugins/Interchange/Runtime/Source/Import/Private/Scene/InterchangeCameraActorFactory.cpp
--- a/Engine/Plugins/Interchange/Runtime/Source/Import/Private/Scene/InterchangeCameraActorFactory.cpp
+++ b/Engine/Plugins/Interchange/Runtime/Source/Import/Private/Scene/InterchangeCameraActorFactory.cpp
@@ -63,11 +63,10 @@ namespace UE::InterchangeCameraActorFactory::Private
 			FactoryNodeCopy->RemoveCustomAttributesForClass(UCineCameraComponent::StaticClass());
 			ActorHelper::ApplyAllCustomAttributes(ParamsCopy, *RootSceneComponent);
 		}
-		NodeContainer->ReplaceNode(FactoryNodeCopy->GetUniqueID(), CameraFactoryNode);
+		NodeContainer->ReplaceFactoryNode(FactoryNodeCopy->GetUniqueID(), CameraFactoryNode);
 	}
 }
 
-
 UClass* UInterchangeCineCameraActorFactory::GetFactoryClass() const
 {
 	return ACineCameraActor::StaticClass();
diff --git a/Engine/Plugins/Interchange/Runtime/Source/Nodes/Private/InterchangeSceneNode.cpp b/Engine/Plugins/Interchange/Runtime/Source/Nodes/Private/InterchangeSceneNode.cpp
--- a/Engine/Plugins/Interchange/Runtime/Source/Nodes/Private/InterchangeSceneNode.cpp
+++ b/Engine/Plugins/Interchange/Runtime/Source/Nodes/Private/InterchangeSceneNode.cpp
@@ -277,12 +277,9 @@ bool UInterchangeSceneNode::SetCustomAssetInstanceUid(const FString& AttributeVa
 
 void UInterchangeSceneNode::ResetAllGlobalTransformCaches(const UInterchangeBaseNodeContainer* BaseNodeContainer)
 {
-	BaseNodeContainer->IterateNodes([](const FString& NodeUid, UInterchangeBaseNode* Node)
+	BaseNodeContainer->IterateNodesOfType<UInterchangeSceneNode>([](const FString& NodeUid, UInterchangeSceneNode* SceneNode)
 		{
-			if (UInterchangeSceneNode* SceneNode = Cast<UInterchangeSceneNode>(Node))
-			{
-				SceneNode->GlobalTransformCache.Reset();
-			}
+			SceneNode->GlobalTransformCache.Reset();
 		});
 }
 
diff --git a/Engine/Plugins/Interchange/Runtime/Source/Pipelines/Private/InterchangeGenericAudioPipeline.cpp b/Engine/Plugins/Interchange/Runtime/Source/Pipelines/Private/InterchangeGenericAudioPipeline.cpp
--- a/Engine/Plugins/Interchange/Runtime/Source/Pipelines/Private/InterchangeGenericAudioPipeline.cpp
+++ b/Engine/Plugins/Interchange/Runtime/Source/Pipelines/Private/InterchangeGenericAudioPipeline.cpp
@@ -29,15 +29,9 @@ void UInterchangeGenericAudioPipeline::ExecutePipeline(UInterchangeBaseNodeConta
 
 	BaseNodeContainer = InBaseNodeContainer;
 
-	BaseNodeContainer->IterateNodes([this](const FString& NodeUid, UInterchangeBaseNode* Node)
+	BaseNodeContainer->IterateNodesOfType<UInterchangeAudioSoundWaveNode>([this](const FString& NodeUid, UInterchangeAudioSoundWaveNode* Node)
 	{
-		if (Node->GetNodeContainerType() == EInterchangeNodeContainerType::TranslatedAsset)
-		{
-			if (UInterchangeAudioSoundWaveNode* SoundWaveNode = Cast<UInterchangeAudioSoundWaveNode>(Node))
-			{
-				SoundWaveNodes.Add(SoundWaveNode);
-			}
-		}
+		SoundWaveNodes.Add(Node);
 	});
 
 	if (bImportSounds)
diff --git a/Engine/Plugins/Interchange/Runtime/Source/Pipelines/Private/InterchangeGenericGeometryCachePipeline.cpp b/Engine/Plugins/Interchange/Runtime/Source/Pipelines/Private/InterchangeGenericGeometryCachePipeline.cpp
--- a/Engine/Plugins/Interchange/Runtime/Source/Pipelines/Private/InterchangeGenericGeometryCachePipeline.cpp
+++ b/Engine/Plugins/Interchange/Runtime/Source/Pipelines/Private/InterchangeGenericGeometryCachePipeline.cpp
@@ -193,7 +193,7 @@ UInterchangeGeometryCacheFactoryNode* UInterchangeGenericMeshPipeline::CreateGeo
 		}
 	}
 
-	const FString StaticMeshUid = UInterchangeFactoryBaseNode::BuildFactoryNodeUid(GeometryCacheUid);
+	const FString StaticMeshUid = TEXT("Factory_GeometryCache_") + GeometryCacheUid;
 	UInterchangeGeometryCacheFactoryNode* GeometryCacheFactoryNode = NewObject<UInterchangeGeometryCacheFactoryNode>(BaseNodeContainer, NAME_None);
 	if (!ensure(GeometryCacheFactoryNode))
 	{
diff --git a/Engine/Plugins/Interchange/Runtime/Source/Pipelines/Private/InterchangeGenericGroomPipeline.cpp b/Engine/Plugins/Interchange/Runtime/Source/Pipelines/Private/InterchangeGenericGroomPipeline.cpp
--- a/Engine/Plugins/Interchange/Runtime/Source/Pipelines/Private/InterchangeGenericGroomPipeline.cpp
+++ b/Engine/Plugins/Interchange/Runtime/Source/Pipelines/Private/InterchangeGenericGroomPipeline.cpp
@@ -106,18 +106,10 @@ void UInterchangeGenericGroomPipeline::ExecutePipeline(UInterchangeBaseNodeConta
 	}
 
 	// Find all the translated groom nodes we need for this pipeline
-	TArray<UInterchangeGroomNode*> GroomNodes;
 	if (bImportGrooms || bImportGroomCaches)
 	{
-		InBaseNodeContainer->IterateNodes(
-			[&GroomNodes](const FString& NodeUid, UInterchangeBaseNode* Node)
-			{
-				if (UInterchangeGroomNode* GroomNode = Cast<UInterchangeGroomNode>(Node))
-				{
-					GroomNodes.Add(GroomNode);
-				}
-			}
-		);
+		TArray<UInterchangeGroomNode*> GroomNodes;
+		InBaseNodeContainer->GetNodesOfType<UInterchangeGroomNode>(GroomNodes);
 
 		// Create groom factory nodes
 		for (const UInterchangeGroomNode* GroomNode : GroomNodes)
diff --git a/Engine/Plugins/Interchange/Runtime/Source/Pipelines/Private/InterchangeGenericScenesPipeline.cpp b/Engine/Plugins/Interchange/Runtime/Source/Pipelines/Private/InterchangeGenericScenesPipeline.cpp
--- a/Engine/Plugins/Interchange/Runtime/Source/Pipelines/Private/InterchangeGenericScenesPipeline.cpp
+++ b/Engine/Plugins/Interchange/Runtime/Source/Pipelines/Private/InterchangeGenericScenesPipeline.cpp
@@ -78,7 +78,7 @@ if(AttributeType AttributeName; TranslatedNode->GetCustom##AttributeName(Attribu
 namespace UE::Interchange::Private
 {
 	//Either a (!Joint || RootJoint) can be a parent (only those get FactoryNodes) :
-	FString FindFactoryParentSceneNodeUid(UInterchangeBaseNodeContainer* BaseNodeContainer, TArray<FString>& ActiveSkeletonUids, const UInterchangeSceneNode* SceneNode)
+	FString FindFactoryParentSceneNodeUid(UInterchangeBaseNodeContainer* BaseNodeContainer, TSet<FString>& ActiveSkeletonUids, const UInterchangeSceneNode* SceneNode)
 	{
 		FString ParentUid = SceneNode->GetParentUid();
 
@@ -531,7 +531,7 @@ void UInterchangeGenericLevelPipeline::ExecuteSceneNodePreImport(const FTransfor
 		}
 	}
 
-	UInterchangeActorFactoryNode* ActorFactoryNode = CreateActorFactoryNode(SceneNode, TranslatedAssetNode);
+	UInterchangeActorFactoryNode* ActorFactoryNode = CreateActorFactoryNode(TranslatedAssetNode);
 
 	if (!ensure(ActorFactoryNode))
 	{
@@ -679,7 +679,7 @@ void UInterchangeGenericLevelPipeline::ExecuteSceneNodePreImport(const FTransfor
 #endif
 }
 
-UInterchangeActorFactoryNode* UInterchangeGenericLevelPipeline::CreateActorFactoryNode(const UInterchangeSceneNode* SceneNode, const UInterchangeBaseNode* TranslatedAssetNode) const
+UInterchangeActorFactoryNode* UInterchangeGenericLevelPipeline::CreateActorFactoryNode(const UInterchangeBaseNode* TranslatedAssetNode) const
 {
 	if (!ensure(BaseNodeContainer))
 	{
diff --git a/Engine/Plugins/Interchange/Runtime/Source/Pipelines/Private/InterchangeGenericSkeletalMeshPipeline.cpp b/Engine/Plugins/Interchange/Runtime/Source/Pipelines/Private/InterchangeGenericSkeletalMeshPipeline.cpp
--- a/Engine/Plugins/Interchange/Runtime/Source/Pipelines/Private/InterchangeGenericSkeletalMeshPipeline.cpp
+++ b/Engine/Plugins/Interchange/Runtime/Source/Pipelines/Private/InterchangeGenericSkeletalMeshPipeline.cpp
@@ -162,8 +162,6 @@ void UInterchangeGenericMeshPipeline::ExecutePreImportPipelineSkeletalMesh()
 					if (SkeletalMeshFactoryNode != nullptr)
 					{
 						SetSkeletalMeshDependencies(SkeletonRootUid, SkeletalMeshFactoryNode);
-						SkeletonFactoryNodes.Add(SkeletonFactoryNode);
-						SkeletalMeshFactoryNodes.Add(SkeletalMeshFactoryNode);
 						bFoundInstances = true;
 					}
 				}
@@ -223,8 +221,6 @@ void UInterchangeGenericMeshPipeline::ExecutePreImportPipelineSkeletalMesh()
 					if (SkeletalMeshFactoryNode != nullptr)
 					{
 						SetSkeletalMeshDependencies(SkeletonRootUid, SkeletalMeshFactoryNode);
-						SkeletonFactoryNodes.Add(SkeletonFactoryNode);
-						SkeletalMeshFactoryNodes.Add(SkeletalMeshFactoryNode);
 						bFoundInstances = true;
 					}
 				}
diff --git a/Engine/Plugins/Interchange/Runtime/Source/Pipelines/Private/InterchangeGenericTexturePipeline.cpp b/Engine/Plugins/Interchange/Runtime/Source/Pipelines/Private/InterchangeGenericTexturePipeline.cpp
--- a/Engine/Plugins/Interchange/Runtime/Source/Pipelines/Private/InterchangeGenericTexturePipeline.cpp
+++ b/Engine/Plugins/Interchange/Runtime/Source/Pipelines/Private/InterchangeGenericTexturePipeline.cpp
@@ -456,20 +456,7 @@ void UInterchangeGenericTexturePipeline::ExecutePipeline(UInterchangeBaseNodeCon
 	}
 	
 	//Find all translated node we need for this pipeline
-	BaseNodeContainer->IterateNodes([this](const FString& NodeUid, UInterchangeBaseNode* Node)
-	{
-		switch(Node->GetNodeContainerType())
-		{
-			case EInterchangeNodeContainerType::TranslatedAsset:
-			{
-				if (UInterchangeTextureNode* TextureNode = Cast<UInterchangeTextureNode>(Node))
-				{
-					TextureNodes.Add(TextureNode);
-				}
-			}
-			break;
-		}
-	});
+	BaseNodeContainer->GetNodesOfType<UInterchangeTextureNode>(TextureNodes);
 
 	if (bImportTextures)
 	{
diff --git a/Engine/Plugins/Interchange/Runtime/Source/Pipelines/Private/InterchangePipelineMeshesUtilities.cpp b/Engine/Plugins/Interchange/Runtime/Source/Pipelines/Private/InterchangePipelineMeshesUtilities.cpp
--- a/Engine/Plugins/Interchange/Runtime/Source/Pipelines/Private/InterchangePipelineMeshesUtilities.cpp
+++ b/Engine/Plugins/Interchange/Runtime/Source/Pipelines/Private/InterchangePipelineMeshesUtilities.cpp
@@ -27,52 +27,43 @@ namespace UE::Private::InterchangeMeshPipeline
 		//{ParentNodePtr, LODContainerName} => [{CandidateNode, CandidateLODIndex}..]
 		TMap<TPair<const UInterchangeSceneNode*, FString>, TArray<TPair<UInterchangeSceneNode*, int32>>> LODContainerCandidates;
 		const FString LodPrefix = TEXT("LOD");
-		BaseNodeContainer->IterateNodes(
-			[&BaseNodeContainer, &LodPrefix, &LODContainerCandidates](const FString& NodeUid, UInterchangeBaseNode* Node)
+		BaseNodeContainer->IterateNodesOfType<UInterchangeSceneNode>(
+			[&BaseNodeContainer, &LodPrefix, &LODContainerCandidates](const FString& NodeUid, UInterchangeSceneNode* SceneNode)
 			{
-				if (Node->GetNodeContainerType() == EInterchangeNodeContainerType::TranslatedScene)
+				if (!SceneNode->IsEnabled())
 				{
-					UInterchangeSceneNode* SceneNode = Cast<UInterchangeSceneNode>(Node);
-					if (!SceneNode)
-					{
-						return;
-					}
-
-					if (!SceneNode->IsEnabled())
-					{
-						return;
-					}
+					return;
+				}
 
-					FString AssetInstanceUid;
-					if (!SceneNode->GetCustomAssetInstanceUid(AssetInstanceUid))
-					{
-						return;
-					}
+				FString AssetInstanceUid;
+				if (!SceneNode->GetCustomAssetInstanceUid(AssetInstanceUid))
+				{
+					return;
+				}
 
-					if (const UInterchangeMeshLODContainerNode* MeshLODContainerNode = Cast<UInterchangeMeshLODContainerNode>(BaseNodeContainer->GetNode(AssetInstanceUid)))
-					{
-						return;
-					}
+				if (const UInterchangeMeshLODContainerNode* MeshLODContainerNode = Cast<UInterchangeMeshLODContainerNode>(BaseNodeContainer->GetNode(AssetInstanceUid)))
+				{
+					return;
+				}
 
-					FString SceneNodeName = SceneNode->GetDisplayLabel();
+				FString SceneNodeName = SceneNode->GetDisplayLabel();
 
-					int32 LODIndex;
-					FString LODContainerName;
-					FString Error;
-					if (UInterchangeMeshLODContainerNode::CheckForLODPattern(SceneNodeName, LODIndex, LODContainerName, Error))
-					{
-						FString ParentUid = SceneNode->GetParentUid();
-						const UInterchangeSceneNode* ConstParentSceneNode = Cast<UInterchangeSceneNode>(BaseNodeContainer->GetNode(ParentUid));
+				int32 LODIndex;
+				FString LODContainerName;
+				FString Error;
+				if (UInterchangeMeshLODContainerNode::CheckForLODPattern(SceneNodeName, LODIndex, LODContainerName, Error))
+				{
+					FString ParentUid = SceneNode->GetParentUid();
+					const UInterchangeSceneNode* ConstParentSceneNode = Cast<UInterchangeSceneNode>(BaseNodeContainer->GetNode(ParentUid));
 
-						LODContainerCandidates.FindOrAdd(TPair<const UInterchangeSceneNode*, FString>(ConstParentSceneNode, LODContainerName)).Add(TPair<UInterchangeSceneNode*, int32>(SceneNode, LODIndex));
-					}
-					else if (Error.Len())
+					LODContainerCandidates.FindOrAdd(TPair<const UInterchangeSceneNode*, FString>(ConstParentSceneNode, LODContainerName)).Add(TPair<UInterchangeSceneNode*, int32>(SceneNode, LODIndex));
+				}
+				else if (Error.Len())
+				{
+					if (UE::Interchange::Private::MeshHelper::GetMeshCollisionFromName(SceneNodeName) == EInterchangeMeshCollision::None)
 					{
-						if (UE::Interchange::Private::MeshHelper::GetMeshCollisionFromName(SceneNodeName) == EInterchangeMeshCollision::None)
-						{
-							//Do not warn for presumed Collision nodes: (for ep: UCX_SM_Test_LOD0)
-							UE_LOG(LogInterchangePipelineMeshesUtilities, Warning, TEXT("%s"), *Error);
-						}
+						//Do not warn for presumed Collision nodes: (for ep: UCX_SM_Test_LOD0)
+						UE_LOG(LogInterchangePipelineMeshesUtilities, Warning, TEXT("%s"), *Error);
 					}
 				}
 			}
@@ -648,38 +639,32 @@ UInterchangePipelineMeshesUtilities* UInterchangePipelineMeshesUtilities::Create
 
 	//Find all translated scene node we need for this pipeline
 	bool bHasSockets = false;
-	BaseNodeContainer->IterateNodes(
-		[&PipelineMeshesUtilities, &BaseNodeContainer, &SkeletonRootNodeUids, &bHasSockets, &ProcessMeshNode, &ProcessedMeshLODSceneNodes](const FString& NodeUid, const UInterchangeBaseNode* Node)
+	BaseNodeContainer->IterateNodesOfType<UInterchangeSceneNode>(
+		[&PipelineMeshesUtilities, &BaseNodeContainer, &SkeletonRootNodeUids, &bHasSockets, &ProcessMeshNode, &ProcessedMeshLODSceneNodes](const FString& NodeUid, const UInterchangeSceneNode* SceneNode)
 		{
-			if (Node->GetNodeContainerType() == EInterchangeNodeContainerType::TranslatedScene)
+			if (ProcessedMeshLODSceneNodes.Contains(NodeUid))
 			{
-				if (const UInterchangeSceneNode* SceneNode = Cast<UInterchangeSceneNode>(Node))
-				{
-					if (ProcessedMeshLODSceneNodes.Contains(NodeUid))
-					{
-						return;
-					}
+				return;
+			}
 
-					if (SceneNode->IsA(UInterchangeJointNode::StaticClass()))
-					{
-						const UInterchangeSceneNode* ParentJointNode = Cast<UInterchangeSceneNode>(BaseNodeContainer->GetNode(SceneNode->GetParentUid()));
-						if (!ParentJointNode || !ParentJointNode->IsA(UInterchangeJointNode::StaticClass()))
-						{
-							SkeletonRootNodeUids.Add(SceneNode->GetUniqueID());
-						}
-					}
+			if (SceneNode->IsA(UInterchangeJointNode::StaticClass()))
+			{
+				const UInterchangeSceneNode* ParentJointNode = Cast<UInterchangeSceneNode>(BaseNodeContainer->GetNode(SceneNode->GetParentUid()));
+				if (!ParentJointNode || !ParentJointNode->IsA(UInterchangeJointNode::StaticClass()))
+				{
+					SkeletonRootNodeUids.Add(SceneNode->GetUniqueID());
+				}
+			}
 
-					if (IsSceneNodeASocket(SceneNode))
-					{
-						bHasSockets = true;
-					}
+			if (IsSceneNodeASocket(SceneNode))
+			{
+				bHasSockets = true;
+			}
 
-					FString AssetInstanceUid;
-					if (SceneNode->GetCustomAssetInstanceUid(AssetInstanceUid))
-					{
-						ProcessMeshNode(NodeUid, SceneNode, AssetInstanceUid);
-					}
-				}
+			FString AssetInstanceUid;
+			if (SceneNode->GetCustomAssetInstanceUid(AssetInstanceUid))
+			{
+				ProcessMeshNode(NodeUid, SceneNode, AssetInstanceUid);
 			}
 		}
 	);
@@ -708,34 +693,31 @@ UInterchangePipelineMeshesUtilities* UInterchangePipelineMeshesUtilities::Create
 		else
 		{
 			//Import of Local Sockets:
-			BaseNodeContainer->IterateNodes(
-				[&PipelineMeshesUtilities, &BaseNodeContainer](const FString& NodeUid, const UInterchangeBaseNode* Node)
+			BaseNodeContainer->IterateNodesOfType<UInterchangeSceneNode>(
+				[&PipelineMeshesUtilities, &BaseNodeContainer](const FString& NodeUid, const UInterchangeSceneNode* SceneNode)
 				{
-					if (const UInterchangeSceneNode* SceneNode = Cast<UInterchangeSceneNode>(Node))
+					if (IsSceneNodeASocket(SceneNode))
 					{
-						if (IsSceneNodeASocket(SceneNode))
+						FString MeshUid;
+						if (!SceneNode->GetCustomAssetInstanceUid(MeshUid))
 						{
-							FString MeshUid;
-							if (!SceneNode->GetCustomAssetInstanceUid(MeshUid))
+							const UInterchangeSceneNode* ParentMeshSceneNode = Cast<UInterchangeSceneNode>(BaseNodeContainer->GetNode(SceneNode->GetParentUid()));
+							while (ParentMeshSceneNode)
 							{
-								const UInterchangeSceneNode* ParentMeshSceneNode = Cast<UInterchangeSceneNode>(BaseNodeContainer->GetNode(SceneNode->GetParentUid()));
-								while (ParentMeshSceneNode)
+								if (ParentMeshSceneNode->GetCustomAssetInstanceUid(MeshUid))
 								{
-									if (ParentMeshSceneNode->GetCustomAssetInstanceUid(MeshUid))
-									{
-										break;
-									}
-
-									ParentMeshSceneNode = Cast<UInterchangeSceneNode>(BaseNodeContainer->GetNode(ParentMeshSceneNode->GetParentUid()));
+									break;
 								}
-							}
 
-							if (!MeshUid.IsEmpty())
-							{
-								FInterchangeMeshGeometry& MeshGeometry = PipelineMeshesUtilities->MeshGeometriesPerMeshUid.FindChecked(MeshUid);
-								MeshGeometry.AttachedSocketUids.Add(SceneNode->GetUniqueID());
+								ParentMeshSceneNode = Cast<UInterchangeSceneNode>(BaseNodeContainer->GetNode(ParentMeshSceneNode->GetParentUid()));
 							}
 						}
+
+						if (!MeshUid.IsEmpty())
+						{
+							FInterchangeMeshGeometry& MeshGeometry = PipelineMeshesUtilities->MeshGeometriesPerMeshUid.FindChecked(MeshUid);
+							MeshGeometry.AttachedSocketUids.Add(SceneNode->GetUniqueID());
+						}
 					}
 				}
 			);
@@ -812,30 +794,19 @@ void UInterchangePipelineMeshesUtilities::IterateAllSkinnedMeshInstance(TFunctio
 
 void UInterchangePipelineMeshesUtilities::GetAllStaticMeshInstance(TArray<FString>& MeshInstanceUids, const bool bAllowNotVisibles) const
 {
-	if (bAllowNotVisibles)
+	MeshInstanceUids.Reserve(MeshInstancesPerMeshInstanceUid.Num());
+	for (const TPair<FString, FInterchangeMeshInstance>& MeshInstanceUidAndMeshInstance : MeshInstancesPerMeshInstanceUid)
 	{
-		MeshInstanceUids.Reserve(MeshInstancesPerMeshInstanceUid.Num());
-		for (const TPair<FString, FInterchangeMeshInstance>& MeshInstanceUidAndMeshInstance : MeshInstancesPerMeshInstanceUid)
+		const FInterchangeMeshInstance& MeshInstance = MeshInstanceUidAndMeshInstance.Value;
+
+		if (bAllowNotVisibles || MeshInstance.bIsVisible)
 		{
-			const FInterchangeMeshInstance& MeshInstance = MeshInstanceUidAndMeshInstance.Value;
 			if (CurrentDataContext.IsStaticMeshInstance(MeshInstance, BaseNodeContainer))
 			{
 				MeshInstanceUids.Add(MeshInstanceUidAndMeshInstance.Key);
 			}
 		}
 	}
-	else
-	{
-		MeshInstanceUids.Reserve(MeshInstancesPerMeshInstanceUid.Num());
-		for (const TPair<FString, FInterchangeMeshInstance>& MeshInstanceUidAndMeshInstance : MeshInstancesPerMeshInstanceUid)
-		{
-			const FInterchangeMeshInstance& MeshInstance = MeshInstanceUidAndMeshInstance.Value;
-			if (MeshInstance.bIsVisible && CurrentDataContext.IsStaticMeshInstance(MeshInstance, BaseNodeContainer))
-			{
-				MeshInstanceUids.Add(MeshInstanceUidAndMeshInstance.Key);
-			}
-		}
-	}
 }
 
 void UInterchangePipelineMeshesUtilities::IterateAllStaticMeshInstance(TFunctionRef<void(const FInterchangeMeshInstance&)> IterationLambda) const
diff --git a/Engine/Plugins/Interchange/Runtime/Source/Pipelines/Private/InterchangeSparseVolumeTexturePipeline.cpp b/Engine/Plugins/Interchange/Runtime/Source/Pipelines/Private/InterchangeSparseVolumeTexturePipeline.cpp
--- a/Engine/Plugins/Interchange/Runtime/Source/Pipelines/Private/InterchangeSparseVolumeTexturePipeline.cpp
+++ b/Engine/Plugins/Interchange/Runtime/Source/Pipelines/Private/InterchangeSparseVolumeTexturePipeline.cpp
@@ -349,15 +349,7 @@ void UInterchangeSparseVolumeTexturePipeline::ExecutePipeline(
 
 	// Find all the translated nodes we need for this pipeline
 	TArray<UInterchangeVolumeNode*> VolumeNodes;
-	BaseNodeContainer->IterateNodes(
-		[&VolumeNodes](const FString& NodeUid, UInterchangeBaseNode* Node)
-		{
-			if (UInterchangeVolumeNode* TextureNode = Cast<UInterchangeVolumeNode>(Node))
-			{
-				VolumeNodes.Add(TextureNode);
-			}
-		}
-	);
+	BaseNodeContainer->GetNodesOfType<UInterchangeVolumeNode>(VolumeNodes);
 
 	TArray<UInterchangeSparseVolumeTextureFactoryNode*> CreatedFactoryNodes;
 
diff --git a/Engine/Plugins/Interchange/Runtime/Source/Pipelines/Public/InterchangeGenericMeshPipeline.h b/Engine/Plugins/Interchange/Runtime/Source/Pipelines/Public/InterchangeGenericMeshPipeline.h
--- a/Engine/Plugins/Interchange/Runtime/Source/Pipelines/Public/InterchangeGenericMeshPipeline.h
+++ b/Engine/Plugins/Interchange/Runtime/Source/Pipelines/Public/InterchangeGenericMeshPipeline.h
@@ -393,12 +393,6 @@ private:
 	 */
 	UE_API void ExecutePreImportPipelineSkeletalMesh();
 
-	/** Skeleton factory assets nodes */
-	TArray<UInterchangeSkeletonFactoryNode*> SkeletonFactoryNodes;
-
-	/** Skeletal mesh factory assets nodes */
-	TArray<UInterchangeSkeletalMeshFactoryNode*> SkeletalMeshFactoryNodes;
-
 	/*
 	* This function will create the SkeletalMeshFactoryNode and all necessary UInterchangeSkeletalMeshLodDataNode nodes.
 	*	It will also process/parse all input nodes into the created UInterchangeSkeletalMeshLodDataNode nodes, and will also generate the Transforms that should be used for Mesh acquisitions.
diff --git a/Engine/Plugins/Interchange/Runtime/Source/Pipelines/Public/InterchangeGenericScenesPipeline.h b/Engine/Plugins/Interchange/Runtime/Source/Pipelines/Public/InterchangeGenericScenesPipeline.h
--- a/Engine/Plugins/Interchange/Runtime/Source/Pipelines/Public/InterchangeGenericScenesPipeline.h
+++ b/Engine/Plugins/Interchange/Runtime/Source/Pipelines/Public/InterchangeGenericScenesPipeline.h
@@ -101,7 +101,13 @@ protected:
 	/**
 	 * Return a new Actor Factory Node to be used for the given SceneNode.
 	 */
-	UE_API virtual UInterchangeActorFactoryNode* CreateActorFactoryNode(const UInterchangeSceneNode* SceneNode, const UInterchangeBaseNode* TranslatedAssetNode) const;
+	UE_API virtual UInterchangeActorFactoryNode* CreateActorFactoryNode(const UInterchangeBaseNode* TranslatedAssetNode) const;
+
+	UE_DEPRECATED(5.8, "Use CreateActorFactoryNode(const UInterchangeBaseNode* TranslatedAssetNode) instead")
+	virtual UInterchangeActorFactoryNode* CreateActorFactoryNode(const UInterchangeSceneNode* SceneNode, const UInterchangeBaseNode* TranslatedAssetNode) const
+	{
+		return CreateActorFactoryNode(TranslatedAssetNode);
+	}
 
 	/**
 	 * Return a new Component Factory Node.
@@ -154,7 +160,7 @@ protected:
 #endif
 
 	UE_API void CacheActiveJointUids();
-	TArray<FString> Cached_ActiveJointUids;
+	TSet<FString> Cached_ActiveJointUids;
 
 	FPreviewScene* PreviewScene = nullptr;
 };
diff --git a/Engine/Source/Runtime/Interchange/Core/Private/Nodes/InterchangeBaseNode.cpp b/Engine/Source/Runtime/Interchange/Core/Private/Nodes/InterchangeBaseNode.cpp
--- a/Engine/Source/Runtime/Interchange/Core/Private/Nodes/InterchangeBaseNode.cpp
+++ b/Engine/Source/Runtime/Interchange/Core/Private/Nodes/InterchangeBaseNode.cpp
@@ -8,6 +8,8 @@
 #include "UObject/Object.h"
 #include "UObject/ObjectMacros.h"
 #include "InterchangeHelper.h"
+#include "InterchangeLogPrivate.h"
+
 
 #include UE_INLINE_GENERATED_CPP_BY_NAME(InterchangeBaseNode)
 
@@ -87,6 +89,21 @@ UInterchangeBaseNode::UInterchangeBaseNode()
 	RegisterAttribute<uint8>(UE::Interchange::FBaseNodeStaticData::NodeContainerTypeKey(), static_cast<uint8>(EInterchangeNodeContainerType::None));
 }
 
+#if UE_BUILD_DEBUG || UE_BUILD_DEVELOPMENT
+static void UpdateNameForDebugging(UInterchangeBaseNode* Node)
+{
+	//For debugging purpose
+	const FString Name = Node->GetDisplayLabel() + " | " + Node->GetUniqueID();
+	if (!Node->Rename(*Name, nullptr, REN_Test))
+	{
+		UE_LOG(LogInterchangeCore, Error, TEXT("Duplicate NodeUid Found `%s`. This might indicates a potential issue in how this node was added without checking for its existence first."), *Name);
+		return;
+	}
+
+	Node->Rename(*Name, nullptr, REN_DontCreateRedirectors | REN_NonTransactional);
+}
+#endif
+
 void UInterchangeBaseNode::InitializeNode(const FString& UniqueID, const FString& DisplayLabel, const EInterchangeNodeContainerType NodeContainerType)
 {
 	UE::Interchange::EAttributeStorageResult Result = Attributes->RegisterAttribute(UE::Interchange::FBaseNodeStaticData::UniqueIDKey(), UniqueID, UE::Interchange::EAttributeProperty::NoHash);
@@ -107,6 +124,10 @@ void UInterchangeBaseNode::InitializeNode(const FString& UniqueID, const FString
 	}
 
 	bIsInitialized = true;
+
+#if UE_BUILD_DEBUG || UE_BUILD_DEVELOPMENT
+	UpdateNameForDebugging(this);
+#endif
 }
 
 FString UInterchangeBaseNode::GetTypeName() const
@@ -344,13 +365,23 @@ bool UInterchangeBaseNode::SetDisplayLabel(const FString& InDisplayLabel)
 	//This can still be bypassed by using the Attributes system directly.
 	FString Sanitized = UE::Interchange::MakeSanitizedName(InDisplayLabel);
 
+	bool bSuccess = false;
+
 	UE::Interchange::EAttributeStorageResult Result = Attributes->RegisterAttribute(UE::Interchange::FBaseNodeStaticData::DisplayLabelKey(), Sanitized);
 	if (IsAttributeStorageResultSuccess(Result))
 	{
 		UE::Interchange::FAttributeStorage::TAttributeHandle<FString> Handle = Attributes->GetAttributeHandle<FString>(UE::Interchange::FBaseNodeStaticData::DisplayLabelKey());
-		return Handle.IsValid();
+		bSuccess = Handle.IsValid();
 	}
-	return false;
+
+#if UE_BUILD_DEBUG || UE_BUILD_DEVELOPMENT
+	if (bIsInitialized && bSuccess)
+	{
+		UpdateNameForDebugging(this);
+	}
+#endif
+
+	return bSuccess;
 }
 
 bool UInterchangeBaseNode::GetNamespace(FString& OutNamespace) const
diff --git a/Engine/Source/Runtime/Interchange/Core/Private/Nodes/InterchangeBaseNodeContainer.cpp b/Engine/Source/Runtime/Interchange/Core/Private/Nodes/InterchangeBaseNodeContainer.cpp
--- a/Engine/Source/Runtime/Interchange/Core/Private/Nodes/InterchangeBaseNodeContainer.cpp
+++ b/Engine/Source/Runtime/Interchange/Core/Private/Nodes/InterchangeBaseNodeContainer.cpp
@@ -44,14 +44,33 @@ FString UInterchangeBaseNodeContainer::AddNode(UInterchangeBaseNode* Node)
 	return NodeUniqueID;
 }
 
-void UInterchangeBaseNodeContainer::ReplaceNode(const FString& NodeUniqueID, UInterchangeFactoryBaseNode* NewNode)
+bool UInterchangeBaseNodeContainer::RemoveFactoryNode(const FString& NodeUniqueID)
 {
-	TRACE_CPUPROFILER_EVENT_SCOPE(UInterchangeBaseNodeContainer::ReplaceNode)
-	if (GetFactoryNode(NodeUniqueID)) //Check existance and confirm it is FactoryNode
+	TRACE_CPUPROFILER_EVENT_SCOPE(UInterchangeBaseNodeContainer::RemoveFactoryNode)
+	if (UInterchangeFactoryBaseNode* OldNode = GetFactoryNode(NodeUniqueID))
 	{
 		Nodes.Remove(NodeUniqueID);
+
+		OldNode->ClearFlags(RF_Standalone);
+		OldNode->ClearInternalFlags(EInternalObjectFlags::Async);
+		OldNode->Rename(nullptr, GetTransientPackage(), REN_DontCreateRedirectors | REN_NonTransactional);
+
+		return true;
+	}
+
+	return false;
+}
+
+bool UInterchangeBaseNodeContainer::ReplaceFactoryNode(const FString& NodeUniqueID, UInterchangeFactoryBaseNode* NewNode)
+{
+	TRACE_CPUPROFILER_EVENT_SCOPE(UInterchangeBaseNodeContainer::ReplaceNode)
+	if (RemoveFactoryNode(NodeUniqueID)) //Check existence and confirm it is FactoryNode
+	{
 		AddNode(NewNode);
+		return true;
 	}
+
+	return false;
 }
 
 bool UInterchangeBaseNodeContainer::IsNodeUidValid(const FString& NodeUniqueID) const
@@ -536,12 +555,8 @@ void UInterchangeBaseNodeContainer::SetupAndReplaceFactoryNode(UInterchangeFacto
 	const FString& OldNodeUID,
 	const FString& ParentNodeUID)
 {
-	Node->InitializeNode(NodeUID, DisplayLabel, ContainerType);
-
-	ReplaceNode(OldNodeUID, Node);
-
-	if (ParentNodeUID.Len())
+	if (RemoveFactoryNode(OldNodeUID))
 	{
-		SetNodeParentUid(NodeUID, ParentNodeUID);
+		SetupNode(Node, NodeUID, DisplayLabel, ContainerType, ParentNodeUID);
 	}
 }
diff --git a/Engine/Source/Runtime/Interchange/Core/Public/Nodes/InterchangeBaseNodeContainer.h b/Engine/Source/Runtime/Interchange/Core/Public/Nodes/InterchangeBaseNodeContainer.h
--- a/Engine/Source/Runtime/Interchange/Core/Public/Nodes/InterchangeBaseNodeContainer.h
+++ b/Engine/Source/Runtime/Interchange/Core/Public/Nodes/InterchangeBaseNodeContainer.h
@@ -71,7 +71,16 @@ public:
 	INTERCHANGECORE_API FString AddNode(UInterchangeBaseNode* Node);
 
 	UFUNCTION(BlueprintCallable, Category = "Interchange | Node Container")
-	INTERCHANGECORE_API void ReplaceNode(const FString& NodeUniqueID, UInterchangeFactoryBaseNode* NewNode);
+	INTERCHANGECORE_API bool ReplaceFactoryNode(const FString& NodeUniqueID, UInterchangeFactoryBaseNode* NewNode);
+
+	UE_DEPRECATED(5.8, "Use ReplaceFactoryNode instead.")
+	void ReplaceNode(const FString& NodeUniqueID, UInterchangeFactoryBaseNode* NewNode)
+	{
+		ReplaceFactoryNode(NodeUniqueID, NewNode);
+	}
+
+	UFUNCTION(BlueprintCallable, Category = "Interchange | Node Container")
+	INTERCHANGECORE_API bool RemoveFactoryNode(const FString& NodeUniqueID);
 
 	/** Return true if the node unique ID exists in the container. */
 	UFUNCTION(BlueprintCallable, Category = "Interchange | Node Container")
@@ -85,7 +94,7 @@ public:
 	 * @Param Namespace - Is the new namespace you want to set. Pass an empty string to remove an existing namespace.
 	 * @Param TargetClass - Optional, this parameter represent the node class we want to apply the namespace on. If null all node will be tagged with the namespace
 	 * 
-	 * @Note - Changing all node namespace wont work since some node class use a combinaison of the unique ID in there attributes and we cannot change those attribute
+	 * @Note - Changing all node namespace wont work since some node class use a combination of the unique ID in there attributes and we cannot change those attribute
 	 *         in a generic way.
 	 */
 	UFUNCTION(BlueprintCallable, Category = "Interchange | Node Container")
@@ -106,6 +115,20 @@ public:
 		}
 	}
 
+	template <typename T>
+	void GetNodesOfType(TArray<T*>& OutNodes) const
+	{
+		OutNodes.Empty();
+
+		for (const TPair<FString, TObjectPtr<UInterchangeBaseNode>>& NodeKeyValue : Nodes)
+		{
+			if (T* Node = Cast<T>(NodeKeyValue.Value))
+			{
+				OutNodes.Add(Node);
+			}
+		}
+	}
+
 	template <typename T>
 	void GetNodeUIDsOfType(TArray<FString>& OutNodeUIDs) const
 	{
@@ -124,7 +147,6 @@ public:
 	 */
 	void IterateNodeChildren(const FString& NodeUniqueID, TFunctionRef<void(const UInterchangeBaseNode*)> IterationLambda) const
 	{
-
 		if (const UInterchangeBaseNode* Node = GetNode(NodeUniqueID))
 		{
 			IterationLambda(Node);
