Fixed issue withh global transform

#rnx

#Synced-CL 51270901

diff --git a/Engine/Plugins/Interchange/Runtime/Source/Nodes/Private/InterchangeSceneNode.cpp b/Engine/Plugins/Interchange/Runtime/Source/Nodes/Private/InterchangeSceneNode.cpp
--- a/Engine/Plugins/Interchange/Runtime/Source/Nodes/Private/InterchangeSceneNode.cpp
+++ b/Engine/Plugins/Interchange/Runtime/Source/Nodes/Private/InterchangeSceneNode.cpp
@@ -91,6 +91,12 @@ UInterchangeSceneNode::UInterchangeSceneNode()
 	Tags.Initialize(Attributes, UE::Interchange::FSceneNodeStaticData::GetTagsKey());
 	CurveAnimationTypes.Initialize(Attributes.ToSharedRef(), UE::Interchange::FSceneNodeStaticData::GetCurveAnimationTypesKey());
 	ComponentUids.Initialize(Attributes.ToSharedRef(), UE::Interchange::FSceneNodeStaticData::GetComponentUidsKey());
+
+	// The handling of the global and local transforms of a scene node should be revisited.
+	// Currently if a translator does not set a local transform on a scene node some logic will fail.
+	// See UE-351819 aas an example.
+	// No logic should fail if the local transform was not explicitly set because it should implicitly be the identity. 
+	InterchangePrivateNodeBase::SetCustomAttribute<FTransform>(*Attributes, Macro_CustomLocalTransformKey, TEXT("SceneNode.SetLocalTransform"), FTransform::Identity);
 }
 
 /**
