[Interchange] Transform generation fix for StaticMeshes and Skeleton+RigidMesh import behavior fixeds.
Re-aligned Post LOD Refactor transform generation for StaticMeshes and SkeletalMeshes forced to StaticMeshes.
Added ability to Asset import Skeletons with Non-SkinnedMeshes.
Adjusted Skeleton/SkeletalMesh behavior regarding Non-SkinnedMeshes.
(Any Skeleton that has any Mesh referenced should now create a SkeletalMesh+Skeleton, regardless of Mesh being Skinned or not.)
Also fixes regression for FBX SkinnedMeshes refrencing Nodes without Bone attributes.

Broken example file that is fixed with given CL can be found in UE-194244.

#rnx

#virtualized
#Synced-CL 51299505

#Synced-CL 51299505

diff --git a/Engine/Plugins/Interchange/Runtime/Source/Import/Private/Mesh/InterchangeMeshHelper.cpp b/Engine/Plugins/Interchange/Runtime/Source/Import/Private/Mesh/InterchangeMeshHelper.cpp
--- a/Engine/Plugins/Interchange/Runtime/Source/Import/Private/Mesh/InterchangeMeshHelper.cpp
+++ b/Engine/Plugins/Interchange/Runtime/Source/Import/Private/Mesh/InterchangeMeshHelper.cpp
@@ -896,6 +896,47 @@ namespace UE::Interchange::Private::MeshHelper
 		GlobalTransform = SceneNodeGeometricTransform * GlobalTransform;
 	}
 
+	//Resolves the optional RootJointParentGlobalInverseMatrix, or in case it's not set, 
+	// it traverses the parent hierarchy of the SceneNode until it finds the RootJoint's Parent (RootJoint's Parent, is the first SceneNode that is not a JointNode) and acquire the GlobalInverseMatrix from found node.
+	FMatrix ResolveRootJointParentGlobalInverseMatrix(UInterchangeBaseNodeContainer& BaseNodeContainer, const UInterchangeSceneNode* SceneNode, const FTransform& GlobalOffsetTransform, const TOptional<TPair<FString, FMatrix>>& RootJointParentUidAndGlobalInverseMatrixOptional)
+	{
+		FMatrix RootJointParentGlobalInverseMatrix = FMatrix::Identity;
+		if (RootJointParentUidAndGlobalInverseMatrixOptional.IsSet())
+		{
+			RootJointParentGlobalInverseMatrix = RootJointParentUidAndGlobalInverseMatrixOptional.GetValue().Value;
+		}
+		else
+		{
+			const UInterchangeSceneNode* Current = SceneNode;
+			const UInterchangeSceneNode* Candidate = Cast<UInterchangeSceneNode>(BaseNodeContainer.GetNode(Current->GetParentUid()));
+
+			//In this case we will try to identify the root joint node via iterating through Parents until we hit the top or a non-UInterchangeJointNode.
+			static const uint32 MaxDepthCheckLimit = UINT16_MAX;
+			uint32 MaxDepthChecker = 0;
+			while (Candidate && Candidate->IsA<UInterchangeJointNode>())
+			{
+				Current = Candidate;
+				Candidate = Cast<UInterchangeSceneNode>(BaseNodeContainer.GetNode(Current->GetParentUid()));
+				MaxDepthChecker++;
+				if (MaxDepthChecker > MaxDepthCheckLimit)
+				{
+					Candidate = nullptr;
+					break;
+				}
+			}
+
+			if (Candidate)
+			{
+				FTransform RootJointParentGlobalTransform = FTransform::Identity;
+				Candidate->GetCustomGlobalTransform(&BaseNodeContainer, GlobalOffsetTransform, RootJointParentGlobalTransform);
+
+				RootJointParentGlobalInverseMatrix = RootJointParentGlobalTransform.ToMatrixWithScale().Inverse();
+			}
+		}
+
+		return RootJointParentGlobalInverseMatrix;
+	}
+
 	/**
 	* When we are trying to acquire the wanted/needed transform for a MeshNode we have to combine/calculate with the following transforms:
 	*		- GlobalTransform: only if bBakeMeshes is true. (Global transform of the instantiating SceneNode)
@@ -913,36 +954,121 @@ namespace UE::Interchange::Private::MeshHelper
 		const bool bBakeMeshes, const bool bBakePivotMeshes, 
 		const FTransform& GlobalOffsetTransform,
 		const bool bUseTimeZeroAsBindPose,
-		const FMatrix& RootJointParentGlobalInverseMatrix)
+		const TOptional<TPair<FString, FMatrix>>& RootJointParentUidAndGlobalInverseMatrixOptional,
+		const bool bProcessingSkeletalMesh,
+		const bool bAlignSkeletalMeshInScene)
 	{
-		if (!BaseNodeContainer)
+		if (!BaseNodeContainer || !SceneNode)
 		{
-			return FTransform();
+			return FTransform::Identity;
 		}
 
-		FTransform OutGlobalTransform;
-		if (bBakeMeshes)
+		FTransform OutGlobalTransform = FTransform::Identity;
+		
+		if (const UInterchangeJointNode* JointNode = Cast<UInterchangeJointNode>(SceneNode))
+		{
+			FTransform JointNodeSkeletonTransform = FTransform::Identity;
+
+			if (!bUseTimeZeroAsBindPose || !JointNode->GetTimeZeroGlobalTransform(BaseNodeContainer, GlobalOffsetTransform, JointNodeSkeletonTransform))
+			{
+				//Even if BindPose transform is not present it will fallback to Local Transform for the GlobalTransform generation
+				//in case that fails, there is no other fallback mechanism to be applied.
+				//Hence the reason why we can ignore the success check returned from GetBindPoseGlobalTransform.
+				JointNode->GetBindPoseGlobalTransform(BaseNodeContainer, GlobalOffsetTransform, JointNodeSkeletonTransform);
+			}
+
+			if (bBakeMeshes)
+			{
+				OutGlobalTransform = JointNodeSkeletonTransform;
+			}
+			else
+			{
+				FMatrix RootJointParentGlobalInverseMatrix = ResolveRootJointParentGlobalInverseMatrix(*BaseNodeContainer, SceneNode, GlobalOffsetTransform, RootJointParentUidAndGlobalInverseMatrixOptional);
+
+				//We calculate the 'Skeleton Transform' here:
+				// In other words the Global/Relative transform to the RootJointNode's Parent.
+
+				FMatrix JointNodeGlobalMatrix = JointNodeSkeletonTransform.ToMatrixWithScale();
+				FTransform SkeletonTransform(JointNodeGlobalMatrix * RootJointParentGlobalInverseMatrix);
+
+				OutGlobalTransform = SkeletonTransform;
+			}
+		}
+		else if (bBakeMeshes)
 		{
 			FTransform SceneNodeGlobalTransform;
 			if (SceneNode->GetCustomGlobalTransform(BaseNodeContainer, GlobalOffsetTransform, SceneNodeGlobalTransform))
 			{
 				OutGlobalTransform = SceneNodeGlobalTransform;
 			}
+
+			if (!bAlignSkeletalMeshInScene && 
+				bProcessingSkeletalMesh && 
+				RootJointParentUidAndGlobalInverseMatrixOptional.IsSet())
+			{
+				//**bAlignSkeletalMeshInScene==FALSE**
+				//If the SK and SKM was aligned prior to the Scene Composition,
+				//in that case we have to make sure we are baking the final SKM Instantiating Node's transform.
+				// Which is the SkeletonRoot's Parent Global transform.
+				//This is already part of the Mesh's SceneNode's GlobalTransform, IF the SceneNode is a child of the RootJoint's ParentNode,
+				if (!BaseNodeContainer->GetIsAncestor(SceneNode->GetUniqueID(), RootJointParentUidAndGlobalInverseMatrixOptional.GetValue().Key))
+				{
+					//But if the SceneNode is not part of the RootJointParent children tree (aka GetIsAncestor returns false), 
+					// then we need to apply the RootJointParentGlobalMatrix:
+					//So for example, this scenario:
+					//		SceneRoot
+					//			RootJointParent
+					//				RootJoint (aka Skeleton)
+					//					BoneXYZ
+					//					...
+					//			TransformNodeABC
+					//				Mesh (Mesh referenced)
+					//	As the bAlignSkeletalMeshInScene==false, which means the alignment is happening outside of the Scene Composition (prior to it).
+					//
+					//In this case the RootJointParent's Global Matrix needs to be additionally baked.
+					// Test file for this: glTF 2.0 example file RiggedFigure.
+					FMatrix RootJointParentGlobalMatrix = ResolveRootJointParentGlobalInverseMatrix(*BaseNodeContainer, SceneNode, GlobalOffsetTransform, RootJointParentUidAndGlobalInverseMatrixOptional).Inverse();
+					FTransform SkeletonTransform(OutGlobalTransform.ToMatrixWithScale() * RootJointParentGlobalMatrix);
+
+					OutGlobalTransform = SkeletonTransform;
+				}
+			}
 		}
-		else if (const UInterchangeJointNode* JointNode = Cast<UInterchangeJointNode>(SceneNode))
-		{
-			//We calculate the 'Skeleton Transform' here:
-			// In other words the Global/Relative transform to the RootJointNode's Parent.
+		else if (bAlignSkeletalMeshInScene && bProcessingSkeletalMesh && RootJointParentUidAndGlobalInverseMatrixOptional.IsSet())
+		{
+			//**bAlignSkeletalMeshInScene==TRUE**
+			//If it is expected that SKM and SK is aligned in the Scene (via the Scene Composition).
+			//Then we have to bake the complete transform for the mesh (SceneNode), regardless of the bBakeMesh setting.
+			// However we do have to remove the final SKM Instantiating Node's global transform 
+			// (aka the RootJointParentGlobalTransfrom needs to be inverse applied, so that we keep SKM and SK allignment.)
+			// So for example, for this scenario:
+			//		SceneRoot
+			//			RootJointParent
+			//				RootJoint (aka Skeleton)
+			//					BoneXYZ
+			//					...
+			//				TransformNodeABC
+			//					Mesh (Mesh referenced)
+			//	We have to remove the RootJointParent Global Transform in case bBakeMesh==false (which is this else if branch) 
+			// And we also have to remove the RootJointParent Global Transform (aka apply the RootJointParentGlobalTransform.Inverse())
+			//		for this scenario:
+			//			SceneRoot
+			//				RootJointParent
+			//					RootJoint (aka Skeleton)
+			//						BoneXYZ
+			//						...
+			//				TransformNodeABC
+			//					Mesh (Mesh referenced)
 
-			FTransform JointNodeSkeletonTransform;
 
-			if (!bUseTimeZeroAsBindPose || !JointNode->GetTimeZeroGlobalTransform(BaseNodeContainer, GlobalOffsetTransform, JointNodeSkeletonTransform))
+			FTransform SceneNodeGlobalTransform;
+			if (SceneNode->GetCustomGlobalTransform(BaseNodeContainer, GlobalOffsetTransform, SceneNodeGlobalTransform))
 			{
-				JointNode->GetBindPoseGlobalTransform(BaseNodeContainer, GlobalOffsetTransform, JointNodeSkeletonTransform);
+				OutGlobalTransform = SceneNodeGlobalTransform;
 			}
 
-			FMatrix JointNodeGlobalMatrix = JointNodeSkeletonTransform.ToMatrixWithScale();
-			FTransform SkeletonTransform(JointNodeGlobalMatrix * RootJointParentGlobalInverseMatrix);
+			FMatrix RootJointParentGlobalInverseMatrix = ResolveRootJointParentGlobalInverseMatrix(*BaseNodeContainer, SceneNode, GlobalOffsetTransform, RootJointParentUidAndGlobalInverseMatrixOptional);
+			FTransform SkeletonTransform(OutGlobalTransform.ToMatrixWithScale() * RootJointParentGlobalInverseMatrix);
 
 			OutGlobalTransform = SkeletonTransform;
 		}
diff --git a/Engine/Plugins/Interchange/Runtime/Source/Import/Public/Mesh/InterchangeMeshHelper.h b/Engine/Plugins/Interchange/Runtime/Source/Import/Public/Mesh/InterchangeMeshHelper.h
--- a/Engine/Plugins/Interchange/Runtime/Source/Import/Public/Mesh/InterchangeMeshHelper.h
+++ b/Engine/Plugins/Interchange/Runtime/Source/Import/Public/Mesh/InterchangeMeshHelper.h
@@ -106,13 +106,20 @@ namespace UE::Interchange::Private::MeshHelper
 	/* return the result of the global transform with the geometric and pivot transform of the scene node. */
 	INTERCHANGEIMPORT_API void AddSceneNodeGeometricAndPivotToGlobalTransform(FTransform& GlobalTransform, const UInterchangeSceneNode* SceneNode, const bool bBakeMeshes, const bool bBakePivotMeshes);
 
-	/* In case the input SceneNode is a JointNode we are going to calculate the SkeletonTransform of the given Node (for bBakeMeshes=false)*/
+	/**
+	* In case the input SceneNode is a JointNode we are going to calculate the SkeletonTransform of the given Node (for bBakeMeshes=false).
+	* RootJointParentUidAndGlobalInverseMatrixOptional - Specifically used for SkeletalMesh's transform generations.
+	* bProcessingSkeletalMesh - Changes transform generation behavior for Non Joint SceneNode transforms 
+	* 							(we always bake to RootJoint or SceneRoot for SkinnedMeshes).
+	*/
 	INTERCHANGEIMPORT_API FTransform CalculateCompleteSceneNodeTransform(
 		UInterchangeBaseNodeContainer* BaseNodeContainer, const UInterchangeSceneNode* SceneNode,
 		const bool bBakeMeshes, const bool bBakePivotMeshes,
 		const FTransform& GlobalOffsetTransform,
 		const bool bUseTimeZeroAsBindPose = false,
-		const FMatrix& RootJointParentGlobalInverseMatrix = FMatrix::Identity);
+		const TOptional<TPair<FString, FMatrix>>& RootJointParentUidAndGlobalInverseMatrixOptional = TOptional<TPair<FString, FMatrix>>(),
+		const bool bProcessingSkeletalMesh = false,
+		const bool bAlignSkeletalMeshInScene = false);
 
 	INTERCHANGEIMPORT_API void SkeletalMeshFactorySetupAssetMaterialArray(TArray<FSkeletalMaterial>& ExistMaterials
 		, TMap<FString, FString> ImportedSlotMaterialDependencies
diff --git a/Engine/Plugins/Interchange/Runtime/Source/Parsers/Fbx/Private/FbxAPI.cpp b/Engine/Plugins/Interchange/Runtime/Source/Parsers/Fbx/Private/FbxAPI.cpp
--- a/Engine/Plugins/Interchange/Runtime/Source/Parsers/Fbx/Private/FbxAPI.cpp
+++ b/Engine/Plugins/Interchange/Runtime/Source/Parsers/Fbx/Private/FbxAPI.cpp
@@ -164,6 +164,11 @@ namespace UE
 					SourceNode->SetCustomSourceFrameRateDenominator(Denominator);
 				}
 
+				// Fbx legacy has a special way to bake the skeletal mesh that do not fit the interchange standard
+				// The interchange skeletal mesh factory will read this to use the proper bake transform so it match legacy behavior.
+				// This fix the issue with blender armature bone skip
+				SourceNode->SetCustomAlignSkeletalMeshInScene(true);
+
 				//Fbx legacy does not allow Scene Root Nodes to be part of the skeletons (to be joints).
 				SourceNode->SetCustomAllowSceneRootAsJoint(false);
 
diff --git a/Engine/Plugins/Interchange/Runtime/Source/Parsers/Fbx/Private/FbxScene.cpp b/Engine/Plugins/Interchange/Runtime/Source/Parsers/Fbx/Private/FbxScene.cpp
--- a/Engine/Plugins/Interchange/Runtime/Source/Parsers/Fbx/Private/FbxScene.cpp
+++ b/Engine/Plugins/Interchange/Runtime/Source/Parsers/Fbx/Private/FbxScene.cpp
@@ -621,16 +621,28 @@ namespace UE
 				const bool bIsRootNode = Node == SDKScene->GetRootNode();
 
 				UInterchangeSceneNode* UnrealNode;
-				bool bIsNodeContainJointAttribute = false;
-				const int32 AttributeCount = Node->GetNodeAttributeCount();
 
 				bool bIsJoint = false;
+				bool bIsNodeContainJointAttribute = false;
+				
+				const FString NodeName = Parser.GetFbxHelper()->GetFbxObjectName(Node);
+				const FString NodeUniqueID = Parser.GetFbxHelper()->GetFbxNodeHierarchyName(Node);
+				const FString ParentNodeUid = UnrealParentNode ? UnrealParentNode->GetUniqueID() : TEXT("");
+
+				const bool bNodeUnderCommonRootJoint = IsNodeUnderCommonJointRootNode(Node, CommonJointRootNodes);
+
+				const int32 AttributeCount = Node->GetNodeAttributeCount();
 
 				// We need to know in advance if we have to create a UInterchangeSceneNode or UInterchangeJointNode.
 				for (int32 AttributeIndex = 0; AttributeIndex < AttributeCount; ++AttributeIndex)
 				{
 					FbxNodeAttribute* NodeAttribute = Node->GetNodeAttributeByIndex(AttributeIndex);
 
+					if (!NodeAttribute)
+					{
+						continue;
+					}
+
 					if (NodeAttribute->GetAttributeType() == FbxNodeAttribute::eSkeleton)
 					{
 						// All eSkeleton/bones are converted to an Interchange Joint node
@@ -639,20 +651,12 @@ namespace UE
 					}
 					else
 					{
-						const bool bNodeUnderCommonRootJoint = IsNodeUnderCommonJointRootNode(Node, CommonJointRootNodes);
-
 						if (NodeAttribute->GetAttributeType() == FbxNodeAttribute::eNull && bNodeUnderCommonRootJoint)
 						{
 							// Any eNull under an eSkeleton or equivalent is considered an eSkeleton too
 							bIsJoint = true;
 							bIsNodeContainJointAttribute = true;
 						}
-						else if (ForceJointNodes.Contains(Node) || (!bIsRootNode && bNodeUnderCommonRootJoint))
-						{
-							// Finally, any node that is forced as a bone or that is under an eSkeleton or equivalent (except the RootNode) is considered non-native bones.
-							bIsJoint = true;
-							bIsNodeContainJointAttribute = false;
-						}
 					}
 
 					if (bIsJoint)
@@ -661,9 +665,15 @@ namespace UE
 					}
 				}
 
-				const FString NodeName = Parser.GetFbxHelper()->GetFbxObjectName(Node);
-				const FString NodeUniqueID = Parser.GetFbxHelper()->GetFbxNodeHierarchyName(Node);
-				const FString ParentNodeUid = UnrealParentNode ? UnrealParentNode->GetUniqueID() : TEXT("");
+				if (!bIsJoint) //Can happen when AttributeCount is 0.
+				{
+					if (ForceJointNodes.Contains(Node) || (!bIsRootNode && bNodeUnderCommonRootJoint))
+					{
+						// Finally, any node that is forced as a bone or that is under an eSkeleton or equivalent (except the RootNode) is considered non-native bones.
+						bIsJoint = true;
+						bIsNodeContainJointAttribute = false;
+					}
+				}
 
 				if (bIsJoint)
 				{
diff --git a/Engine/Plugins/Interchange/Runtime/Source/Parsers/Fbx/Private/Ufbx/UfbxParser.cpp b/Engine/Plugins/Interchange/Runtime/Source/Parsers/Fbx/Private/Ufbx/UfbxParser.cpp
--- a/Engine/Plugins/Interchange/Runtime/Source/Parsers/Fbx/Private/Ufbx/UfbxParser.cpp
+++ b/Engine/Plugins/Interchange/Runtime/Source/Parsers/Fbx/Private/Ufbx/UfbxParser.cpp
@@ -199,6 +199,11 @@ bool FUfbxParser::LoadFbxFile(const FString& Filename, UInterchangeBaseNodeConta
 		SourceNode->SetExtraInformation(FSourceNodeExtraInfoStaticData::GetApplicationVersionExtraInfoKey(), ApplicationVersion);
 	}
 
+	// Fbx legacy has a special way to bake the skeletal mesh that do not fit the interchange standard
+	// The interchange skeletal mesh factory will read this to use the proper bake transform so it match legacy behavior.
+	// This fix the issue with blender armature bone skip
+	SourceNode->SetCustomAlignSkeletalMeshInScene(true);
+
 	//Fbx legacy does not allow Scene Root Nodes to be part of the skeletons (to be joints).
 	SourceNode->SetCustomAllowSceneRootAsJoint(false);
 
diff --git a/Engine/Plugins/Interchange/Runtime/Source/Pipelines/Private/InterchangeGenericAnimationPipeline.cpp b/Engine/Plugins/Interchange/Runtime/Source/Pipelines/Private/InterchangeGenericAnimationPipeline.cpp
--- a/Engine/Plugins/Interchange/Runtime/Source/Pipelines/Private/InterchangeGenericAnimationPipeline.cpp
+++ b/Engine/Plugins/Interchange/Runtime/Source/Pipelines/Private/InterchangeGenericAnimationPipeline.cpp
@@ -33,7 +33,7 @@ namespace UE::Interchange::Private
 {
 	const FString ConvertedFromRigidAnimationPrefixIdentifier = TEXT("\\SkeletalAnimation\\ConvertedFromRigidAnimation\\");
 
-	bool IsTranslatedDataContainOnlyJointAnimation(const UInterchangeBaseNodeContainer* InBaseNodeContainer, bool bConvertStaticsWithMorphTargetsToSkeletals)
+	bool IsTranslatedDataContainOnlyJointAnimation(const UInterchangeBaseNodeContainer* InBaseNodeContainer, bool bConvertStaticsWithMorphTargetsToSkeletals, bool bImportMeshesInBoneHierarchy)
 	{
 		//Its valid to call GetMeshesInformationFromTranslatedData with a null container
 		if (!InBaseNodeContainer)
@@ -49,13 +49,12 @@ namespace UE::Interchange::Private
 			});
 		if (bContainJointAnimation)
 		{
-			//if we have bone animation and no skinned mesh, we want to import animation only.
 			bool bContainSkinnedMeshNode = false;
-			InBaseNodeContainer->BreakableIterateNodesOfType<UInterchangeMeshNode>([&bContainSkinnedMeshNode, &bConvertStaticsWithMorphTargetsToSkeletals](const FString& NodeUid, UInterchangeMeshNode* MeshNode)
+			InBaseNodeContainer->BreakableIterateNodesOfType<UInterchangeMeshNode>([&bContainSkinnedMeshNode, &bConvertStaticsWithMorphTargetsToSkeletals, &bImportMeshesInBoneHierarchy](const FString& NodeUid, UInterchangeMeshNode* MeshNode)
 				{
 					if (!MeshNode->IsMorphTarget())
 					{
-						if (MeshNode->IsSkinnedMesh())
+						if (MeshNode->IsSkinnedMesh() || bImportMeshesInBoneHierarchy)
 						{
 							bContainSkinnedMeshNode = true;
 						}
@@ -212,7 +211,10 @@ void UInterchangeGenericAnimationPipeline::AdjustSettingsForContext(const FInter
 	TArray<FString> HideCategories;
 	if(ContextParams.ContextType == EInterchangePipelineContext::AssetImport)
 	{
-		if(UE::Interchange::Private::IsTranslatedDataContainOnlyJointAnimation(ContextParams.BaseNodeContainer, CommonSkeletalMeshesAndAnimationsProperties->bConvertStaticsWithMorphTargetsToSkeletals))
+		if(UE::Interchange::Private::IsTranslatedDataContainOnlyJointAnimation(
+			ContextParams.BaseNodeContainer, 
+			CommonSkeletalMeshesAndAnimationsProperties->bConvertStaticsWithMorphTargetsToSkeletals,
+			CommonSkeletalMeshesAndAnimationsProperties->bImportMeshesInBoneHierarchy))
 		{
 			bImportAnimations = true;
 			CommonSkeletalMeshesAndAnimationsProperties->bImportOnlyAnimations = true;
diff --git a/Engine/Plugins/Interchange/Runtime/Source/Pipelines/Private/InterchangeLODDataParser.cpp b/Engine/Plugins/Interchange/Runtime/Source/Pipelines/Private/InterchangeLODDataParser.cpp
--- a/Engine/Plugins/Interchange/Runtime/Source/Pipelines/Private/InterchangeLODDataParser.cpp
+++ b/Engine/Plugins/Interchange/Runtime/Source/Pipelines/Private/InterchangeLODDataParser.cpp
@@ -18,6 +18,7 @@
 #include "InterchangeSkeletonHelper.h"
 #include "InterchangeSkeletalMeshFactoryNode.h"
 #include "InterchangeSkeletalMeshLodDataNode.h"
+#include "Nodes/InterchangeSourceNode.h"
 #include "InterchangeStaticMeshFactoryNode.h"
 #include "InterchangeStaticMeshLodDataNode.h"
 
@@ -231,7 +232,13 @@ FInterchangeMeshLODDataParser<TMeshFactoryType>::FInterchangeMeshLODDataParser(U
 		CommonPipelineDataFactoryNode->GetCustomGlobalOffsetTransform(GlobalOffsetTransform);
 	}
 
-	RootJointParentGlobalInverseMatrix.SetIdentity();
+	bAlignSkeletalMeshInScene = false;
+	if (const UInterchangeSourceNode* SourceNode = UInterchangeSourceNode::GetUniqueInstance(&BaseNodeContainer))
+	{
+		SourceNode->GetCustomAlignSkeletalMeshInScene(bAlignSkeletalMeshInScene);
+	}
+
+	RootJointParentUidAndGlobalInverseMatrix.Reset();
 	if (InRootJointNodeUid.Len())
 	{
 		if (const UInterchangeSceneNode* RootJointNode = Cast<UInterchangeSceneNode>(BaseNodeContainer.GetNode(InRootJointNodeUid)))
@@ -241,7 +248,7 @@ FInterchangeMeshLODDataParser<TMeshFactoryType>::FInterchangeMeshLODDataParser(U
 				FTransform RootJointParentGlobalTransform;
 				RootJointParentNode->GetCustomGlobalTransform(&BaseNodeContainer, GlobalOffsetTransform, RootJointParentGlobalTransform);
 
-				RootJointParentGlobalInverseMatrix = RootJointParentGlobalTransform.ToMatrixWithScale().Inverse();
+				RootJointParentUidAndGlobalInverseMatrix = TPair<FString, FMatrix>(RootJointNode->GetParentUid(), RootJointParentGlobalTransform.ToMatrixWithScale().Inverse());
 			}
 		}
 		else
@@ -353,7 +360,8 @@ void FInterchangeMeshLODDataParser<TMeshFactoryType>::ProcessCollision(const FSt
         {
             FTransform RenderLocalToGlobal = UE::Interchange::Private::MeshHelper::CalculateCompleteSceneNodeTransform(&BaseNodeContainer, RenderSceneNode,
                 true, true, //Collisions are always baking global transforms
-                GlobalOffsetTransform, bUseTimeZeroAsBindPose, RootJointParentGlobalInverseMatrix);
+                GlobalOffsetTransform, bUseTimeZeroAsBindPose, 
+				RootJointParentUidAndGlobalInverseMatrix, bCreatingSkeletalMesh, bAlignSkeletalMeshInScene);
 
             // We want to ultimately produce "TransformToBake = CompleteTransform * RenderLocalToGlobal.Inverse();"
             TransformToBake = CompleteTransform * RenderLocalToGlobal.Inverse();
@@ -416,9 +424,8 @@ void FInterchangeMeshLODDataParser<TMeshFactoryType>::ProcessCollision(const FSt
 //	- UInterchangeGeometryCacheNode(which inherits from UInterchangeMeshNode, so can be handled in the same branch as UInterchangeMeshNode)
 //	- UInterchangeInstancedStaticMeshComponentNode
 //  - UInterchangeMeshBundleNode
-//CompleteTransform includes Pivot, Geometric, and GloblaTransforms and if it comes from a LODContainerNode/MeshBundle it includes MeshTransform.
 template <class TMeshFactoryType>
-void FInterchangeMeshLODDataParser<TMeshFactoryType>::ProcessNode(const UInterchangeBaseNode* BaseNode, const FTransform& CompleteTransform, const FTransform& GeometricTransform, const FString& SceneNodeName, const int32 LODIndex, bool bProcessAsSubMesh, bool bProcessAsCollision)
+void FInterchangeMeshLODDataParser<TMeshFactoryType>::ProcessNode(const UInterchangeBaseNode* BaseNode, const FTransform& MeshTransform, const UInterchangeSceneNode* CurrentTransformNode, const int32 LODIndex, bool bProcessAsSubMesh, bool bProcessAsCollision)
 {
 	if (!BaseNode)
 	{
@@ -437,6 +444,20 @@ void FInterchangeMeshLODDataParser<TMeshFactoryType>::ProcessNode(const UInterch
 
 		MeshFactoryNode->AddTargetNodeUid(MeshNodeUid);
 
+		FTransform CompleteTransformOfCurrentTransformNode = UE::Interchange::Private::MeshHelper::CalculateCompleteSceneNodeTransform(&BaseNodeContainer, CurrentTransformNode,
+			bBakeMeshes || bProcessAsCollision, bBakePivotMeshes || bProcessAsCollision, //Collisions are always baking global transforms
+			GlobalOffsetTransform, bUseTimeZeroAsBindPose, 
+			RootJointParentUidAndGlobalInverseMatrix, bCreatingSkeletalMesh, bAlignSkeletalMeshInScene);
+		FTransform CompleteTransform = MeshTransform * CompleteTransformOfCurrentTransformNode;
+
+		FString SceneNodeName = TEXT("");
+		FTransform GeometricTransform = FTransform::Identity;
+		if (CurrentTransformNode)
+		{
+			CurrentTransformNode->GetCustomGeometricTransform(GeometricTransform);
+			SceneNodeName = CurrentTransformNode->GetDisplayLabel();
+		}
+
 		auto AddLODMeshData = [this, &CompleteTransform, &GeometricTransform, &LODIndex, &SceneNodeName, &MeshNodeUid]()
 			{
 				if (LODIndex == INDEX_NONE)
@@ -521,14 +542,7 @@ void FInterchangeMeshLODDataParser<TMeshFactoryType>::ProcessNode(const UInterch
 		FString AssetInstanceUid;
 		if (SceneNode->GetCustomAssetInstanceUid(AssetInstanceUid))
 		{
-			FTransform OverrideSceneNodeCompleteTransform = UE::Interchange::Private::MeshHelper::CalculateCompleteSceneNodeTransform(&BaseNodeContainer, SceneNode, 
-                bBakeMeshes || bProcessAsCollision, bBakePivotMeshes || bProcessAsCollision, //Collisions are always baking global transforms
-                GlobalOffsetTransform, bUseTimeZeroAsBindPose, RootJointParentGlobalInverseMatrix);
-
-			FTransform GeometricTransformFromSceneNode;
-			SceneNode->GetCustomGeometricTransform(GeometricTransformFromSceneNode);
-
-			ProcessNode(BaseNodeContainer.GetNode(AssetInstanceUid), OverrideSceneNodeCompleteTransform, GeometricTransformFromSceneNode, SceneNode->GetDisplayLabel(), LODIndex, bProcessAsSubMesh, bProcessAsCollision);
+			ProcessNode(BaseNodeContainer.GetNode(AssetInstanceUid), FTransform::Identity, SceneNode, LODIndex, bProcessAsSubMesh, bProcessAsCollision);
 		}
 
 		//#interchange_LODRefactor_Note: Technically based on old implementation we should be calling GetSlotMaterialDependencies on the SceneNode if the AssetInstanceUid does not pan out.
@@ -555,7 +569,7 @@ void FInterchangeMeshLODDataParser<TMeshFactoryType>::ProcessNode(const UInterch
 		{
 			for (const FInterchangeLODMeshData& NodeUidTransformsPair : LODMeshesPerLODIndex.Value)
 			{
-				ProcessNode(BaseNodeContainer.GetNode(NodeUidTransformsPair.MeshUid), NodeUidTransformsPair.Transform * CompleteTransform, GeometricTransform, SceneNodeName, LODMeshesPerLODIndex.Key, true, bProcessAsCollision);
+				ProcessNode(BaseNodeContainer.GetNode(NodeUidTransformsPair.MeshUid), NodeUidTransformsPair.Transform, CurrentTransformNode, LODMeshesPerLODIndex.Key, true, bProcessAsCollision);
 			}
 		}
 
@@ -573,13 +587,6 @@ void FInterchangeMeshLODDataParser<TMeshFactoryType>::ProcessNode(const UInterch
 			{
 				if (const UInterchangeBaseNode* AssetBaseNode = BaseNodeContainer.GetNode(AssetInstanceUid))
 				{
-					FTransform OverrideSceneNodeCompleteTransform = UE::Interchange::Private::MeshHelper::CalculateCompleteSceneNodeTransform(&BaseNodeContainer, ISMSceneNode, 
-                        bBakeMeshes || bProcessAsCollision, bBakePivotMeshes || bProcessAsCollision, //Collisions are always baking global transforms
-                        GlobalOffsetTransform, bUseTimeZeroAsBindPose, RootJointParentGlobalInverseMatrix);
-
-					FTransform GeometricTransformFromSceneNode;
-					ISMSceneNode->GetCustomGeometricTransform(GeometricTransformFromSceneNode);
-
 					if (bBakeMeshes)
 					{
 						//In this case we want to add all instances to the LODDataNode
@@ -588,14 +595,14 @@ void FInterchangeMeshLODDataParser<TMeshFactoryType>::ProcessNode(const UInterch
 
 						for (const FTransform& InstanceTransform : InstanceTransforms)
 						{
-							ProcessNode(AssetBaseNode, InstanceTransform * OverrideSceneNodeCompleteTransform, GeometricTransformFromSceneNode, SceneNodeName, LODIndex, bProcessAsSubMesh, bProcessAsCollision);
+							ProcessNode(AssetBaseNode, InstanceTransform, ISMSceneNode, LODIndex, bProcessAsSubMesh, bProcessAsCollision);
 						}
 					}
 					else
 					{
 						//bBakeMeshes := false, we don't bake the instances, we only use the single occurance
 						// => Primary use case is Level import, in which case the instances will be appled to the Instanced Static Mesh Component.
-						ProcessNode(AssetBaseNode, OverrideSceneNodeCompleteTransform, GeometricTransformFromSceneNode, SceneNodeName, LODIndex, bProcessAsSubMesh, bProcessAsCollision);
+						ProcessNode(AssetBaseNode, FTransform::Identity, ISMSceneNode, LODIndex, bProcessAsSubMesh, bProcessAsCollision);
 					}
 				}
 			}
@@ -619,7 +626,7 @@ void FInterchangeMeshLODDataParser<TMeshFactoryType>::ProcessNode(const UInterch
 			{
 				for (const FTransform& Transform : Pair.Value)
 				{
-					ProcessNode(MeshBaseNode, Transform * CompleteTransform, GeometricTransform, SceneNodeName, LODIndex, true, bProcessAsCollision);
+					ProcessNode(MeshBaseNode, Transform, CurrentTransformNode, LODIndex, true, bProcessAsCollision);
 				}
 			}
 			
@@ -751,8 +758,8 @@ void FInterchangeMeshLODDataParser<TMeshFactoryType>::ProcessNodeUids(const TArr
 	for (const FString& NodeUid : RenderNodeUids)
 	{
 		ProcessNode(BaseNodeContainer.GetNode(NodeUid), 
-            FTransform::Identity/*CompleteMeshTransform*/, FTransform::Identity/*GeometricTransform*/, 
-            TEXT("") /*SceneNodeName*/, InitialLODIndex,
+            FTransform::Identity/*MeshTransform*/, nullptr /*CurrentNodeTransform*/,
+			InitialLODIndex,
             false /*bProcessAsSubmesh*/, false/*bProcessAsCollision*/);
 	}
 
@@ -761,8 +768,8 @@ void FInterchangeMeshLODDataParser<TMeshFactoryType>::ProcessNodeUids(const TArr
 		for (const FString& NodeUid : CollisionMeshUidsToProcess)
 		{
 			ProcessNode(BaseNodeContainer.GetNode(NodeUid),
-				FTransform::Identity /*CompleteMeshTransform*/, FTransform::Identity /*GeometricTransform*/,
-				TEXT("") /*SceneNodeName*/, InitialLODIndex,
+				FTransform::Identity /*MeshTransform*/, nullptr /*CurrentNodeTransform*/,
+				InitialLODIndex,
 				false /*bProcessAsSubmesh*/, true/*bProcessAsCollision*/);
 		}
 	}
@@ -781,8 +788,8 @@ void FInterchangeMeshLODDataParser<TMeshFactoryType>::ProcessNodeUids(const TArr
 			}
 
 			ProcessNode(BaseNodeContainer.GetNode(NodeUid),
-				FTransform::Identity /*CompleteMeshTransform*/, FTransform::Identity /*GeometricTransform*/,
-				TEXT("") /*SceneNodeName*/, InitialLODIndex,
+				FTransform::Identity /*MeshTransform*/, nullptr /*CurrentNodeTransform*/,
+				InitialLODIndex,
 				false /*bProcessAsSubmesh*/, false/*bProcessAsCollision*/);
 		}
 	}
diff --git a/Engine/Plugins/Interchange/Runtime/Source/Pipelines/Private/InterchangePipelineMeshesUtilities.cpp b/Engine/Plugins/Interchange/Runtime/Source/Pipelines/Private/InterchangePipelineMeshesUtilities.cpp
--- a/Engine/Plugins/Interchange/Runtime/Source/Pipelines/Private/InterchangePipelineMeshesUtilities.cpp
+++ b/Engine/Plugins/Interchange/Runtime/Source/Pipelines/Private/InterchangePipelineMeshesUtilities.cpp
@@ -153,69 +153,6 @@ namespace UE::Private::InterchangeMeshPipeline
 			}
 		}
 	}
-
-	void CollectAllChildren(UInterchangeBaseNodeContainer* BaseNodeContainer, const FString SceneNodeUids, TArray<FString>& AllChildrenUids)
-	{
-		AllChildrenUids.Add(SceneNodeUids);
-		TArray<FString> ChildrenUids = BaseNodeContainer->GetNodeChildrenUids(SceneNodeUids);
-		for (const FString& ChildUid : ChildrenUids)
-		{
-			CollectAllChildren(BaseNodeContainer, ChildUid, AllChildrenUids);
-		}
-	}
-
-	bool IsSceneNodeNestedInSkeleton(UInterchangeBaseNodeContainer* BaseNodeContainer, const UInterchangeSceneNode* SceneNode)
-	{
-		auto IsJointHookToDeformer = [&BaseNodeContainer](const FString JointNodeUid)->bool
-		{
-			TArray<FString> AllChildrenUids;
-			CollectAllChildren(BaseNodeContainer, JointNodeUid, AllChildrenUids);
-			bool bResult = false;
-			BaseNodeContainer->BreakableIterateNodesOfType<UInterchangeMeshNode>([&AllChildrenUids, &bResult](const FString& MeshUid, UInterchangeMeshNode* MeshNode)
-			{
-				if (MeshNode->IsSkinnedMesh())
-				{
-					TArray<FString> SkeletonDependencies;
-					MeshNode->GetSkeletonDependencies(SkeletonDependencies);
-					for (const FString& MeshJointUid : SkeletonDependencies)
-					{
-						if (AllChildrenUids.Contains(MeshJointUid))
-						{
-							bResult = true;
-							return true;
-						}
-					}
-				}
-				return false;
-			});
-
-			return bResult;
-		};
-
-		if (SceneNode->IsA(UInterchangeJointNode::StaticClass()))
-		{
-			if(IsJointHookToDeformer(SceneNode->GetUniqueID()))
-			{
-				return true;
-			}
-		}
-		FString ParentUid = SceneNode->GetParentUid();
-		const UInterchangeSceneNode* ParentNode = Cast<UInterchangeSceneNode>(BaseNodeContainer->GetNode(ParentUid));
-		while (ParentNode)
-		{
-			if (ParentNode->IsA(UInterchangeJointNode::StaticClass()))
-			{
-				if (IsJointHookToDeformer(ParentNode->GetUniqueID()))
-				{
-					return true;
-				}
-			}
-			ParentUid = ParentNode->GetParentUid();
-			ParentNode = Cast<UInterchangeSceneNode>(BaseNodeContainer->GetNode(ParentUid));
-		}
-		return false;
-	}
-
 }
 
 static bool IsSceneNodeASocket(const UInterchangeSceneNode* SceneNode)
@@ -436,13 +373,11 @@ bool FInterchangePipelineMeshesUtilitiesContext::IsSkeletalMeshInstance(const FI
 	}
 	else if (bImportMeshesInBoneHierarchy)
 	{
-		if (const UInterchangeSceneNode* SceneNode = Cast<const UInterchangeSceneNode>(MeshInstance.InstancingNode))
+		if (const UInterchangeJointNode* SceneNode = Cast<const UInterchangeJointNode>(MeshInstance.InstancingNode))
 		{
-			if (UE::Private::InterchangeMeshPipeline::IsSceneNodeNestedInSkeleton(BaseNodeContainer, SceneNode))
-			{
-				bOutIsStaticMeshNestedInSkeleton = true;
-				return true;
-			}
+			//If current Instancing Node is a JointNode then it is part of a Skeleton.
+			bOutIsStaticMeshNestedInSkeleton = true;
+			return true;
 		}
 	}
 	return false;
@@ -471,7 +406,23 @@ bool FInterchangePipelineMeshesUtilitiesContext::IsStaticMeshGeometry(const FInt
 	{
 		return false;
 	}
-	return !MeshGeometry.BaseNode->IsA<UInterchangeGeometryCacheNode>() && !IsSkeletalMeshGeometry(MeshGeometry);
+
+	if (MeshGeometry.BaseNode->IsA<UInterchangeGeometryCacheNode>())
+	{
+		return false;
+	}
+
+	if (bConvertSkeletalMeshToStaticMesh)
+	{
+		return true;
+	}
+
+	if (bImportMeshesInBoneHierarchy && MeshGeometry.bIsReferencedBySkeleton)
+	{
+		return false;
+	}
+
+	return !IsSkeletalMeshGeometry(MeshGeometry);
 }
 
 bool FInterchangePipelineMeshesUtilitiesContext::IsSkeletalMeshGeometry(const FInterchangeMeshGeometry& MeshGeometry)
@@ -566,9 +517,17 @@ UInterchangePipelineMeshesUtilities* UInterchangePipelineMeshesUtilities::Create
 				SceneNode->GetCustomActorVisibility(bIsVisible);
 			}
 
+			bool bIsNestedIntoSkeleton = false;
+			if (const UInterchangeJointNode* JointNode = Cast<const UInterchangeJointNode>(BaseNode))
+			{
+				bIsNestedIntoSkeleton = true;
+			}
+
 			if (AssetNode && (AssetNode->IsA<UInterchangeMeshNode>() || AssetNode->IsA<UInterchangeMeshLODContainerNode>()))
 			{
 				FInterchangeMeshGeometry& MeshGeometry = PipelineMeshesUtilities->MeshGeometriesPerMeshUid.FindChecked(AssetUid);
+				MeshGeometry.UidsOfInstancingNodes.Add(NodeUid);
+				MeshGeometry.bIsReferencedBySkeleton |= bIsNestedIntoSkeleton;
 
 				const UInterchangeBaseNode* BaseParentNode = BaseNodeContainer->GetNode(BaseNode->GetParentUid());
 
@@ -576,7 +535,6 @@ UInterchangePipelineMeshesUtilities* UInterchangePipelineMeshesUtilities::Create
 				{
 					FInterchangeMeshInstance& MeshInstance = PipelineMeshesUtilities->MeshInstancesPerMeshInstanceUid.FindOrAdd(NodeUid);
 					MeshInstance.InstancingNode = BaseNode;
-					MeshGeometry.UidsOfInstancingNodes.Add(NodeUid);
 					MeshInstance.ReferencingMeshGeometryUids.Add(AssetUid);
 					MeshInstance.bReferenceSkinnedMesh = false;
 					MeshInstance.bReferenceMorphTarget = false;
@@ -588,7 +546,6 @@ UInterchangePipelineMeshesUtilities* UInterchangePipelineMeshesUtilities::Create
 				{
 					FInterchangeMeshInstance& MeshInstance = PipelineMeshesUtilities->MeshInstancesPerMeshInstanceUid.FindOrAdd(NodeUid);
 					MeshInstance.InstancingNode = BaseNode;
-					MeshGeometry.UidsOfInstancingNodes.Add(NodeUid);
 					MeshInstance.ReferencingMeshGeometryUids.Add(AssetUid);
 					MeshInstance.bReferenceSkinnedMesh |= MeshGeometry.IsSkinnedMesh();
 					MeshInstance.bReferenceMorphTarget |= MeshGeometry.IsMorphTarget();
@@ -616,6 +573,7 @@ UInterchangePipelineMeshesUtilities* UInterchangePipelineMeshesUtilities::Create
 				{
 					FInterchangeMeshGeometry& MeshGeometry = PipelineMeshesUtilities->MeshGeometriesPerMeshUid.FindChecked(ReferenceMeshNodeUid);
 					MeshGeometry.UidsOfInstancingNodes.Add(NodeUid);
+					MeshGeometry.bIsReferencedBySkeleton |= bIsNestedIntoSkeleton;
 				}
 
 				MeshInstance.ReferencingMeshGeometryUids.Add(AssetUid);
diff --git a/Engine/Plugins/Interchange/Runtime/Source/Pipelines/Public/InterchangeLODDataParser.h b/Engine/Plugins/Interchange/Runtime/Source/Pipelines/Public/InterchangeLODDataParser.h
--- a/Engine/Plugins/Interchange/Runtime/Source/Pipelines/Public/InterchangeLODDataParser.h
+++ b/Engine/Plugins/Interchange/Runtime/Source/Pipelines/Public/InterchangeLODDataParser.h
@@ -9,6 +9,7 @@ class UInterchangeBaseNode;
 class UInterchangeBaseNodeContainer;
 class UInterchangeGenericMeshPipeline;
 class UInterchangePipelineMeshesUtilities;
+class UInterchangeSceneNode;
 class UInterchangeStaticMeshFactoryNode;
 class UInterchangeStaticMeshLodDataNode;
 class UInterchangeSkeletalMeshLodDataNode;
@@ -57,9 +58,20 @@ private:
 	TSet<FString> UniqueMeshUids; /*All MeshUids processed (including all LODs).*/
 	TMap<int32, int32> OriginalToCompactLODIndexMap;
 
-	FMatrix RootJointParentGlobalInverseMatrix; //Used for calculating skeletalMesh(JointNode instantiating a mesh)'s relative transform to RootJoineNode's parent (aka Skeleton Transform).
-												//This is needed, because GetTimeZeroGlobalTransform/GetBindPoseGlobalTransform on SceneNodes fallback onto LocalTransforms if BindPose/T0 is not present.
-												//Inherently this means that we get the Transforms relative to the Root SceneNode, regardless if they are Joint or not.
+	/**
+	* Used for calculating skeletalMesh's relative transform to RootJoineNode's parent (aka Skeleton Transform).
+	* This is needed, because GetTimeZeroGlobalTransform/GetBindPoseGlobalTransform on SceneNodes fallback onto LocalTransforms, if BindPose/T0 is not present.
+	* Inherently this means that we get the Transform's relative to the Root SceneNode, regardless if they are Joint or not,
+	*	To fix for this behavior we apply the RootJoint's Parent's Global Transform inverse when BakeMesh==false, to acquire 'SkeletonTransform'.
+	*	   (See @UE::Interchange::Private::MeshHelper::CalculateCompleteSceneNodeTransform)
+	*	If it is not set, the CalculateCompleteScenNodeTransform will try to identify the RootJoint's ParentNode and acquire described transform for application in there.
+	*	   (Based on Parent hierarchy and UInterchangeJointNode classification)
+	* The FString is the RootJointParent's UID, this is used for identifying, if the the instancing Node is part of the child tree/hierarchy of the RootJointParent.
+	*	(See @UE::Interchange::Private::MeshHelper::CalculateCompleteSceneNodeTransform)
+	*/
+	TOptional<TPair<FString, FMatrix>> RootJointParentUidAndGlobalInverseMatrix;
+	bool bCreatingSkeletalMesh = !std::is_same_v<TMeshFactoryType, UInterchangeStaticMeshFactoryNode>;
+	bool bAlignSkeletalMeshInScene;
 
 	TArray<FString> AllNodeUids; //Used for collission processing, combination of NodeUids and CollisionNodeUids
 
@@ -96,9 +108,8 @@ private:
 	/*
 	* bProcessAsSubMesh:= such as Mesh of a MeshLODContainerNode/MeshBundle/ISMComponentNode
 	*						we track these, as the ones flagged as sub mesh do not need to be configured via TargetNodes and the AssemblyParts do not need to be set either.
-	* GeometricTransform:= tracked for re-skinning to T0 calculations.
 	*/
-	void ProcessNode(const UInterchangeBaseNode* BaseNode, const FTransform& CompleteTransform, const FTransform& GeometricTransform, const FString& SceneNodeName, const int32 LODIndex, bool bProcessAsSubMesh, bool bProcessAsCollision);
+	void ProcessNode(const UInterchangeBaseNode* BaseNode, const FTransform& MeshTransform, const UInterchangeSceneNode* CurrentTransformNode, const int32 LODIndex, bool bProcessAsSubMesh, bool bProcessAsCollision);
 
 	//Parses All nodes to identify used LODIndices, sets up OriginalToCompactLODIndexMap, and creates all MeshLodDataNodes.
 	void PreProcess(const TArray<FString>& NodeUids);
diff --git a/Engine/Plugins/Interchange/Runtime/Source/Pipelines/Public/InterchangePipelineMeshesUtilities.h b/Engine/Plugins/Interchange/Runtime/Source/Pipelines/Public/InterchangePipelineMeshesUtilities.h
--- a/Engine/Plugins/Interchange/Runtime/Source/Pipelines/Public/InterchangePipelineMeshesUtilities.h
+++ b/Engine/Plugins/Interchange/Runtime/Source/Pipelines/Public/InterchangePipelineMeshesUtilities.h
@@ -79,6 +79,7 @@ struct FInterchangeMeshGeometry
 
 	FInterchangeMeshGeometry()
 		: BaseNode(nullptr)
+		, bIsReferencedBySkeleton(false)
 		, bIsMorphTarget(false)
 		, bIsSkinnedMesh(false)
 		, bHasMorphTargetDependencies(false)
@@ -107,6 +108,11 @@ struct FInterchangeMeshGeometry
 	UPROPERTY(EditAnywhere, Category = "Interchange | Pipeline | MeshesGeometry")
 	TArray<FString> AttachedSocketUids;
 
+	/*
+	* Indicates that at least 1 of the Instancing Node Uids is a Joint (aka part of a skeleton) or not.
+	*/
+	UPROPERTY(EditAnywhere, Category = "Interchange | Pipeline | MeshesGeometry")
+	bool bIsReferencedBySkeleton;
 
 	void SetupMeshGeometry(const UInterchangeBaseNodeContainer& BaseNodeContainer, const UInterchangeBaseNode* InBaseNode);
 
diff --git a/Engine/Source/Runtime/Interchange/Core/Private/Nodes/InterchangeSourceNode.cpp b/Engine/Source/Runtime/Interchange/Core/Private/Nodes/InterchangeSourceNode.cpp
--- a/Engine/Source/Runtime/Interchange/Core/Private/Nodes/InterchangeSourceNode.cpp
+++ b/Engine/Source/Runtime/Interchange/Core/Private/Nodes/InterchangeSourceNode.cpp
@@ -236,3 +236,12 @@ bool UInterchangeSourceNode::SetCustomAllowSceneRootAsJoint(const bool& Attribut
 	IMPLEMENT_NODE_ATTRIBUTE_SETTER_NODELEGATE(AllowSceneRootAsJoint, bool)
 }
 
+bool UInterchangeSourceNode::GetCustomAlignSkeletalMeshInScene(bool& AttributeValue) const
+{
+	IMPLEMENT_NODE_ATTRIBUTE_GETTER(AlignSkeletalMeshInScene, bool);
+}
+
+bool UInterchangeSourceNode::SetCustomAlignSkeletalMeshInScene(const bool& AttributeValue)
+{
+	IMPLEMENT_NODE_ATTRIBUTE_SETTER_NODELEGATE(AlignSkeletalMeshInScene, bool)
+}
diff --git a/Engine/Source/Runtime/Interchange/Core/Public/Nodes/InterchangeSourceNode.h b/Engine/Source/Runtime/Interchange/Core/Public/Nodes/InterchangeSourceNode.h
--- a/Engine/Source/Runtime/Interchange/Core/Public/Nodes/InterchangeSourceNode.h
+++ b/Engine/Source/Runtime/Interchange/Core/Public/Nodes/InterchangeSourceNode.h
@@ -206,6 +206,22 @@ public:
 	UFUNCTION(BlueprintCallable, Category = "Interchange | Node | Source")
 	INTERCHANGECORE_API bool SetCustomAllowSceneRootAsJoint(const bool& bAllowSceneRootAsJoint);
 
+
+	/** 
+	* Gets, SkeletalMeshes requires the Scene Transforms to be aligned with the Skeleton. 
+	*	If false then Skeleton and SkeletalMesh should be aligned on Asset level already.
+	*	(prior to the Scene Composition)
+	* For example:
+	*		true:=FBX
+	*		false:=glTF
+	*/
+	UFUNCTION(BlueprintCallable, Category = "Interchange | Node | Source")
+	INTERCHANGECORE_API bool GetCustomAlignSkeletalMeshInScene(bool& AlignSkeletalMeshInScene) const;
+
+	/** Sets, if SkeletalMeshes requires the Scene Transforms to be aligned with the Skeleton. */
+	UFUNCTION(BlueprintCallable, Category = "Interchange | Node | Source")
+	INTERCHANGECORE_API bool SetCustomAlignSkeletalMeshInScene(const bool& AlignSkeletalMeshInScene);
+
 private:
 	IMPLEMENT_NODE_ATTRIBUTE_KEY(SourceFrameRateNumerator);
 	IMPLEMENT_NODE_ATTRIBUTE_KEY(SourceFrameRateDenominator);
@@ -221,6 +237,7 @@ private:
 	IMPLEMENT_NODE_ATTRIBUTE_KEY(SkeletalMeshFrontAxis);
 	IMPLEMENT_NODE_ATTRIBUTE_KEY(NaniteTriangleThreshold);
 	IMPLEMENT_NODE_ATTRIBUTE_KEY(AllowSceneRootAsJoint);
+	IMPLEMENT_NODE_ATTRIBUTE_KEY(AlignSkeletalMeshInScene);
 
 	// Extra InformationTo show in the Config Panel.
 	UE::Interchange::TMapAttributeHelper<FString, FString> ExtraInformation;
