Interchange - New Skeletal Mesh Combine Options

Added new import options to control how Skeletal Meshes that share the same Skeleton are handled during import.

You can now choose to:

* Combine all meshes using the same Skeleton into a single Skeletal Mesh asset. (FBX default)
* Combine only visible meshes using the same Skeleton into a single Skeletal Mesh asset.
* Import meshes individually without combining them. (glTF/GLB default)

This provides improved flexibility when importing modular characters, multi-part skeletal assets, and DCC exports containing multiple mesh sections.

(above text generated with chatGPT for release notes)

New tests has been added to validate all three options.

#tests All Interchange tests passing locally (except 8 tests currently failing on Main related to USD materials)
[RN] Added new import options to control how Skeletal Meshes that share the same Skeleton are handled during import.

#virtualized
#Synced-CL 51380339

diff --git a/Engine/Plugins/Experimental/Enterprise/DatasmithInterchange/Source/DatasmithInterchange/Private/InterchangeDatasmithLevelPipeline.cpp b/Engine/Plugins/Experimental/Enterprise/DatasmithInterchange/Source/DatasmithInterchange/Private/InterchangeDatasmithLevelPipeline.cpp
--- a/Engine/Plugins/Experimental/Enterprise/DatasmithInterchange/Source/DatasmithInterchange/Private/InterchangeDatasmithLevelPipeline.cpp
+++ b/Engine/Plugins/Experimental/Enterprise/DatasmithInterchange/Source/DatasmithInterchange/Private/InterchangeDatasmithLevelPipeline.cpp
@@ -156,7 +156,8 @@ UInterchangeActorFactoryNode* UInterchangeDatasmithLevelPipeline::CreateActorFac
 	}
 }
 
-void UInterchangeDatasmithLevelPipeline::SetUpFactoryNode(UInterchangeActorFactoryNode* ActorFactoryNode, const UInterchangeSceneNode* SceneNode, const UInterchangeBaseNode* TranslatedAssetNode) const
+void UInterchangeDatasmithLevelPipeline::SetUpFactoryNode(UInterchangeActorFactoryNode* ActorFactoryNode, const UInterchangeSceneNode* SceneNode,
+	const UInterchangeBaseNode* TranslatedAssetNode, bool bIsRootJoint /* false */) const
 {
 	if (const UInterchangeDatasmithAreaLightNode* AreaLightNode = Cast<UInterchangeDatasmithAreaLightNode>(TranslatedAssetNode))
 	{
@@ -166,7 +167,7 @@ void UInterchangeDatasmithLevelPipeline::SetUpFactoryNode(UInterchangeActorFacto
 	}
 	else
 	{
-		Super::SetUpFactoryNode(ActorFactoryNode, SceneNode, TranslatedAssetNode);
+		Super::SetUpFactoryNode(ActorFactoryNode, SceneNode, TranslatedAssetNode, bIsRootJoint);
 	}
 }
 
diff --git a/Engine/Plugins/Experimental/Enterprise/DatasmithInterchange/Source/DatasmithInterchange/Private/InterchangeDatasmithLevelPipeline.h b/Engine/Plugins/Experimental/Enterprise/DatasmithInterchange/Source/DatasmithInterchange/Private/InterchangeDatasmithLevelPipeline.h
--- a/Engine/Plugins/Experimental/Enterprise/DatasmithInterchange/Source/DatasmithInterchange/Private/InterchangeDatasmithLevelPipeline.h
+++ b/Engine/Plugins/Experimental/Enterprise/DatasmithInterchange/Source/DatasmithInterchange/Private/InterchangeDatasmithLevelPipeline.h
@@ -27,7 +27,7 @@ protected:
 
 	virtual UInterchangeActorFactoryNode* CreateActorFactoryNode(const UInterchangeBaseNode* TranslatedAssetNode) const override;
 
-	virtual void SetUpFactoryNode(UInterchangeActorFactoryNode* ActorFactoryNode, const UInterchangeSceneNode* SceneNode, const UInterchangeBaseNode* TranslatedAssetNode) const override;
+	virtual void SetUpFactoryNode(UInterchangeActorFactoryNode* ActorFactoryNode, const UInterchangeSceneNode* SceneNode, const UInterchangeBaseNode* TranslatedAssetNode, bool bIsRootJoint = false) const override;
 
 private:
 
diff --git a/Engine/Plugins/Fab/Source/Fab/Private/Importers/GenericAssetImporter.cpp b/Engine/Plugins/Fab/Source/Fab/Private/Importers/GenericAssetImporter.cpp
--- a/Engine/Plugins/Fab/Source/Fab/Private/Importers/GenericAssetImporter.cpp
+++ b/Engine/Plugins/Fab/Source/Fab/Private/Importers/GenericAssetImporter.cpp
@@ -63,6 +63,7 @@ UObject* FFabGenericImporter::GetImportOptions(const FString& SourceFile, UObjec
 						GenericAssetsPipeline->MeshPipeline->bImportStaticMeshes                  = true;
 						GenericAssetsPipeline->MeshPipeline->bImportSkeletalMeshes                = true;
 						GenericAssetsPipeline->MeshPipeline->CombineStaticMeshesBehavior          = EInterchangeCombineStaticMeshesBehavior::All;
+						GenericAssetsPipeline->MeshPipeline->CombineSkeletalMeshesBehavior        = EInterchangeCombineSkeletalMeshesBehavior::BySkeleton;
 						GenericAssetsPipeline->MeshPipeline->SkeletalMeshImportContentType        = EInterchangeSkeletalMeshContentType::All;
 						GenericAssetsPipeline->MeshPipeline->bGenerateLightmapUVs                 = true;
 						GenericAssetsPipeline->MeshPipeline->bBuildNanite                         = false;
diff --git a/Engine/Plugins/Interchange/Runtime/Source/FactoryNodes/Public/InterchangeSkeletonFactoryNode.h b/Engine/Plugins/Interchange/Runtime/Source/FactoryNodes/Public/InterchangeSkeletonFactoryNode.h
--- a/Engine/Plugins/Interchange/Runtime/Source/FactoryNodes/Public/InterchangeSkeletonFactoryNode.h
+++ b/Engine/Plugins/Interchange/Runtime/Source/FactoryNodes/Public/InterchangeSkeletonFactoryNode.h
@@ -19,6 +19,12 @@ class UInterchangeSkeletonFactoryNode : public UInterchangeFactoryBaseNode
 	GENERATED_BODY()
 
 public:
+
+	UInterchangeSkeletonFactoryNode()
+	{
+		SkeletalMeshFactoryNodeUids.Initialize(Attributes.ToSharedRef(), "__SkeletalMeshFactoryNodeUids__");
+	}
+
 	/**
 	 * Initialize node data.
 	 * @param: UniqueID - The unique ID for this node.
@@ -101,16 +107,41 @@ public:
 		IMPLEMENT_NODE_ATTRIBUTE_SETTER_NODELEGATE(UseTimeZeroForBindPose, bool)
 	}
 
-	UFUNCTION(BlueprintCallable, Category = "Interchange | Node | Skeleton")
+	UFUNCTION(BlueprintCallable, Category = "Interchange | Node | Skeleton",
+		meta = (DeprecatedFunction, DeprecationMessage = "Use GetCustomSkeletalMeshFactoryNodeUidAtIndex (or GetCustomSkeletalMeshFactoryNodeUids) instead."))
 	bool GetCustomSkeletalMeshFactoryNodeUid(FString& AttributeValue) const
 	{
-		IMPLEMENT_NODE_ATTRIBUTE_GETTER(SkeletalMeshFactoryNodeUid, FString);
+		if (GetCustomSkeletalMeshFactoryNodeCount() > 0)
+		{
+			GetCustomSkeletalMeshFactoryNodeUidAtIndex(0, AttributeValue);
+			return true;
+		}
+
+		return false;
+	}
+
+	UFUNCTION(BlueprintCallable, Category = "Interchange | Node | Skeleton")
+	int32 GetCustomSkeletalMeshFactoryNodeCount() const
+	{
+		return SkeletalMeshFactoryNodeUids.GetCount();
+	}
+
+	UFUNCTION(BlueprintCallable, Category = "Interchange | Node | Skeleton")
+	void GetCustomSkeletalMeshFactoryNodeUids(TArray<FString>& OutDependencies) const
+	{
+		SkeletalMeshFactoryNodeUids.GetItems(OutDependencies);
 	}
 
 	UFUNCTION(BlueprintCallable, Category = "Interchange | Node | Skeleton")
-	bool SetCustomSkeletalMeshFactoryNodeUid(const FString& AttributeValue)
+	void GetCustomSkeletalMeshFactoryNodeUidAtIndex(const int32 Index, FString& OutDependency) const
 	{
-		IMPLEMENT_NODE_ATTRIBUTE_SETTER_NODELEGATE(SkeletalMeshFactoryNodeUid, FString)
+		SkeletalMeshFactoryNodeUids.GetItem(Index, OutDependency);
+	}
+
+	UFUNCTION(BlueprintCallable, Category = "Interchange | Node | Skeleton")
+	bool AddCustomSkeletalMeshFactoryNodeUid(const FString& InDependency)
+	{
+		return SkeletalMeshFactoryNodeUids.AddItem(InDependency);
 	}
 
 private:
@@ -152,7 +183,10 @@ private:
 	//Skeleton
 	const UE::Interchange::FAttributeKey Macro_CustomRootJointIDKey = UE::Interchange::FAttributeKey(TEXT("RootJointID"));
 	const UE::Interchange::FAttributeKey Macro_CustomUseTimeZeroForBindPoseKey = UE::Interchange::FAttributeKey(TEXT("UseTimeZeroForBindPose"));
-	const UE::Interchange::FAttributeKey Macro_CustomSkeletalMeshFactoryNodeUidKey = UE::Interchange::FAttributeKey(TEXT("SkeletalMeshFactoryNodeUid"));
+	//const UE::Interchange::FAttributeKey Macro_CustomSkeletalMeshFactoryNodeUidKey = UE::Interchange::FAttributeKey(TEXT("SkeletalMeshFactoryNodeUid"));
+	const UE::Interchange::FAttributeKey Macro_CustomSkeletalMeshFactoryNodeUidKey = UE::Interchange::FAttributeKey(TEXT("SkeletalMeshFactoryNodeUids"));
+	
+	UE::Interchange::TArrayAttributeHelper<FString> SkeletalMeshFactoryNodeUids;
 	
 protected:
 #if WITH_ENGINE
diff --git a/Engine/Plugins/Interchange/Runtime/Source/Import/Private/Animation/InterchangeAnimSequenceFactory.cpp b/Engine/Plugins/Interchange/Runtime/Source/Import/Private/Animation/InterchangeAnimSequenceFactory.cpp
--- a/Engine/Plugins/Interchange/Runtime/Source/Import/Private/Animation/InterchangeAnimSequenceFactory.cpp
+++ b/Engine/Plugins/Interchange/Runtime/Source/Import/Private/Animation/InterchangeAnimSequenceFactory.cpp
@@ -14,7 +14,6 @@
 #include "InterchangeMeshNode.h"
 #include "InterchangeJointNode.h"
 #include "InterchangeSkeletonFactoryNode.h"
-#include "InterchangeSkeletalMeshFactoryNode.h"
 #include "InterchangeSourceData.h"
 #include "InterchangeTranslatorBase.h"
 #include "Nodes/InterchangeBaseNode.h"
@@ -1537,7 +1536,7 @@ UInterchangeFactoryBase::FImportAssetResult UInterchangeAnimSequenceFactory::End
 
 	if (!bIsReImport)
 	{
-		/** Apply all AnimSequenceFactoryNode custom attributes to the skeletal mesh asset */
+		/** Apply all AnimSequenceFactoryNode custom attributes to the anim sequence asset */
 		AnimSequenceFactoryNode->ApplyAllCustomAttributeToObject(AnimSequence);
 	}
 	else
diff --git a/Engine/Plugins/Interchange/Runtime/Source/Pipelines/Private/InterchangeGenericAnimationPipeline.cpp b/Engine/Plugins/Interchange/Runtime/Source/Pipelines/Private/InterchangeGenericAnimationPipeline.cpp
--- a/Engine/Plugins/Interchange/Runtime/Source/Pipelines/Private/InterchangeGenericAnimationPipeline.cpp
+++ b/Engine/Plugins/Interchange/Runtime/Source/Pipelines/Private/InterchangeGenericAnimationPipeline.cpp
@@ -636,7 +636,6 @@ void UInterchangeGenericAnimationPipeline::CreateAnimSequenceFactoryNode(UInterc
 	const bool bAddCurveMetadataToSkeleton = CommonSkeletalMeshesAndAnimationsProperties->bAddCurveMetadataToSkeleton;
 
 	const UInterchangeSkeletonFactoryNode* SkeletonFactoryNode = nullptr;
-	const UInterchangeSkeletalMeshFactoryNode* SkeletalMeshFactoryNode = nullptr;
 
 	const FString SkeletonFactoryNodeUid = UInterchangeFactoryBaseNode::BuildFactoryNodeUid(SkeletonNodeUid);
 	SkeletonFactoryNode = Cast<const UInterchangeSkeletonFactoryNode>(BaseNodeContainer->GetFactoryNode(SkeletonFactoryNodeUid));
@@ -702,13 +701,6 @@ void UInterchangeGenericAnimationPipeline::CreateAnimSequenceFactoryNode(UInterc
 		return;
 	}
 
-	FString SkeletalMeshFactoryNodeUid;
-
-	if (SkeletonFactoryNode->GetCustomSkeletalMeshFactoryNodeUid(SkeletalMeshFactoryNodeUid))
-	{
-		SkeletalMeshFactoryNode = Cast<const UInterchangeSkeletalMeshFactoryNode>(BaseNodeContainer->GetFactoryNode(SkeletalMeshFactoryNodeUid));
-	}
-
 	double SampleRate = 30.;
 	double StartTime = 0.;
 	double StopTime = 0.;
@@ -829,10 +821,6 @@ void UInterchangeGenericAnimationPipeline::CreateAnimSequenceFactoryNode(UInterc
 	{
 		AnimSequenceFactoryNode->SetCustomSkeletonFactoryNodeUid(SkeletonFactoryNode->GetUniqueID());
 	}
-	if (SkeletalMeshFactoryNode)
-	{
-		AnimSequenceFactoryNode->AddFactoryDependencyUid(SkeletalMeshFactoryNode->GetUniqueID());
-	}
 
 	AnimSequenceFactoryNode->SetCustomImportBoneTracks(bImportBoneTracks);
 	AnimSequenceFactoryNode->SetCustomImportBoneTracksSampleRate(SampleRate);
@@ -856,7 +844,7 @@ void UInterchangeGenericAnimationPipeline::CreateAnimSequenceFactoryNode(UInterc
 		AnimSequenceFactoryNode->SetAnimatedMaterialCurveSuffixe(MaterialSuffixe);
 	}
 
-	//USkeleton cannot be created without a valid skeletal mesh
+	//Animation sequence cannot be created without a valid skeleton
 	FString SkeletonUid;
 	if(SkeletonFactoryNode)
 	{
@@ -864,6 +852,8 @@ void UInterchangeGenericAnimationPipeline::CreateAnimSequenceFactoryNode(UInterc
 		AnimSequenceFactoryNode->AddFactoryDependencyUid(SkeletonUid);
 	}
 
+
+
 	FString RootJointUid;
 	if (SkeletonFactoryNode && SkeletonFactoryNode->GetCustomRootJointUid(RootJointUid))
 	{
@@ -933,6 +923,8 @@ void UInterchangeGenericAnimationPipeline::CreateAnimSequenceFactoryNode(UInterc
 	}
 
 	//Iterate dependencies
+	//Skeletal Meshes factories are responsible to populate USkeletons using MergeAllBonesToBoneTree.
+	//For this reason, Animation sequence has a dependency on all the Skeletal Meshes using that Skeleton.
 	{
 		TArray<FString> SkeletalMeshNodeUids;
 		BaseNodeContainer->GetNodes(UInterchangeSkeletalMeshFactoryNode::StaticClass(), SkeletalMeshNodeUids);
diff --git a/Engine/Plugins/Interchange/Runtime/Source/Pipelines/Private/InterchangeGenericMeshPipeline.cpp b/Engine/Plugins/Interchange/Runtime/Source/Pipelines/Private/InterchangeGenericMeshPipeline.cpp
--- a/Engine/Plugins/Interchange/Runtime/Source/Pipelines/Private/InterchangeGenericMeshPipeline.cpp
+++ b/Engine/Plugins/Interchange/Runtime/Source/Pipelines/Private/InterchangeGenericMeshPipeline.cpp
@@ -300,6 +300,7 @@ bool UInterchangeGenericMeshPipeline::IsPropertyChangeNeedRefresh(const FPropert
 		GET_MEMBER_NAME_CHECKED(UInterchangeGenericMeshPipeline, bImportSkeletalMeshes),
 		GET_MEMBER_NAME_CHECKED(UInterchangeGenericMeshPipeline, bCreatePhysicsAsset),
 		GET_MEMBER_NAME_CHECKED(UInterchangeGenericMeshPipeline, CombineStaticMeshesBehavior),
+		GET_MEMBER_NAME_CHECKED(UInterchangeGenericMeshPipeline, CombineSkeletalMeshesBehavior),
 	};
 	
 	if (NeedRefreshProperties.Contains(PropertyChangedEvent.GetPropertyName()))
diff --git a/Engine/Plugins/Interchange/Runtime/Source/Pipelines/Private/InterchangeGenericScenesPipeline.cpp b/Engine/Plugins/Interchange/Runtime/Source/Pipelines/Private/InterchangeGenericScenesPipeline.cpp
--- a/Engine/Plugins/Interchange/Runtime/Source/Pipelines/Private/InterchangeGenericScenesPipeline.cpp
+++ b/Engine/Plugins/Interchange/Runtime/Source/Pipelines/Private/InterchangeGenericScenesPipeline.cpp
@@ -446,26 +446,52 @@ void UInterchangeGenericLevelPipeline::ExecuteSceneNodePreImport(const FTransfor
 		return;
 	}
 
-	const UInterchangeBaseNode* TranslatedAssetNode = nullptr;
-	bool bRootJointNode = false;
-	FString SkeletalMeshFactoryNodeUid;
-
 	FString SkeletonFactoryNodeUid = UInterchangeFactoryBaseNode::BuildFactoryNodeUid(SceneNode->GetUniqueID());
 	const UInterchangeSkeletonFactoryNode* SkeletonFactoryNode = Cast<UInterchangeSkeletonFactoryNode>(BaseNodeContainer->GetFactoryNode(SkeletonFactoryNodeUid));
+
+	bool bRootJointNode = false;
+
+	//If this SceneNode has a SkeletonFactoryNode associated with it, then we need to create a Skeletal actors
 	if (SkeletonFactoryNode)
 	{
-		if (SkeletonFactoryNode->GetCustomSkeletalMeshFactoryNodeUid(SkeletalMeshFactoryNodeUid))
+		//Create one Skeletal actor for each SkeletalMeshFactoryNode referenced by this SkeletonFactoryNode.
+		const int32 Count = SkeletonFactoryNode->GetCustomSkeletalMeshFactoryNodeCount();
+		for (int32 Index = 0; Index < Count; ++Index)
 		{
-			if (const UInterchangeFactoryBaseNode* SkeletalMeshFactoryNode = BaseNodeContainer->GetFactoryNode(SkeletalMeshFactoryNodeUid))
-			{
-				TArray<FString> NodeUids;
-				SkeletalMeshFactoryNode->GetTargetNodeUids(NodeUids);
+			FString SkeletalMeshFactoryNodeUid;
+			SkeletonFactoryNode->GetCustomSkeletalMeshFactoryNodeUidAtIndex(Index, SkeletalMeshFactoryNodeUid);
 
-				if (NodeUids.Num() > 0)
-				{
-					TranslatedAssetNode = BaseNodeContainer->GetNode(NodeUids[0]);
-					bRootJointNode = true;
-				}
+			CreateActorFactoryNode(GlobalOffsetTransform, SceneNode, SkeletalMeshFactoryNodeUid);
+			bRootJointNode = true;
+		}
+	}
+
+	if (!bRootJointNode)
+	{
+		CreateActorFactoryNode(GlobalOffsetTransform, SceneNode, {});
+	}
+}
+
+void UInterchangeGenericLevelPipeline::CreateActorFactoryNode(const FTransform& GlobalOffsetTransform, const UInterchangeSceneNode* SceneNode, const FString& SkeletalMeshFactoryNodeUid)
+{	
+	const UInterchangeBaseNode* TranslatedAssetNode = nullptr;
+	bool bRootJointNode = false;
+
+	FString ActorName = SceneNode->GetDisplayLabel();
+
+	if (!SkeletalMeshFactoryNodeUid.IsEmpty())
+	{
+		if (const UInterchangeFactoryBaseNode* SkeletalMeshFactoryNode = BaseNodeContainer->GetFactoryNode(SkeletalMeshFactoryNodeUid))
+		{
+			ActorName = SkeletalMeshFactoryNode->GetDisplayLabel();
+
+			TArray<FString> NodeUids;
+			SkeletalMeshFactoryNode->GetTargetNodeUids(NodeUids);
+
+			if (NodeUids.Num() > 0)
+			{
+				TranslatedAssetNode = BaseNodeContainer->GetNode(NodeUids[0]);
+				bRootJointNode = true;
 			}
 		}
 	}
@@ -483,7 +509,7 @@ void UInterchangeGenericLevelPipeline::ExecuteSceneNodePreImport(const FTransfor
 			bool bIsSkinnedMesh = MeshNode->IsSkinnedMesh();
 			
 			// Skinned mesh are added when the bRootJointNode is true.
-			// In this case we dont want to add an empty staticmesh actor.
+			// In this case we don't want to add an empty static mesh actor.
 			if (bIsSkinnedMesh)
 			{
 				return;
@@ -548,7 +574,17 @@ void UInterchangeGenericLevelPipeline::ExecuteSceneNodePreImport(const FTransfor
 	SceneNode->GetTags(Tags);
 	ActorFactoryNode->AddTags(Tags);
 
-	FString NodeUid = SceneNode->GetUniqueID() + (bRootJointNode ? TEXT("_SkeletonNode") : TEXT(""));
+	FString NodeUid;
+
+	if (bRootJointNode)
+	{
+		NodeUid = SkeletalMeshFactoryNodeUid + TEXT("_SkeletonNode");
+	}
+	else
+	{
+		NodeUid = SceneNode->GetUniqueID();
+	}
+
 	const FString ActorFactoryNodeUid = UInterchangeFactoryBaseNode::BuildFactoryNodeUid(NodeUid);
 	FString ParentFactoryNodeUid;
 	if (!SceneNode->GetParentUid().IsEmpty())
@@ -560,7 +596,7 @@ void UInterchangeGenericLevelPipeline::ExecuteSceneNodePreImport(const FTransfor
 			ParentFactoryNodeUid = UInterchangeFactoryBaseNode::BuildFactoryNodeUid(ParentNodeUid);
 		}
 	}
-	BaseNodeContainer->SetupNode(ActorFactoryNode, ActorFactoryNodeUid, SceneNode->GetDisplayLabel(), EInterchangeNodeContainerType::FactoryData, ParentFactoryNodeUid);
+	BaseNodeContainer->SetupNode(ActorFactoryNode, ActorFactoryNodeUid, ActorName, EInterchangeNodeContainerType::FactoryData, ParentFactoryNodeUid);
 #if WITH_EDITORONLY_DATA
 	//The level must be create before any actor asset since all actors will be created in the specified level
 	if (LevelFactoryNode)
@@ -594,8 +630,6 @@ void UInterchangeGenericLevelPipeline::ExecuteSceneNodePreImport(const FTransfor
 	}
 
 	//TODO move this code to the factory, a stack over pipeline can change the global offset transform which will affect this value.
-	//We prioritize Local (Relative) Transforms due to issues introduced by 0 scales with Global Transforms.
-	//In case the LocalTransform is not available we fallback onto GlobalTransforms
 	FTransform LocalTransform;
 	if (SceneNode->GetCustomLocalTransform(LocalTransform))
 	{
@@ -612,27 +646,6 @@ void UInterchangeGenericLevelPipeline::ExecuteSceneNodePreImport(const FTransfor
 
 		ActorFactoryNode->SetCustomLocalTransform(LocalTransform);
 	}
-	else
-	{
-		FTransform GlobalTransform;
-		if (SceneNode->GetCustomGlobalTransform(BaseNodeContainer, GlobalOffsetTransform, GlobalTransform))
-		{
-			if (bRootJointNode)
-			{
-				GlobalTransform = FTransform::Identity;
-				//LocalTransform of RootjointNode is already baked into the Skeletal and animation.
-				//due to that we acquire the Parent SceneNode and get its GlobalTransform:
-				if (!SceneNode->GetParentUid().IsEmpty())
-				{
-					if (const UInterchangeSceneNode* ParentSceneNode = Cast<UInterchangeSceneNode>(BaseNodeContainer->GetNode(SceneNode->GetParentUid())))
-					{
-						ParentSceneNode->GetCustomGlobalTransform(BaseNodeContainer, GlobalOffsetTransform, GlobalTransform);
-					}
-				}
-			}
-			ActorFactoryNode->SetCustomGlobalTransform(GlobalTransform);
-		}
-	}
 
 	bool bCustomVisibility;
 	// Copy component's visibility from SceneNode or set component visible by default
@@ -659,7 +672,7 @@ void UInterchangeGenericLevelPipeline::ExecuteSceneNodePreImport(const FTransfor
 
 	if (TranslatedAssetNode)
 	{
-		SetUpFactoryNode(ActorFactoryNode, SceneNode, TranslatedAssetNode);
+		SetUpFactoryNode(ActorFactoryNode, SceneNode, TranslatedAssetNode, bRootJointNode);
 	}
 
 	SetUpComponents(ActorFactoryNode, SceneNode);
@@ -800,7 +813,7 @@ void UInterchangeGenericLevelPipeline::CreateComponentFactoryNode(const UInterch
 }
 
 
-void UInterchangeGenericLevelPipeline::SetUpFactoryNode(UInterchangeActorFactoryNode* ActorFactoryNode, const UInterchangeSceneNode* SceneNode, const UInterchangeBaseNode* TranslatedAssetNode) const
+void UInterchangeGenericLevelPipeline::SetUpFactoryNode(UInterchangeActorFactoryNode* ActorFactoryNode, const UInterchangeSceneNode* SceneNode, const UInterchangeBaseNode* TranslatedAssetNode, bool bIsRootJoint /* false */) const
 {
 	if (!ensure(BaseNodeContainer && ActorFactoryNode && SceneNode && TranslatedAssetNode))
 	{
@@ -811,16 +824,8 @@ void UInterchangeGenericLevelPipeline::SetUpFactoryNode(UInterchangeActorFactory
 	{
 		TArray<FString> TargetNodeUids;
 		ActorFactoryNode->GetTargetNodeUids(TargetNodeUids);
-		bool bSkeletal = false;
-		if (TargetNodeUids.Num() > 0)
-		{
-			if (const UInterchangeSkeletalMeshFactoryNode* SkeletalMeshFactoryNode = Cast<UInterchangeSkeletalMeshFactoryNode>(BaseNodeContainer->GetFactoryNode(TargetNodeUids[0])))
-			{
-				bSkeletal = true;
-			}
-		}
 
-		if (bSkeletal)
+		if (bIsRootJoint)
 		{
 			ActorFactoryNode->SetCustomActorClassName(ASkeletalMeshActor::StaticClass()->GetPathName());
 			ActorFactoryNode->SetCustomMobility(EComponentMobility::Movable);
@@ -844,14 +849,14 @@ void UInterchangeGenericLevelPipeline::SetUpFactoryNode(UInterchangeActorFactory
 			}
 
 			MeshActorFactoryNode->AddFactoryDependencyUid(UInterchangeFactoryBaseNode::BuildFactoryNodeUid(TranslatedAssetNode->GetUniqueID()));
-			if (!bSkeletal)
+			if (bIsRootJoint)
 			{
-				MeshActorFactoryNode->SetCustomInstancedAssetFactoryNodeUid(UInterchangeFactoryBaseNode::BuildFactoryNodeUid(TranslatedAssetNode->GetUniqueID()));
+				// Directly points to the SkeletalMeshActorFactory that holds reference to SkeletalMesh Asset in the CustomReferencedObject.
+				MeshActorFactoryNode->SetCustomInstancedAssetFactoryNodeUid(TargetNodeUids[0]);
 			}
 			else
 			{
-				// Directly points to the SkeletalMeshActorFactory that holds reference to SkeletalMesh Asset in the CustomReferencedObject.
-				MeshActorFactoryNode->SetCustomInstancedAssetFactoryNodeUid(TargetNodeUids[0]);
+				MeshActorFactoryNode->SetCustomInstancedAssetFactoryNodeUid(UInterchangeFactoryBaseNode::BuildFactoryNodeUid(TranslatedAssetNode->GetUniqueID()));
 			}
 
 			FTransform GeometricTransform;
diff --git a/Engine/Plugins/Interchange/Runtime/Source/Pipelines/Private/InterchangeGenericSkeletalMeshPipeline.cpp b/Engine/Plugins/Interchange/Runtime/Source/Pipelines/Private/InterchangeGenericSkeletalMeshPipeline.cpp
--- a/Engine/Plugins/Interchange/Runtime/Source/Pipelines/Private/InterchangeGenericSkeletalMeshPipeline.cpp
+++ b/Engine/Plugins/Interchange/Runtime/Source/Pipelines/Private/InterchangeGenericSkeletalMeshPipeline.cpp
@@ -94,7 +94,7 @@ void UInterchangeGenericMeshPipeline::ExecutePreImportPipelineSkeletalMesh()
 	}
 
 #if WITH_EDITOR
-	//Make sure the generic pipeline we cover all skeletalmesh build settings by asserting when we import
+	//Make sure the generic pipeline we cover all skeletal mesh build settings by asserting when we import
 	Async(EAsyncExecution::TaskGraphMainThread, []()
 		{
 			static bool bVerifyBuildProperties = false;
@@ -114,124 +114,57 @@ void UInterchangeGenericMeshPipeline::ExecutePreImportPipelineSkeletalMesh()
 
 	TMap<FString, TArray<FString>> SkeletalMeshFactoryDependencyOrderPerSkeletonRootNodeUid;
 
-	auto SetSkeletalMeshDependencies = [&SkeletalMeshFactoryDependencyOrderPerSkeletonRootNodeUid](const FString& JointNodeUid, UInterchangeSkeletalMeshFactoryNode* SkeletalMeshFactoryNode)
-	{
-		TArray<FString>& SkeletalMeshFactoryDependencyOrder = SkeletalMeshFactoryDependencyOrderPerSkeletonRootNodeUid.FindOrAdd(JointNodeUid);
-		//Updating the skeleton is not multi thread safe, so we add dependency between skeletalmesh altering the same skeleton
-		//TODO make the skeletalMesh ReferenceSkeleton thread safe to allow multiple parallel skeletalmesh factory on the same skeleton asset.
-		int32 DependencyIndex = SkeletalMeshFactoryDependencyOrder.AddUnique(SkeletalMeshFactoryNode->GetUniqueID());
-		if (DependencyIndex > 0)
-		{
-			const FString SkeletalMeshFactoryNodeDependencyUid = SkeletalMeshFactoryDependencyOrder[DependencyIndex - 1];
-			SkeletalMeshFactoryNode->AddFactoryDependencyUid(SkeletalMeshFactoryNodeDependencyUid);
-		}
-	};
-
-	const FString SourceFileName = SourceDatas.Num() ? SourceDatas[0]->GetFilename() : TEXT("No Source File Specified");
-
-	PRAGMA_DISABLE_DEPRECATION_WARNINGS
-	const bool bCombineSkeletalMeshes = bCombineSkeletalMeshes_DEPRECATED;
-	PRAGMA_ENABLE_DEPRECATION_WARNINGS
-	if (bCombineSkeletalMeshes)
-	{
-		//////////////////////////////////////////////////////////////////////////
-		//Combined everything we can
-		TMap<FString, TArray<FString>> MeshUidsPerSkeletonRootUid;
-		auto CreatePerSkeletonRootUidCombinedSkinnedMesh = [this, &MeshUidsPerSkeletonRootUid, &SetSkeletalMeshDependencies, &SourceFileName]()
+	auto CreateSkeletalMeshFactoryForSkeletonRootUid = [this, &SkeletalMeshFactoryDependencyOrderPerSkeletonRootNodeUid](const FString& RootJointUid, const TArray<FString>& NodeUids)
 		{
-#if !WITH_EDITOR || !WITH_EDITORONLY_DATA
-			if (MeshUidsPerSkeletonRootUid.Num() > 0)
-			{
-				UE_LOG(LogInterchangePipeline, Warning, TEXT("Cannot import skeletalMesh asset in runtime, this is an editor only feature."));
-			}
-			return false;
-#else
-			if (MeshUidsPerSkeletonRootUid.Num() > 0)
+			UInterchangeSkeletalMeshFactoryNode* SkeletalMeshFactoryNode = CreateSkeletalMeshFactoryNode(RootJointUid, NodeUids);
+			if (SkeletalMeshFactoryNode != nullptr)
 			{
-
-				bool bFoundInstances = false;
-				for (const TPair<FString, TArray<FString>>& SkeletonRootUidAndMeshUids : MeshUidsPerSkeletonRootUid)
+				TArray<FString>& SkeletalMeshFactoryDependencyOrder = SkeletalMeshFactoryDependencyOrderPerSkeletonRootNodeUid.FindOrAdd(RootJointUid);
+				//Skeletal Mesh Factory are responsible of populating the USkeleton, and updating the skeleton is not multi thread safe, so we add dependency between skeletal mesh altering the same skeleton
+				//TODO make the skeletalMesh ReferenceSkeleton thread safe to allow multiple parallel skeletal mesh factory on the same skeleton asset.
+				int32 DependencyIndex = SkeletalMeshFactoryDependencyOrder.AddUnique(SkeletalMeshFactoryNode->GetUniqueID());
+				if (DependencyIndex > 0)
 				{
-					const FString& SkeletonRootUid = SkeletonRootUidAndMeshUids.Key;
-					//Every iteration is a skeletalmesh asset that combine all MeshInstances sharing the same skeleton root node
-					UInterchangeSkeletonFactoryNode* SkeletonFactoryNode = CommonSkeletalMeshesAndAnimationsProperties->CreateSkeletonFactoryNode(BaseNodeContainer, SkeletonRootUid, SourceFileName);
-					//The MeshUids can represent a SceneNode pointing on a MeshNode or directly a MeshNode;
-					const TArray<FString>& MeshNodeOrMeshInstanceNodeUids = SkeletonRootUidAndMeshUids.Value;
-
-					UInterchangeSkeletalMeshFactoryNode* SkeletalMeshFactoryNode = CreateSkeletalMeshFactoryNode(SkeletonRootUid, MeshNodeOrMeshInstanceNodeUids);
-					if (SkeletalMeshFactoryNode != nullptr)
-					{
-						SetSkeletalMeshDependencies(SkeletonRootUid, SkeletalMeshFactoryNode);
-						bFoundInstances = true;
-					}
+					const FString SkeletalMeshFactoryNodeDependencyUid = SkeletalMeshFactoryDependencyOrder[DependencyIndex - 1];
+					SkeletalMeshFactoryNode->AddFactoryDependencyUid(SkeletalMeshFactoryNodeDependencyUid);
 				}
-				return bFoundInstances;
 			}
-			return false;
-#endif
 		};
 
-		const bool bUseSingleBoneForConvertedMeshes = CommonMeshesProperties->ForceAllMeshAsType == EInterchangeForceMeshType::IFMT_SkeletalMesh && CommonMeshesProperties->bSingleBoneSkeleton;
-		PipelineMeshesUtilities->GetCombinedSkinnedMeshInstances(MeshUidsPerSkeletonRootUid, bUseSingleBoneForConvertedMeshes);
-		CreatePerSkeletonRootUidCombinedSkinnedMesh();
-	}
-	else
+	TMap<FString, TArray<FString>> MeshUidsPerSkeletonRootUid;
+	const bool bUseSingleBoneForConvertedMeshes = CommonMeshesProperties->ForceAllMeshAsType == EInterchangeForceMeshType::IFMT_SkeletalMesh && CommonMeshesProperties->bSingleBoneSkeleton;
+	const bool bVisibleOnly = CombineSkeletalMeshesBehavior == EInterchangeCombineSkeletalMeshesBehavior::BySkeletonVisibleOnly;
+	PipelineMeshesUtilities->GetCombinedSkinnedMeshInstances(MeshUidsPerSkeletonRootUid, bUseSingleBoneForConvertedMeshes, bVisibleOnly);
+
+	if (MeshUidsPerSkeletonRootUid.Num() > 0)
 	{
-		//////////////////////////////////////////////////////////////////////////
-		//Do not combined meshes
-		TArray<FString> MeshUids;
-		auto CreatePerSkeletonRootUidSkinnedMesh = [this, &MeshUids, &SetSkeletalMeshDependencies, &SourceFileName]()
-		{
 #if !WITH_EDITOR || !WITH_EDITORONLY_DATA
-			if (MeshUids.Num() > 0)
-			{
-				UE_LOG(LogInterchangePipeline, Warning, TEXT("Cannot import skeletalMesh asset in runtime, this is an editor only feature."));
-			}
-			return false;
+		UE_LOG(LogInterchangePipeline, Warning, TEXT("Cannot import skeletalMesh asset in runtime, this is an editor only feature."));
 #else
-			if (MeshUids.Num() > 0)
+		const FString SourceFileName = SourceDatas.Num() ? SourceDatas[0]->GetFilename() : TEXT("No Source File Specified");
+
+		for (const TPair<FString, TArray<FString>>& SkeletonRootUidAndMeshUids : MeshUidsPerSkeletonRootUid)
+		{
+			const FString& SkeletonRootUid = SkeletonRootUidAndMeshUids.Key;
+			//Every iteration is a skeletal mesh asset that combine all MeshInstances sharing the same skeleton root node
+			CommonSkeletalMeshesAndAnimationsProperties->CreateSkeletonFactoryNode(BaseNodeContainer, SkeletonRootUid, SourceFileName);
+			
+			//The MeshUids can represent a SceneNode pointing on a MeshNode or directly a MeshNode;
+			const TArray<FString>& MeshNodeOrMeshInstanceNodeUids = SkeletonRootUidAndMeshUids.Value;
+
+			if (CombineSkeletalMeshesBehavior == EInterchangeCombineSkeletalMeshesBehavior::DoNotCombine)
 			{
-				bool bFoundInstances = false;
-				for (const FString& MeshUid : MeshUids)
+				for (const FString& MeshUid : MeshNodeOrMeshInstanceNodeUids)
 				{
-					//Every iteration is a skeletalmesh asset that combine all MeshInstances sharing the same skeleton root node
-					//The MeshUids can represent a SceneNode pointing on a MeshNode or directly a MeshNode;
-					TMap<int32, TArray<FString>> MeshUidsPerLodIndex;
-					FString SkeletonRootUid;
-					if (!(PipelineMeshesUtilities->IsValidMeshInstanceUid(MeshUid)))
-					{
-						continue;
-					}
-					SkeletonRootUid = PipelineMeshesUtilities->GetMeshInstanceSkeletonRootUid(MeshUid);
-					if (SkeletonRootUid.IsEmpty())
-					{
-						if (CommonSkeletalMeshesAndAnimationsProperties->bConvertStaticsWithMorphTargetsToSkeletals)
-						{
-							SkeletonRootUid = MeshUid;
-						}
-						else
-						{
-							//Log an error
-							continue;
-						}
-					}
-					UInterchangeSkeletonFactoryNode* SkeletonFactoryNode = CommonSkeletalMeshesAndAnimationsProperties->CreateSkeletonFactoryNode(BaseNodeContainer, SkeletonRootUid, SourceFileName);
-
-					UInterchangeSkeletalMeshFactoryNode* SkeletalMeshFactoryNode = CreateSkeletalMeshFactoryNode(SkeletonRootUid, { MeshUid });
-					if (SkeletalMeshFactoryNode != nullptr)
-					{
-						SetSkeletalMeshDependencies(SkeletonRootUid, SkeletalMeshFactoryNode);
-						bFoundInstances = true;
-					}
+					CreateSkeletalMeshFactoryForSkeletonRootUid(SkeletonRootUid, { MeshUid });
 				}
-				return bFoundInstances;
 			}
-			return false;
+			else
+			{
+				CreateSkeletalMeshFactoryForSkeletonRootUid(SkeletonRootUid, MeshNodeOrMeshInstanceNodeUids);
+			}
+		}
 #endif
-		};
-		
-		PipelineMeshesUtilities->GetAllSkinnedMeshInstance(MeshUids);
-		CreatePerSkeletonRootUidSkinnedMesh();
 	}
 
 	CreateAssemblyPartDependencies();
@@ -254,7 +187,6 @@ UInterchangeSkeletalMeshFactoryNode* UInterchangeGenericMeshPipeline::CreateSkel
 	{
 		return nullptr;
 	}
-
 	
 	auto GetFirstNodeInfo = [this, &NodeUids](FString& OutFirstMeshNodeUid)->const UInterchangeBaseNode*
 	{
@@ -342,7 +274,7 @@ UInterchangeSkeletalMeshFactoryNode* UInterchangeGenericMeshPipeline::CreateSkel
 		SkeletalMeshFactoryNode->SetCustomKeepSectionsSeparate(CommonMeshesProperties->bKeepSectionsSeparate);
 	}
 
-	SkeletonFactoryNode->SetCustomSkeletalMeshFactoryNodeUid(SkeletalMeshFactoryNode->GetUniqueID());
+	SkeletonFactoryNode->AddCustomSkeletalMeshFactoryNodeUid(SkeletalMeshFactoryNode->GetUniqueID());
 
 	bool bUseTimeZeroAsBindPose = false;
 	SkeletonFactoryNode->GetCustomUseTimeZeroForBindPose(bUseTimeZeroAsBindPose);
diff --git a/Engine/Plugins/Interchange/Runtime/Source/Pipelines/Private/InterchangePipelineMeshesUtilities.cpp b/Engine/Plugins/Interchange/Runtime/Source/Pipelines/Private/InterchangePipelineMeshesUtilities.cpp
--- a/Engine/Plugins/Interchange/Runtime/Source/Pipelines/Private/InterchangePipelineMeshesUtilities.cpp
+++ b/Engine/Plugins/Interchange/Runtime/Source/Pipelines/Private/InterchangePipelineMeshesUtilities.cpp
@@ -955,7 +955,7 @@ void UInterchangePipelineMeshesUtilities::IterateAllMeshInstanceUsingMeshGeometr
 }
 
 void UInterchangePipelineMeshesUtilities::GetCombinedSkinnedMeshInstances(TMap<FString, TArray<FString>>& OutMeshInstanceUidsPerSkeletonRootUid,
-	bool bUseSingleBoneForConvertedMeshes /*= false*/) const
+	bool bUseSingleBoneForConvertedMeshes /*= false*/, bool bVisibleOnly /*= false*/) const
 {
 	if (!ensure(BaseNodeContainer))
 	{
@@ -968,10 +968,14 @@ void UInterchangePipelineMeshesUtilities::GetCombinedSkinnedMeshInstances(TMap<F
 		SourceNode->GetCustomAllowSceneRootAsJoint(bAllowSceneRootAsJoint);
 	}
 
-
 	for (const TPair<FString, FInterchangeMeshInstance>& MeshInstanceUidAndMeshInstance : MeshInstancesPerMeshInstanceUid)
 	{
 		const FInterchangeMeshInstance& MeshInstance = MeshInstanceUidAndMeshInstance.Value;
+		if (bVisibleOnly && !MeshInstance.bIsVisible)
+		{
+			//Skip invisible nodes
+			continue;
+		}
 
 		bool bIsNestedIntoSkeleton = false;
 		if (!CurrentDataContext.IsSkeletalMeshInstance(MeshInstance, BaseNodeContainer, bIsNestedIntoSkeleton))
@@ -1005,7 +1009,8 @@ void UInterchangePipelineMeshesUtilities::GetCombinedSkinnedMeshInstances(TMap<F
 				if (bIsNestedIntoSkeleton || !MeshInstance.bHasMorphTargets)
 				{
 					//Find the deepest joint node
-					if (const UInterchangeSceneNode* SceneNode = Cast<UInterchangeSceneNode>(BaseNodeContainer->GetNode(MeshInstanceUidAndMeshInstance.Key)))
+					const UInterchangeSceneNode* SceneNode = Cast<UInterchangeSceneNode>(BaseNodeContainer->GetNode(MeshInstanceUidAndMeshInstance.Key));
+					if (SceneNode)
 					{
 						FString ParentUid = SceneNode->GetParentUid();
 						FString LastSceneNodeUid = SkeletonRootUid;
diff --git a/Engine/Plugins/Interchange/Runtime/Source/Pipelines/Public/InterchangeGenericMeshPipeline.h b/Engine/Plugins/Interchange/Runtime/Source/Pipelines/Public/InterchangeGenericMeshPipeline.h
--- a/Engine/Plugins/Interchange/Runtime/Source/Pipelines/Public/InterchangeGenericMeshPipeline.h
+++ b/Engine/Plugins/Interchange/Runtime/Source/Pipelines/Public/InterchangeGenericMeshPipeline.h
@@ -31,11 +31,21 @@ class UInterchangeMeshLODContainerNode;
 UENUM(BlueprintType)
 enum class EInterchangeCombineStaticMeshesBehavior : uint8
 {
-	All UMETA(DisplayName="Combine All", Tooltip="Combine all the meshes into a single static mesh"),
+	All UMETA(DisplayName="Combine All", Tooltip="Combine all the meshes into a single Static Mesh"),
 	
 	VisibleOnly UMETA(DisplayName = "Combine Visible Meshes Only", Tooltip="Skip adding the hidden meshes to the combined Static Mesh"),
 
-	DoNotCombine UMETA(DisplayName = "Do Not Combine Static Meshes", Tooltip = "All Static Meshes will be imported as individual assets")
+	DoNotCombine UMETA(DisplayName = "Do Not Combine Static Meshes", Tooltip = "All meshes will be imported as individual Static Mesh assets")
+};
+
+UENUM(BlueprintType)
+enum class EInterchangeCombineSkeletalMeshesBehavior : uint8
+{
+	BySkeleton UMETA(DisplayName="Combine All That Use Same Skeleton", Tooltip="Combine all the meshes that use a specific Skeleton as a single Skeletal Mesh"),
+	
+	BySkeletonVisibleOnly UMETA(DisplayName = "Combine All Visible That Use Same Skeleton", Tooltip="Combine all visible meshes that use a specific Skeleton as a single Skeletal Mesh"),
+
+	DoNotCombine UMETA(DisplayName = "Do Not Combine Skeletal Meshes", Tooltip = "All meshes will be imported as individual Skeletal Mesh assets")
 };
 
 /* Hide drop down will make sure the class is not showing in the class picker */
@@ -201,24 +211,9 @@ public:
 	UPROPERTY()
 	EInterchangeSkeletalMeshContentType LastSkeletalMeshImportContentType;
 
-	UE_DEPRECATED(5.5, "bCombineSkeletalMeshes is no longer used")
-	UFUNCTION(BlueprintGetter, meta = (DeprecatedFunction, DeprecationMessage = "bCombineSkeletalMeshes is no longer used"))
-	bool GetCombineSkeletalMeshes() const
-	{
-		PRAGMA_DISABLE_DEPRECATION_WARNINGS
-		return bCombineSkeletalMeshes_DEPRECATED;
-		PRAGMA_ENABLE_DEPRECATION_WARNINGS
-	}
-
-	UE_DEPRECATED(5.5, "bCombineSkeletalMeshes is no longer used")
-	UFUNCTION(BlueprintSetter, meta = (DeprecatedFunction, DeprecationMessage = "bCombineSkeletalMeshes is no longer used"))
-	void SetCombineSkeletalMeshes(bool InbCombineSkeletalMeshes) {}
-
-	/** If enabled, all skinned mesh nodes that belong to the same skeleton root joint are combined into a single skeletal mesh. */
-	UE_DEPRECATED(5.5, "Please do not access this member. It will be remove in the next version.")
-	UPROPERTY(BlueprintReadWrite, BlueprintGetter = GetCombineSkeletalMeshes, BlueprintSetter = SetCombineSkeletalMeshes, Category = "Skeletal Meshes")
-	bool bCombineSkeletalMeshes_DEPRECATED = true;
-
+	/** Specifies the behavior for combining skeletal meshes into a single skeletal mesh. */
+	UPROPERTY(EditAnywhere, BlueprintReadWrite, Category = "Skeletal Meshes", meta=(DisplayName = "Combine Skeletal Meshes"))
+	EInterchangeCombineSkeletalMeshesBehavior CombineSkeletalMeshesBehavior = EInterchangeCombineSkeletalMeshesBehavior::BySkeleton;
 
 	/** If enabled, imports all morph target shapes found in the source. */
 	UPROPERTY(EditAnywhere, BlueprintReadWrite, Category = "Skeletal Meshes")
diff --git a/Engine/Plugins/Interchange/Runtime/Source/Pipelines/Public/InterchangeGenericScenesPipeline.h b/Engine/Plugins/Interchange/Runtime/Source/Pipelines/Public/InterchangeGenericScenesPipeline.h
--- a/Engine/Plugins/Interchange/Runtime/Source/Pipelines/Public/InterchangeGenericScenesPipeline.h
+++ b/Engine/Plugins/Interchange/Runtime/Source/Pipelines/Public/InterchangeGenericScenesPipeline.h
@@ -117,7 +117,7 @@ protected:
 	/**
 	 * Use to set up the given factory node's attributes after its initialization.
 	 */
-	UE_API virtual void SetUpFactoryNode(UInterchangeActorFactoryNode* ActorFactoryNode, const UInterchangeSceneNode* SceneNode, const UInterchangeBaseNode* TranslatedAssetNode) const;
+	UE_API virtual void SetUpFactoryNode(UInterchangeActorFactoryNode* ActorFactoryNode, const UInterchangeSceneNode* SceneNode, const UInterchangeBaseNode* TranslatedAssetNode, bool bIsRootJoint = false) const;
 
 	UE_API virtual void PopulateSceneNodesPreImport(TArray<UInterchangeSceneNode*>& SceneNodes,
 		TArray<UInterchangeSceneComponentNode*>& SceneComponentNodes) const;
@@ -163,6 +163,9 @@ protected:
 	TSet<FString> Cached_ActiveJointUids;
 
 	FPreviewScene* PreviewScene = nullptr;
+
+private:
+	void CreateActorFactoryNode(const FTransform& GlobalOffsetTransform, const UInterchangeSceneNode* SceneNode, const FString& SkeletalMeshFactoryNodeUid);
 };
 
 
diff --git a/Engine/Plugins/Interchange/Runtime/Source/Pipelines/Public/InterchangePipelineMeshesUtilities.h b/Engine/Plugins/Interchange/Runtime/Source/Pipelines/Public/InterchangePipelineMeshesUtilities.h
--- a/Engine/Plugins/Interchange/Runtime/Source/Pipelines/Public/InterchangePipelineMeshesUtilities.h
+++ b/Engine/Plugins/Interchange/Runtime/Source/Pipelines/Public/InterchangePipelineMeshesUtilities.h
@@ -371,8 +371,10 @@ public:
 	* We cannot create a skinned mesh that has multiple skeleton root nodes. This function returns a combined MeshInstance per skeleton roots.
 	* 
 	* @Param bUseSingleBoneForConvertedMeshes - If true, the skeleton roots assigned to non-skinned meshes will be using a single bone. Otherwise, the skeleton will follow the hierarchy scene node hierarchy.
+	* @Param bVisibleOnly - If true, invisible meshes will be ignored
 	*/
-	UE_API void GetCombinedSkinnedMeshInstances(TMap<FString, TArray<FString>>& OutMeshInstanceUidsPerSkeletonRootUid, bool bUseSingleBoneForConvertedMeshes = false) const;
+	UE_API void GetCombinedSkinnedMeshInstances(TMap<FString, TArray<FString>>& OutMeshInstanceUidsPerSkeletonRootUid,
+		bool bUseSingleBoneForConvertedMeshes = false, bool bVisibleOnly = false) const;
 	
 	/**
 	* Return the skeleton root node UID. This is the UID for a UInterchangeSceneNode that has a "Joint" specialized type.
diff --git a/Engine/Plugins/Tests/InterchangeTests/Source/InterchangeTests/Private/ImportTestFunctions/SkeletonImportTestFunctions.cpp b/Engine/Plugins/Tests/InterchangeTests/Source/InterchangeTests/Private/ImportTestFunctions/SkeletonImportTestFunctions.cpp
new file mode 100644
--- /dev/null
+++ b/Engine/Plugins/Tests/InterchangeTests/Source/InterchangeTests/Private/ImportTestFunctions/SkeletonImportTestFunctions.cpp
@@ -0,0 +1,24 @@
+// Copyright Epic Games, Inc. All Rights Reserved.
+
+#include "ImportTestFunctions/SkeletonImportTestFunctions.h"
+#include "Animation/Skeleton.h"
+
+#include UE_INLINE_GENERATED_CPP_BY_NAME(SkeletonImportTestFunctions)
+
+
+UClass* USkeletonImportTestFunctions::GetAssociatedAssetType() const
+{
+	return USkeleton::StaticClass();
+}
+
+
+FInterchangeTestFunctionResult USkeletonImportTestFunctions::CheckImportedSkeletonCount(const TArray<USkeleton*>& Skeletons, int32 ExpectedNumberOfImportedSkeletons)
+{
+	FInterchangeTestFunctionResult Result;
+	if (Skeletons.Num() != ExpectedNumberOfImportedSkeletons)
+	{
+		Result.AddError(FString::Printf(TEXT("Expected %d skeleton, imported %d."), ExpectedNumberOfImportedSkeletons, Skeletons.Num()));
+	}
+
+	return Result;
+}
\ No newline at end of file
diff --git a/Engine/Plugins/Tests/InterchangeTests/Source/InterchangeTests/Public/ImportTestFunctions/SkeletonImportTestFunctions.h b/Engine/Plugins/Tests/InterchangeTests/Source/InterchangeTests/Public/ImportTestFunctions/SkeletonImportTestFunctions.h
new file mode 100644
--- /dev/null
+++ b/Engine/Plugins/Tests/InterchangeTests/Source/InterchangeTests/Public/ImportTestFunctions/SkeletonImportTestFunctions.h
@@ -0,0 +1,29 @@
+// Copyright Epic Games, Inc. All Rights Reserved.
+
+#pragma once
+
+#include "ImportTestFunctionsBase.h"
+#include "SkeletonImportTestFunctions.generated.h"
+
+#define UE_API INTERCHANGETESTS_API
+
+struct FInterchangeTestFunctionResult;
+class USkeleton;
+
+
+UCLASS(MinimalAPI)
+class USkeletonImportTestFunctions : public UImportTestFunctionsBase
+{
+	GENERATED_BODY()
+
+public:
+
+	// UImportTestFunctionsBase interface
+	UE_API virtual UClass* GetAssociatedAssetType() const override;
+
+	/** Check whether the expected number of skeletal meshes are imported */
+	UFUNCTION(Exec)
+	static UE_API FInterchangeTestFunctionResult CheckImportedSkeletonCount(const TArray<USkeleton*>& Skeletons, int32 ExpectedNumberOfImportedSkeletons);
+};
+
+#undef UE_API
