Interchange AxF [on befalf of Stefan Ivanovski from XRite]:
- Fix specular profile import for Interchange AxF plugin

#14468
#Synced-CL 51290462

diff --git a/Engine/Plugins/Experimental/InterchangeAxF/Runtime/Source/InterchangeAxF/Private/AxFInterchangePipeline.cpp b/Engine/Plugins/Experimental/InterchangeAxF/Runtime/Source/InterchangeAxF/Private/AxFInterchangePipeline.cpp
--- a/Engine/Plugins/Experimental/InterchangeAxF/Runtime/Source/InterchangeAxF/Private/AxFInterchangePipeline.cpp
+++ b/Engine/Plugins/Experimental/InterchangeAxF/Runtime/Source/InterchangeAxF/Private/AxFInterchangePipeline.cpp
@@ -2,13 +2,14 @@
 
 #include "AxFInterchangePipeline.h"
 #include "AxFMaterialObjectNode.h"
-#include "Engine/SpecularProfile.h"
 #include "InterchangeGenericTexturePipeline.h"
-#include "Material/InterchangeMaterialFactory.h"
 #include "InterchangeMaterialFactoryNode.h"
 #include "InterchangePipelineHelper.h"
+#include "InterchangeSpecularProfileFactoryNode.h"
 #include "InterchangeTexture2DFactoryNode.h"
 #include "InterchangeTexture2DNode.h"
+#include "InterchangeSpecularProfileNode.h"
+#include "Material/InterchangeMaterialFactory.h"
 #include "Materials/Material.h"
 #include "Materials/MaterialInstanceConstant.h"
 #include "Nodes/InterchangeSourceNode.h"
@@ -26,71 +27,103 @@ struct FMaterialParameterTypes
 // Map to store the names of the scalar and vector parameters of the predefined materials.
 static TMap<FString, FMaterialParameterTypes> ParameterTypesPerMaterial;
 
-static FString GetParameterNameFromUID(const FString & UID)
+static FString GetParameterNameFromUID(const FString& UID)
 {
-    FString LeftPart, ParameterName;
-    UID.Split("_", &LeftPart, &ParameterName, ESearchCase::IgnoreCase, ESearchDir::FromEnd);
-    return ParameterName;
+	FString LeftPart, ParameterName;
+	UID.Split("_", &LeftPart, &ParameterName, ESearchCase::IgnoreCase, ESearchDir::FromEnd);
+	return ParameterName;
 }
 
 UAxFInterchangePipeline::UAxFInterchangePipeline()
 {
-    TexturePipeline = CreateDefaultSubobject<UInterchangeGenericTexturePipeline>(TEXT("TexturePipeline"));
-    TexturePipeline->bDetectNormalMapTexture = false;
-    TexturePipeline->bAllowNonPowerOfTwo = true;
+	TexturePipeline = CreateDefaultSubobject<UInterchangeGenericTexturePipeline>(TEXT("TexturePipeline"));
+	TexturePipeline->bDetectNormalMapTexture = false;
+	TexturePipeline->bAllowNonPowerOfTwo = true;
 }
 
-void UAxFInterchangePipeline::ExecutePipeline(UInterchangeBaseNodeContainer* BaseNodeContainer, 
-                                              TArray<UInterchangeSourceData*> const& SourceDatas, 
-                                              FString const& ContentBasePath)
+void UAxFInterchangePipeline::ExecutePipeline(UInterchangeBaseNodeContainer* BaseNodeContainer,
+											  TArray<UInterchangeSourceData*> const& SourceDatas,
+											  FString const& ContentBasePath)
 {
-    if (TexturePipeline)
-    {
-        TexturePipeline->ScriptedExecutePipeline(BaseNodeContainer, SourceDatas, ContentBasePath);
-    }
-
-    // Find all translated node we need for this pipeline
-    BaseNodeContainer->IterateNodes(
-        [this](FString const& NodeUid, UInterchangeBaseNode* Node) {
-        switch (Node->GetNodeContainerType())
-        {
-            case EInterchangeNodeContainerType::TranslatedAsset:
-            {
-                if (UAxFMaterialObjectNode* CustomDemoObjectNode = Cast<UAxFMaterialObjectNode>(Node))
-                {
-                    AxFObjectNodes.Add(CustomDemoObjectNode);
-                }
-                else if (UInterchangeTexture2DNode* Texture2DNode = Cast<UInterchangeTexture2DNode>(Node))
-                {
+	if (TexturePipeline)
+	{
+		TexturePipeline->ScriptedExecutePipeline(BaseNodeContainer, SourceDatas, ContentBasePath);
+	}
+
+	TArray<UInterchangeSpecularProfileFactoryNode*> SpecularProfileFactoryNodes;
+
+	// Find all translated and factory nodes we need for this pipeline
+	BaseNodeContainer->IterateNodes(
+		[this, BaseNodeContainer, &SpecularProfileFactoryNodes](FString const& NodeUid, UInterchangeBaseNode* Node) {
+			switch (Node->GetNodeContainerType())
+			{
+			case EInterchangeNodeContainerType::TranslatedAsset:
+			{
+				if (UAxFMaterialObjectNode* CustomDemoObjectNode = Cast<UAxFMaterialObjectNode>(Node))
+				{
+					AxFObjectNodes.Add(CustomDemoObjectNode);
+				}
+				else if (UInterchangeTexture2DNode* Texture2DNode = Cast<UInterchangeTexture2DNode>(Node))
+				{
 					float WidthMM = 0.f, HeightMM = 0.f;
-                    Texture2DNode->GetAttribute(TEXT("WidthMM"), WidthMM);
-                    Texture2DNode->GetAttribute(TEXT("HeightMM"), HeightMM);
+					Texture2DNode->GetAttribute(TEXT("WidthMM"), WidthMM);
+					Texture2DNode->GetAttribute(TEXT("HeightMM"), HeightMM);
 
-                    FString ParameterName = GetParameterNameFromUID(Texture2DNode->GetUniqueID());
-                    FAxFTextureMap& Pair = TexturesMap.FindOrAdd(ParameterName);
-                    Pair.SizeMM = FVector2f{ WidthMM, HeightMM };
+					FString ParameterName = GetParameterNameFromUID(Texture2DNode->GetUniqueID());
+					FAxFTextureMap& Pair = TexturesMap.FindOrAdd(ParameterName);
+					Pair.SizeMM = FVector2f{ WidthMM, HeightMM };
 				}
-            }
-            break;
-            case EInterchangeNodeContainerType::FactoryData:
-            {
-                if (UInterchangeTexture2DFactoryNode* Texture2DFactoryNode = Cast<UInterchangeTexture2DFactoryNode>(Node))
-                {
-                    Texture2DFactoryNodes.Add(Texture2DFactoryNode);
-                }
-            }
-            break;
-            default:;
-        }
-    });
-
-    if (AxFObjectNodes.IsEmpty())
-    {
-        return;
-    }
+				else if (UInterchangeSpecularProfileNode* SpecularProfileNode = Cast<UInterchangeSpecularProfileNode>(Node))
+				{
+					FString SpecProfileUID = SpecularProfileNode->GetUniqueID();
+					SpecProfileUID.ReplaceInline(TEXT("SPN"), TEXT("SP"));
+
+					UInterchangeSpecularProfileFactoryNode* FactoryNode = NewObject<UInterchangeSpecularProfileFactoryNode>(BaseNodeContainer, FName(SpecProfileUID));
+					FactoryNode->InitializeNode(SpecProfileUID, SpecProfileUID, EInterchangeNodeContainerType::FactoryData);
+
+					uint8 SpecularProfileFormat{};
+					if (SpecularProfileNode->GetCustomFormat(SpecularProfileFormat))
+					{
+						FactoryNode->SetCustomFormat(ESpecularProfileFormat(SpecularProfileFormat));
+					}
+
+					FString TextureUID{};
+					if (SpecularProfileNode->GetCustomTexture(TextureUID))
+					{
+						FactoryNode->SetCustomTexture(TextureUID);
+						FactoryNode->AddFactoryDependencyUid(TEXT("Factory_") + TextureUID);
+					}
+
+					SpecularProfileFactoryNodes.Emplace(FactoryNode);
+				}
+			}
+			break;
+			case EInterchangeNodeContainerType::FactoryData:
+			{
+				if (UInterchangeTexture2DFactoryNode* Texture2DFactoryNode = Cast<UInterchangeTexture2DFactoryNode>(Node))
+				{
+					Texture2DFactoryNodes.Add(Texture2DFactoryNode);
+				}
+			}
+			break;
+			default:;
+			}
+		});
+
+	// Defer adding the newly created specular profile factory node to the base node container
+	// because adding it immediately results in a "added during iteration of container" warning.
+	for (UInterchangeSpecularProfileFactoryNode* SpecularProfileFactoryNode : SpecularProfileFactoryNodes)
+	{
+		BaseNodeContainer->AddNode(SpecularProfileFactoryNode);
+	}
+
+	if (AxFObjectNodes.IsEmpty())
+	{
+		return;
+	}
 
 	for (const UAxFMaterialObjectNode* AxFNode : AxFObjectNodes)
-    {
+	{
 		if (!AxFNode)
 		{
 			continue;
@@ -103,7 +136,7 @@ void UAxFInterchangePipeline::ExecutePipeline(UInterchangeBaseNodeContainer* Bas
 		if (Features.Contains(EAxFFeature::CPA2))
 		{
 			BaseMaterialName += TEXT("CPA2");
-            BaseMaterialName += (Features.Contains(EAxFFeature::Flakes) ? TEXT("_Flakes") : TEXT(""));
+			BaseMaterialName += (Features.Contains(EAxFFeature::Flakes) ? TEXT("_Flakes") : TEXT(""));
 		}
 		else
 		{
@@ -127,7 +160,7 @@ void UAxFInterchangePipeline::ExecutePipeline(UInterchangeBaseNodeContainer* Bas
 
 		FString MaterialInstanceUID = TEXT("MI_") + AxFNode->GetDisplayLabel();
 
-        UInterchangeMaterialInstanceFactoryNode* MaterialInstance = NewObject<UInterchangeMaterialInstanceFactoryNode>(BaseNodeContainer, FName(MaterialInstanceUID));
+		UInterchangeMaterialInstanceFactoryNode* MaterialInstance = NewObject<UInterchangeMaterialInstanceFactoryNode>(BaseNodeContainer, FName(MaterialInstanceUID));
 		BaseNodeContainer->SetupNode(MaterialInstance, MaterialInstanceUID, AxFNode->GetDisplayLabel(), EInterchangeNodeContainerType::FactoryData);
 
 		for (const EAxFFeature Feature : Features)
@@ -143,14 +176,14 @@ void UAxFInterchangePipeline::ExecutePipeline(UInterchangeBaseNodeContainer* Bas
 		constexpr float Epsilon = 1e-6;
 		for (FString const& Param : ParameterNames)
 		{
-            if (Param == TEXT("FlakePatchSizeMM"))
-            {
-                const FString XAttributeKey = UInterchangeShaderPortsAPI::MakeInputValueKey(TEXT("WidthMM"));
-                const FString YAttributeKey = UInterchangeShaderPortsAPI::MakeInputValueKey(TEXT("HeightMM"));
+			if (Param == TEXT("FlakePatchSizeMM"))
+			{
+				const FString XAttributeKey = UInterchangeShaderPortsAPI::MakeInputValueKey(TEXT("WidthMM"));
+				const FString YAttributeKey = UInterchangeShaderPortsAPI::MakeInputValueKey(TEXT("HeightMM"));
 				const FLinearColor& FlakePatchSizeMM = ValuesMap[Param];
 				MaterialInstance->AddFloatAttribute(XAttributeKey, FlakePatchSizeMM.R + Epsilon);
 				MaterialInstance->AddFloatAttribute(YAttributeKey, FlakePatchSizeMM.G + Epsilon);
-            }
+			}
 
 			// NOTE: Disable texture path in the case that we have a value for this parameter
 			FString TextureToggleAttributeKey = UInterchangeShaderPortsAPI::MakeInputValueKey(TEXT("Has") + Param + TEXT("Texture"));
@@ -167,12 +200,12 @@ void UAxFInterchangePipeline::ExecutePipeline(UInterchangeBaseNodeContainer* Bas
 				MaterialInstance->AddLinearColorAttribute(AttributeKey, Value);
 			}
 		}
-        
-        if (bUseTriplanarMappingByDefault)
-        {
-            FString GeoUVsToggleAttributeKey = UInterchangeShaderPortsAPI::MakeInputValueKey(TEXT("Use Geo UVs"));
-            MaterialInstance->AddBooleanAttribute(GeoUVsToggleAttributeKey, false);
-        }
+
+		if (bUseTriplanarMappingByDefault)
+		{
+			FString GeoUVsToggleAttributeKey = UInterchangeShaderPortsAPI::MakeInputValueKey(TEXT("Use Geo UVs"));
+			MaterialInstance->AddBooleanAttribute(GeoUVsToggleAttributeKey, false);
+		}
 
 		for (UInterchangeTexture2DFactoryNode* Texture2DFactoryNode : Texture2DFactoryNodes)
 		{
@@ -196,9 +229,9 @@ void UAxFInterchangePipeline::ExecutePipeline(UInterchangeBaseNodeContainer* Bas
 				MaterialInstance->AddStringAttribute(AttributeKey, Texture2DFactoryNode->GetUniqueID());
 			}
 
-            if (Param != TEXT("CPA2ColorTable"))
+			if (Param != TEXT("CPA2ColorTable"))
 			{
-                const FVector2f UVScale = TexturesMap.FindOrAdd(Param).SizeMM;
+				const FVector2f UVScale = TexturesMap.FindOrAdd(Param).SizeMM;
 				FString UVScaleXAttributeKey = UInterchangeShaderPortsAPI::MakeInputValueKey(TEXT("WidthMM"));
 				FString UVScaleYAttributeKey = UInterchangeShaderPortsAPI::MakeInputValueKey(TEXT("HeightMM"));
 				MaterialInstance->AddFloatAttribute(UVScaleXAttributeKey, UVScale.X);
@@ -209,10 +242,16 @@ void UAxFInterchangePipeline::ExecutePipeline(UInterchangeBaseNodeContainer* Bas
 		}
 		Texture2DFactoryNodes.Empty();
 
-        MaterialInstance->SetCustomParent(FullBaseMaterialPath);
-        MaterialInstance->MarkPackageDirty();
-        MaterialInstance->PostEditChange();
-    }
+		for (UInterchangeSpecularProfileFactoryNode* SpecularProfileFactoryNode : SpecularProfileFactoryNodes)
+		{
+			SpecularProfileFactoryNode->SetCustomSubPath(AxFNode->GetDisplayLabel());
+			MaterialInstance->AddFactoryDependencyUid(SpecularProfileFactoryNode->GetUniqueID());
+		}
+
+		MaterialInstance->SetCustomParent(FullBaseMaterialPath);
+		MaterialInstance->MarkPackageDirty();
+		MaterialInstance->PostEditChange();
+	}
 
 	BaseNodeContainer->IterateNodesOfType<UInterchangeFactoryBaseNode>(
 		[ReimportStrategyClosure = ReimportStrategy](const FString& NodeUid, UInterchangeFactoryBaseNode* FactoryNode)
@@ -224,133 +263,133 @@ void UAxFInterchangePipeline::ExecutePipeline(UInterchangeBaseNodeContainer* Bas
 }
 
 void UAxFInterchangePipeline::ExecutePostImportPipeline(
-    UInterchangeBaseNodeContainer const* BaseNodeContainer, FString const& NodeKey, UObject* CreatedAsset,
-    bool bIsAReimport)
+	UInterchangeBaseNodeContainer const* BaseNodeContainer, FString const& NodeKey, UObject* CreatedAsset,
+	bool bIsAReimport)
 {
-    if (TexturePipeline)
-    {
-        TexturePipeline->ScriptedExecutePostImportPipeline(BaseNodeContainer, NodeKey, CreatedAsset, bIsAReimport);
-    }
-
-    UMaterialInstanceConstant* Instance = Cast<UMaterialInstanceConstant>(CreatedAsset);
-    if (Instance != nullptr && !AxFObjectNodes.IsEmpty())
-    {
+	if (TexturePipeline)
+	{
+		TexturePipeline->ScriptedExecutePostImportPipeline(BaseNodeContainer, NodeKey, CreatedAsset, bIsAReimport);
+	}
+
+	UMaterialInstanceConstant* Instance = Cast<UMaterialInstanceConstant>(CreatedAsset);
+	if (Instance != nullptr && !AxFObjectNodes.IsEmpty())
+	{
 		TObjectPtr<UAxFMaterialObjectNode>* FoundData = AxFObjectNodes.FindByPredicate(
-            [Instance](UAxFMaterialObjectNode const* Node) { return Node->GetDisplayLabel() == Instance->GetName(); });
+			[Instance](UAxFMaterialObjectNode const* Node) { return Node->GetDisplayLabel() == Instance->GetName(); });
 
-        if (FoundData == nullptr)
-        {
-            UE_LOG(
-                LogAxFInterchangePipeline, Error,
-                TEXT("Couldn't find appropriate AxF translated node to create Material Instance."))
-            return;
-        }
+		if (FoundData == nullptr)
+		{
+			UE_LOG(
+				LogAxFInterchangePipeline, Error,
+				TEXT("Couldn't find appropriate AxF translated node to create Material Instance."))
+				return;
+		}
 
 		UAxFMaterialObjectNode* AxFData = *FoundData;
 
 		const TArray<EAxFFeature>& Features = AxFData->PayloadData.UsedFeatures;
-		
-        if (SpecularProfile != nullptr)
-        {
-            Instance->bOverrideSpecularProfile = true;
-            Instance->SpecularProfileOverride = SpecularProfile;
-            SpecularProfile = nullptr;
-        }
-
-        if (!Instance->MarkPackageDirty())
-        {
-            UE_LOG(
-                LogAxFInterchangePipeline, Warning,
-                TEXT("Failed to mark package as dirty. Some changes to the material instance %s might not be saved."),
-                *Instance->GetName());
-        }
-
-        Instance->PostEditChange();
-
-        AxFObjectNodes.RemoveSingle(AxFData);
-    }
-
-    Super::ExecutePostImportPipeline(BaseNodeContainer, NodeKey, CreatedAsset, bIsAReimport);
+
+		if (SpecularProfile != nullptr)
+		{
+			Instance->bOverrideSpecularProfile = true;
+			Instance->SpecularProfileOverride = SpecularProfile;
+			SpecularProfile = nullptr;
+		}
+
+		if (!Instance->MarkPackageDirty())
+		{
+			UE_LOG(
+				LogAxFInterchangePipeline, Warning,
+				TEXT("Failed to mark package as dirty. Some changes to the material instance %s might not be saved."),
+				*Instance->GetName());
+		}
+
+		Instance->PostEditChange();
+
+		AxFObjectNodes.RemoveSingle(AxFData);
+	}
+
+	Super::ExecutePostImportPipeline(BaseNodeContainer, NodeKey, CreatedAsset, bIsAReimport);
 }
 
 void UAxFInterchangePipeline::ExecutePostFactoryPipeline(
-    UInterchangeBaseNodeContainer const* BaseNodeContainer, FString const& NodeKey, UObject* CreatedAsset,
-    bool bIsAReimport)
+	UInterchangeBaseNodeContainer const* BaseNodeContainer, FString const& NodeKey, UObject* CreatedAsset,
+	bool bIsAReimport)
 {
-    if (TexturePipeline)
-    {
-        TexturePipeline->ScriptedExecutePostFactoryPipeline(BaseNodeContainer, NodeKey, CreatedAsset, bIsAReimport);
-    }
-
-    FString ParameterName = GetParameterNameFromUID(NodeKey);
-
-    UTexture2D* Texture2D = Cast<UTexture2D>(CreatedAsset);
-    if (Texture2D != nullptr)
-    {
-        ETextureSourceFormat Format = Texture2D->Source.GetFormat();
-        if (Format == TSF_BGRA8 || Format == TSF_G8 || Format == TSF_RGBA16)
-        {
-            Texture2D->CompressionSettings = TC_BC7;
-        }
-        else if (Format == TSF_RGBA32F)
-        {
-            Texture2D->CompressionSettings = TC_HDR_Compressed;
-        }
-        else if (Format == TSF_R32F || Format == TSF_R16F)
-        {
-            Texture2D->CompressionSettings = TC_HalfFloat;
-        }
-        else
-        {
-            UE_LOG(
-                LogAxFInterchangePipeline, Warning,
-                TEXT("Unsupported texture source format! Can't set specific compression settings for %s"),
-                *Texture2D->GetName());
-        }
-
-        Texture2D->MipGenSettings = TMGS_FromTextureGroup;
-        Texture2D->PowerOfTwoMode = ETexturePowerOfTwoSetting::StretchToPowerOfTwo;
-        if (ParameterName == TEXT("Normal") || ParameterName == TEXT("ClearCoatNormal"))
-        {
-            Texture2D->bFlipGreenChannel = true;
-            Texture2D->CompressionSettings = TC_Normalmap;
-        }
-        Texture2D->SRGB = 0;
-
-        if (!Texture2D->MarkPackageDirty())
-        {
-            UE_LOG(
-                LogAxFInterchangePipeline, Warning,
-                TEXT("Failed to mark package as dirty. Some changes to the Texture2D %s might not be saved."),
-                *Texture2D->GetName());
-        }
-        Texture2D->PostEditChange();
+	if (TexturePipeline)
+	{
+		TexturePipeline->ScriptedExecutePostFactoryPipeline(BaseNodeContainer, NodeKey, CreatedAsset, bIsAReimport);
+	}
+
+	FString ParameterName = GetParameterNameFromUID(NodeKey);
+
+	UTexture2D* Texture2D = Cast<UTexture2D>(CreatedAsset);
+	if (Texture2D != nullptr)
+	{
+		ETextureSourceFormat Format = Texture2D->Source.GetFormat();
+		if (Format == TSF_BGRA8 || Format == TSF_G8 || Format == TSF_RGBA16)
+		{
+			Texture2D->CompressionSettings = TC_BC7;
+		}
+		else if (Format == TSF_RGBA32F)
+		{
+			Texture2D->CompressionSettings = TC_HDR_Compressed;
+		}
+		else if (Format == TSF_R32F || Format == TSF_R16F)
+		{
+			Texture2D->CompressionSettings = TC_HalfFloat;
+		}
+		else
+		{
+			UE_LOG(
+				LogAxFInterchangePipeline, Warning,
+				TEXT("Unsupported texture source format! Can't set specific compression settings for %s"),
+				*Texture2D->GetName());
+		}
+
+		Texture2D->MipGenSettings = TMGS_FromTextureGroup;
+		Texture2D->PowerOfTwoMode = ETexturePowerOfTwoSetting::StretchToPowerOfTwo;
+		if (ParameterName == TEXT("Normal") || ParameterName == TEXT("ClearCoatNormal"))
+		{
+			Texture2D->bFlipGreenChannel = true;
+			Texture2D->CompressionSettings = TC_Normalmap;
+		}
+		Texture2D->SRGB = 0;
+
+		if (!Texture2D->MarkPackageDirty())
+		{
+			UE_LOG(
+				LogAxFInterchangePipeline, Warning,
+				TEXT("Failed to mark package as dirty. Some changes to the Texture2D %s might not be saved."),
+				*Texture2D->GetName());
+		}
+		Texture2D->PostEditChange();
 
 		FAxFTextureMap& Pair = TexturesMap.FindOrAdd(ParameterName);
 		Pair.Texture = Texture2D;
-    }
+	}
 
-    USpecularProfile* Profile = Cast<USpecularProfile>(CreatedAsset);
-    if (Profile != nullptr)
-    {
-        SpecularProfile = Profile;
-    }
+	USpecularProfile* Profile = Cast<USpecularProfile>(CreatedAsset);
+	if (Profile != nullptr)
+	{
+		SpecularProfile = Profile;
+	}
 
-    Super::ExecutePostFactoryPipeline(BaseNodeContainer, NodeKey, CreatedAsset, bIsAReimport);
+	Super::ExecutePostFactoryPipeline(BaseNodeContainer, NodeKey, CreatedAsset, bIsAReimport);
 }
 
 void UAxFInterchangePipeline::AdjustSettingsForContext(FInterchangePipelineContextParams const& ContextParams)
 {
-    if (TexturePipeline)
-    {
-        TexturePipeline->AdjustSettingsForContext(ContextParams);
-    }
+	if (TexturePipeline)
+	{
+		TexturePipeline->AdjustSettingsForContext(ContextParams);
+	}
 
-    // Check if the base materials are loaded
+	// Check if the base materials are loaded
 	static bool bRunOnce = []()
 		{
 			const FString BasePath = TEXT("/InterchangeAxFAssets/BaseMaterials/");
-            const FString BaseMaterialNames[] = { 
+			const FString BaseMaterialNames[] = {
 				TEXT("SVBRDF"),
 				TEXT("SVBRDF_Alpha"),
 				TEXT("SVBRDF_POM"),
@@ -408,44 +447,44 @@ void UAxFInterchangePipeline::AdjustSettingsForContext(FInterchangePipelineConte
 			return true;
 		}();
 
-    Super::AdjustSettingsForContext(ContextParams);
+	Super::AdjustSettingsForContext(ContextParams);
 }
 
 inline void UAxFInterchangePipeline::PreDialogCleanup(FName const PipelineStackName)
 {
-    if (TexturePipeline)
-    {
-        TexturePipeline->PreDialogCleanup(PipelineStackName);
-    }
+	if (TexturePipeline)
+	{
+		TexturePipeline->PreDialogCleanup(PipelineStackName);
+	}
 }
 
 bool UAxFInterchangePipeline::IsSettingsAreValid(TOptional<FText>& OutInvalidReason) const
 {
-    if (TexturePipeline && !TexturePipeline->IsSettingsAreValid(OutInvalidReason))
-    {
-        return false;
-    }
+	if (TexturePipeline && !TexturePipeline->IsSettingsAreValid(OutInvalidReason))
+	{
+		return false;
+	}
 
-    return Super::IsSettingsAreValid(OutInvalidReason);
+	return Super::IsSettingsAreValid(OutInvalidReason);
 }
 
 inline void UAxFInterchangePipeline::FilterPropertiesFromTranslatedData(
-    UInterchangeBaseNodeContainer* InBaseNodeContainer)
+	UInterchangeBaseNodeContainer* InBaseNodeContainer)
 {
-    if (TexturePipeline)
-    {
-        TexturePipeline->FilterPropertiesFromTranslatedData(InBaseNodeContainer);
-    }
+	if (TexturePipeline)
+	{
+		TexturePipeline->FilterPropertiesFromTranslatedData(InBaseNodeContainer);
+	}
 
-    Super::FilterPropertiesFromTranslatedData(InBaseNodeContainer);
+	Super::FilterPropertiesFromTranslatedData(InBaseNodeContainer);
 }
 
 void UAxFInterchangePipeline::GetSupportAssetClasses(TArray<UClass*>& PipelineSupportAssetClasses) const
 {
-    if (TexturePipeline)
-    {
-        TexturePipeline->GetSupportAssetClasses(PipelineSupportAssetClasses);
-    }
+	if (TexturePipeline)
+	{
+		TexturePipeline->GetSupportAssetClasses(PipelineSupportAssetClasses);
+	}
 
-    PipelineSupportAssetClasses.Add(UMaterialInstance::StaticClass());
-}
+	PipelineSupportAssetClasses.Add(UMaterialInstance::StaticClass());
+}
\ No newline at end of file
diff --git a/Engine/Plugins/Experimental/InterchangeAxF/Runtime/Source/InterchangeAxF/Private/AxFTranslator.cpp b/Engine/Plugins/Experimental/InterchangeAxF/Runtime/Source/InterchangeAxF/Private/AxFTranslator.cpp
--- a/Engine/Plugins/Experimental/InterchangeAxF/Runtime/Source/InterchangeAxF/Private/AxFTranslator.cpp
+++ b/Engine/Plugins/Experimental/InterchangeAxF/Runtime/Source/InterchangeAxF/Private/AxFTranslator.cpp
@@ -5,7 +5,7 @@
 #include "Engine/SpecularProfile.h"
 #include "Engine/RendererSettings.h"
 #include "ImageUtils.h"
-#include "InterchangeSpecularProfileFactoryNode.h"
+#include "InterchangeSpecularProfileNode.h"
 #include "InterchangeTexture2DNode.h"
 #include "Modules/ModuleManager.h"
 #include "InterchangeImportModule.h"
@@ -18,10 +18,10 @@ DEFINE_LOG_CATEGORY(LogAxFTranslator);
 
 static bool GInterchangeEnableAxFImport = true;
 static FAutoConsoleVariableRef CCvarInterchangeEnableAxFImport(
-    TEXT("Interchange.FeatureFlags.Import.AxF"),
-    GInterchangeEnableAxFImport, 
-    TEXT("Whether AxF Interchange support is enabled."),
-    ECVF_Default);
+	TEXT("Interchange.FeatureFlags.Import.AxF"),
+	GInterchangeEnableAxFImport,
+	TEXT("Whether AxF Interchange support is enabled."),
+	ECVF_Default);
 
 bool UAxFTranslator::CanImportSourceData(const UInterchangeSourceData* InSourceData) const
 {
@@ -42,53 +42,53 @@ bool UAxFTranslator::CanImportSourceData(const UInterchangeSourceData* InSourceD
 
 TArray<FString> UAxFTranslator::GetSupportedFormats() const
 {
-    TArray<FString> Formats;
-    if (GInterchangeEnableAxFImport && IInterchangeImportModule::Get().IsSubstrateEnabled())
-    {
-        Formats.Add(TEXT("axf2ue;Intermediary AxF resource file"));
-    }
+	TArray<FString> Formats;
+	if (GInterchangeEnableAxFImport && IInterchangeImportModule::Get().IsSubstrateEnabled())
+	{
+		Formats.Add(TEXT("axf2ue;Intermediary AxF resource file"));
+	}
 
-    return Formats;
+	return Formats;
 }
 
-bool UAxFTranslator::Translate(UInterchangeBaseNodeContainer &BaseNodeContainer) const
+bool UAxFTranslator::Translate(UInterchangeBaseNodeContainer& BaseNodeContainer) const
 {
-    FString Filename = SourceData->GetFilename();
-
-    FString Path, Name, Extension;
-    FPaths::Split(Filename, Path, Name, Extension);
-
-    FPaths::NormalizeFilename(Filename);
-    if (!FPaths::FileExists(Filename))
-    {
-        return false;
-    }
-
-    FString DisplayLabel = FPaths::GetBaseFilename(Filename);
-    FString NodeUID(Filename);
-    UAxFMaterialObjectNode *AxFMaterialObjectNode = NewObject<UAxFMaterialObjectNode>(&BaseNodeContainer);
-    if (!ensure(AxFMaterialObjectNode))
-    {
-        return false;
-    }
-
-    AxFMaterialObjectNode->InitializeAxFMaterialObjectNode(NodeUID, DisplayLabel);
-    AxFMaterialObjectNode->SetPayloadKey(Filename);
-
-    FString JSONString;
-    if (!FFileHelper::LoadFileToString(JSONString, *SourceData->GetFilename()))
-    {
-        UE_LOG(LogAxFTranslator, Error, TEXT("Failed to read JSON file at %s"), *SourceData->GetFilename());
-        return {};
-    }
-
-    TSharedPtr<FJsonObject> JSON;
-    TSharedRef<TJsonReader<>> Reader = TJsonReaderFactory<>::Create(JSONString);
-    if (!FJsonSerializer::Deserialize(Reader, JSON))
-    {
-        UE_LOG(LogAxFTranslator, Error, TEXT("Failed to deserialize JSON read from file at %s"), *SourceData->GetFilename());
-        return {};
-    }
+	FString Filename = SourceData->GetFilename();
+
+	FString Path, Name, Extension;
+	FPaths::Split(Filename, Path, Name, Extension);
+
+	FPaths::NormalizeFilename(Filename);
+	if (!FPaths::FileExists(Filename))
+	{
+		return false;
+	}
+
+	FString DisplayLabel = FPaths::GetBaseFilename(Filename);
+	FString NodeUID(Filename);
+	UAxFMaterialObjectNode* AxFMaterialObjectNode = NewObject<UAxFMaterialObjectNode>(&BaseNodeContainer);
+	if (!ensure(AxFMaterialObjectNode))
+	{
+		return false;
+	}
+
+	AxFMaterialObjectNode->InitializeAxFMaterialObjectNode(NodeUID, DisplayLabel);
+	AxFMaterialObjectNode->SetPayloadKey(Filename);
+
+	FString JSONString;
+	if (!FFileHelper::LoadFileToString(JSONString, *SourceData->GetFilename()))
+	{
+		UE_LOG(LogAxFTranslator, Error, TEXT("Failed to read JSON file at %s"), *SourceData->GetFilename());
+		return {};
+	}
+
+	TSharedPtr<FJsonObject> JSON;
+	TSharedRef<TJsonReader<>> Reader = TJsonReaderFactory<>::Create(JSONString);
+	if (!FJsonSerializer::Deserialize(Reader, JSON))
+	{
+		UE_LOG(LogAxFTranslator, Error, TEXT("Failed to deserialize JSON read from file at %s"), *SourceData->GetFilename());
+		return {};
+	}
 
 	static const TMap<FString, FString> ParameterNameMapping{
 		{TEXT("alpha"), TEXT("Alpha")},
@@ -115,64 +115,64 @@ bool UAxFTranslator::Translate(UInterchangeBaseNodeContainer &BaseNodeContainer)
 		{TEXT("transmissionColor"), TEXT("TransmissionColor")},
 	};
 
-    UE::Interchange::FAxFMaterialObjectData& AxFData = AxFMaterialObjectNode->PayloadData;
-	for (const TPair<FString, FString>& Pair: ParameterNameMapping)
-    {
+	UE::Interchange::FAxFMaterialObjectData& AxFData = AxFMaterialObjectNode->PayloadData;
+	for (const TPair<FString, FString>& Pair : ParameterNameMapping)
+	{
 		const FString& ParameterName = Pair.Key;
-        const FString & UEParameterName =  Pair.Value;
-        bool ParameterPresent = false;
+		const FString& UEParameterName = Pair.Value;
+		bool ParameterPresent = false;
 
-        if (const TArray<TSharedPtr<FJsonValue>> *TopLevelArray; JSON->TryGetArrayField(ParameterName, TopLevelArray))
-        {
-            ParameterPresent = true;
+		if (const TArray<TSharedPtr<FJsonValue>>* TopLevelArray; JSON->TryGetArrayField(ParameterName, TopLevelArray))
+		{
+			ParameterPresent = true;
 
-            TArray<FLinearColor> ColorsFromJSON;
-            FLinearColor ValuesFromJSON(0.0f, 0.0f, 0.0f);
+			TArray<FLinearColor> ColorsFromJSON;
+			FLinearColor ValuesFromJSON(0.0f, 0.0f, 0.0f);
 			const int32 TopLevelArrayNum = FMath::Clamp(TopLevelArray->Num(), 0, 4);
 
-            for (size_t Index = 0; Index < TopLevelArrayNum; ++Index)
-            {
+			for (size_t Index = 0; Index < TopLevelArrayNum; ++Index)
+			{
 				TSharedPtr<FJsonValue> JsonValue = (*TopLevelArray)[Index];
-                if (JsonValue->Type != EJson::Number && JsonValue->Type != EJson::Array)
-                {
-                    UE_LOG(LogAxFTranslator, Error, TEXT("Invalid value type in array found"));
-                    return {};
-                }
-
-                if (JsonValue->Type == EJson::Number)
-                {
-                    ValuesFromJSON.Component(Index) = JsonValue->AsNumber();
-                }
-            }
-
-            if (TopLevelArray->Num() == 1)
-            {
-                ValuesFromJSON = FLinearColor{ValuesFromJSON.R, ValuesFromJSON.R, ValuesFromJSON.R};
-            }
-
-            AxFData.ValuesMap.Add(UEParameterName) = ValuesFromJSON;
-        }
-        else if (const TSharedPtr<FJsonObject> *ValueObject; JSON->TryGetObjectField(ParameterName, ValueObject))
-        {
-            ParameterPresent = true;
-
-            float WidthMM{}, HeightMM{};
-            if (!ValueObject->Get()->TryGetNumberField(TEXT("widthMM"), WidthMM))
-            {
-                UE_LOG(LogAxFTranslator, Error, TEXT("Failed to find '%s' property in AxF2UE file."), TEXT("widthMM"));
-                return false;
-            }
-            if (!ValueObject->Get()->TryGetNumberField(TEXT("heightMM"), HeightMM))
-            {
-                UE_LOG(LogAxFTranslator, Error, TEXT("Failed to find '%s' property in AxF2UE file."), TEXT("heightMM"));
-                return false;
-            }
-
-            if (UEParameterName == TEXT("FlakePatchSizeMM"))
-            {
-                AxFData.ValuesMap.Add(UEParameterName) = FLinearColor{ WidthMM, HeightMM, 0.0f };
-                continue;
-            }
+				if (JsonValue->Type != EJson::Number && JsonValue->Type != EJson::Array)
+				{
+					UE_LOG(LogAxFTranslator, Error, TEXT("Invalid value type in array found"));
+					return {};
+				}
+
+				if (JsonValue->Type == EJson::Number)
+				{
+					ValuesFromJSON.Component(Index) = JsonValue->AsNumber();
+				}
+			}
+
+			if (TopLevelArray->Num() == 1)
+			{
+				ValuesFromJSON = FLinearColor{ ValuesFromJSON.R, ValuesFromJSON.R, ValuesFromJSON.R };
+			}
+
+			AxFData.ValuesMap.Add(UEParameterName) = ValuesFromJSON;
+		}
+		else if (const TSharedPtr<FJsonObject>* ValueObject; JSON->TryGetObjectField(ParameterName, ValueObject))
+		{
+			ParameterPresent = true;
+
+			float WidthMM{}, HeightMM{};
+			if (!ValueObject->Get()->TryGetNumberField(TEXT("widthMM"), WidthMM))
+			{
+				UE_LOG(LogAxFTranslator, Error, TEXT("Failed to find '%s' property in AxF2UE file."), TEXT("widthMM"));
+				return false;
+			}
+			if (!ValueObject->Get()->TryGetNumberField(TEXT("heightMM"), HeightMM))
+			{
+				UE_LOG(LogAxFTranslator, Error, TEXT("Failed to find '%s' property in AxF2UE file."), TEXT("heightMM"));
+				return false;
+			}
+
+			if (UEParameterName == TEXT("FlakePatchSizeMM"))
+			{
+				AxFData.ValuesMap.Add(UEParameterName) = FLinearColor{ WidthMM, HeightMM, 0.0f };
+				continue;
+			}
 
 			FString TexturePath{};
 			if (!ValueObject->Get()->TryGetStringField(TEXT("path"), TexturePath))
@@ -181,198 +181,197 @@ bool UAxFTranslator::Translate(UInterchangeBaseNodeContainer &BaseNodeContainer)
 				return false;
 			}
 
-            FString TextureUID = TEXT("T_") + DisplayLabel + TEXT("_") + UEParameterName;
-            UInterchangeTexture2DNode *Texture2D = NewObject<UInterchangeTexture2DNode>(&BaseNodeContainer, FName(TextureUID));
+			FString TextureUID = TEXT("T_") + DisplayLabel + TEXT("_") + UEParameterName;
+			UInterchangeTexture2DNode* Texture2D = NewObject<UInterchangeTexture2DNode>(&BaseNodeContainer, FName(TextureUID));
 			BaseNodeContainer.SetupNode(Texture2D, TextureUID, TextureUID, EInterchangeNodeContainerType::TranslatedAsset);
-            Texture2D->SetPayLoadKey(Path + "|" + Name + "." + Extension + "|" + TexturePath);
-            Texture2D->AddFloatAttribute(TEXT("WidthMM"), WidthMM);
-            Texture2D->AddFloatAttribute(TEXT("HeightMM"), HeightMM);
-
-            if (UEParameterName == TEXT("CPA2ColorTable"))
-            {
-                FString SpecProfileUID = TEXT("SP_") + DisplayLabel;
-                UInterchangeSpecularProfileFactoryNode* SpecularProfileNode = NewObject<UInterchangeSpecularProfileFactoryNode>(&BaseNodeContainer, FName(SpecProfileUID));
+			Texture2D->SetPayLoadKey(Path + "|" + Name + "." + Extension + "|" + TexturePath);
+			Texture2D->AddFloatAttribute(TEXT("WidthMM"), WidthMM);
+			Texture2D->AddFloatAttribute(TEXT("HeightMM"), HeightMM);
+
+			if (UEParameterName == TEXT("CPA2ColorTable"))
+			{
+				FString SpecProfileUID = TEXT("SPN_") + DisplayLabel;
+				UInterchangeSpecularProfileNode* SpecularProfileNode = NewObject<UInterchangeSpecularProfileNode>(&BaseNodeContainer, FName(SpecProfileUID));
 				BaseNodeContainer.SetupNode(SpecularProfileNode, SpecProfileUID, SpecProfileUID, EInterchangeNodeContainerType::TranslatedAsset);
-                SpecularProfileNode->SetCustomFormat(ESpecularProfileFormat::HalfVector);
-                SpecularProfileNode->SetCustomTexture(TextureUID);
-                SpecularProfileNode->SetCustomSubPath(AxFMaterialObjectNode->GetDisplayLabel());
-            }
-        }
-        else if (bool ValueBool; JSON->TryGetBoolField(ParameterName, ValueBool)) 
-        {
-            ParameterPresent = ValueBool;
-        }
-
-        // If this parameter is present as either a value [array]
-        // or a texture, check for optional AxF features to be used.
-        if (ParameterPresent)
-        {
-            // Clear Coat
-            if (UEParameterName == TEXT("ClearCoatTransmissionColor") || UEParameterName == TEXT("ClearCoatRoughness") ||
+				SpecularProfileNode->SetCustomFormat((uint8)ESpecularProfileFormat::HalfVector);
+				SpecularProfileNode->SetCustomTexture(TextureUID);
+			}
+		}
+		else if (bool ValueBool; JSON->TryGetBoolField(ParameterName, ValueBool))
+		{
+			ParameterPresent = ValueBool;
+		}
+
+		// If this parameter is present as either a value [array]
+		// or a texture, check for optional AxF features to be used.
+		if (ParameterPresent)
+		{
+			// Clear Coat
+			if (UEParameterName == TEXT("ClearCoatTransmissionColor") || UEParameterName == TEXT("ClearCoatRoughness") ||
 				UEParameterName == TEXT("ClearCoatIOR") || UEParameterName == TEXT("ClearCoatNormal"))
-            {
-                AxFData.UsedFeatures.AddUnique(EAxFFeature::ClearCoat);
-            }
-
-            // Sheen
-            if (UEParameterName == TEXT("SheenColor") || UEParameterName == TEXT("SheenRoughness"))
-            {
-                AxFData.UsedFeatures.AddUnique(EAxFFeature::Sheen);
-            }
-
-            // Anisotropy
-            if (UEParameterName == TEXT("Roughness") && AxFData.ValuesMap.Contains(UEParameterName))
-            {
-                if (const FLinearColor& RoughnessValues = AxFData.ValuesMap[UEParameterName]; RoughnessValues.R != RoughnessValues.G)
-                {
-                    AxFData.UsedFeatures.AddUnique(EAxFFeature::Anisotropy);
-                }
-            }
-            if (UEParameterName == TEXT("AnisotropicRotation"))
-            {
-                AxFData.UsedFeatures.AddUnique(EAxFFeature::Anisotropy);
-            }
-
-            // CPA2
-            if (UEParameterName == TEXT("CPA2ColorTable"))
-            {
-                AxFData.UsedFeatures.AddUnique(EAxFFeature::CPA2);
-            }
-
-            // Alpha transparency
-            if (UEParameterName == TEXT("Alpha"))
-            {
-                AxFData.UsedFeatures.AddUnique(EAxFFeature::Alpha);
-            }
-
-            // Transmission
-            if (UEParameterName == TEXT("TransmissionColor"))
-            {
-                AxFData.UsedFeatures.AddUnique(EAxFFeature::Transmission);
-            }
-
-            // Heightmap
-            if (UEParameterName == TEXT("HeightMap") || UEParameterName == TEXT("HeightStrength") || UEParameterName == TEXT("HeightReferencePlane"))
-            {
-                AxFData.UsedFeatures.AddUnique(EAxFFeature::Heightmap);
-            }
-
-            // Flakes
-            if (UEParameterName == TEXT("GlintDensity") || UEParameterName == TEXT("GlintUVScale") ||
+			{
+				AxFData.UsedFeatures.AddUnique(EAxFFeature::ClearCoat);
+			}
+
+			// Sheen
+			if (UEParameterName == TEXT("SheenColor") || UEParameterName == TEXT("SheenRoughness"))
+			{
+				AxFData.UsedFeatures.AddUnique(EAxFFeature::Sheen);
+			}
+
+			// Anisotropy
+			if (UEParameterName == TEXT("Roughness") && AxFData.ValuesMap.Contains(UEParameterName))
+			{
+				if (const FLinearColor& RoughnessValues = AxFData.ValuesMap[UEParameterName]; RoughnessValues.R != RoughnessValues.G)
+				{
+					AxFData.UsedFeatures.AddUnique(EAxFFeature::Anisotropy);
+				}
+			}
+			if (UEParameterName == TEXT("AnisotropicRotation"))
+			{
+				AxFData.UsedFeatures.AddUnique(EAxFFeature::Anisotropy);
+			}
+
+			// CPA2
+			if (UEParameterName == TEXT("CPA2ColorTable"))
+			{
+				AxFData.UsedFeatures.AddUnique(EAxFFeature::CPA2);
+			}
+
+			// Alpha transparency
+			if (UEParameterName == TEXT("Alpha"))
+			{
+				AxFData.UsedFeatures.AddUnique(EAxFFeature::Alpha);
+			}
+
+			// Transmission
+			if (UEParameterName == TEXT("TransmissionColor"))
+			{
+				AxFData.UsedFeatures.AddUnique(EAxFFeature::Transmission);
+			}
+
+			// Heightmap
+			if (UEParameterName == TEXT("HeightMap") || UEParameterName == TEXT("HeightStrength") || UEParameterName == TEXT("HeightReferencePlane"))
+			{
+				AxFData.UsedFeatures.AddUnique(EAxFFeature::Heightmap);
+			}
+
+			// Flakes
+			if (UEParameterName == TEXT("GlintDensity") || UEParameterName == TEXT("GlintUVScale") ||
 				UEParameterName == TEXT("GlintRoughness") || UEParameterName == TEXT("GlintBrightness"))
-            {
-                AxFData.UsedFeatures.AddUnique(EAxFFeature::Flakes);
-            }
-        }
-    }
+			{
+				AxFData.UsedFeatures.AddUnique(EAxFFeature::Flakes);
+			}
+		}
+	}
 
-    BaseNodeContainer.AddNode(AxFMaterialObjectNode);
+	BaseNodeContainer.AddNode(AxFMaterialObjectNode);
 
-    return true;
+	return true;
 }
 
 TOptional<UE::Interchange::FImportImage>
-UAxFTranslator::GetTexturePayloadData(const FString &PayloadKey, TOptional<FString> &AlternateTexturePath) const
+UAxFTranslator::GetTexturePayloadData(const FString& PayloadKey, TOptional<FString>& AlternateTexturePath) const
 {
-    FString BasePath, RightPart;
-    if (!PayloadKey.Split("|", &BasePath, &RightPart))
-    {
-        UE_LOG(LogAxFTranslator, Error, TEXT("Failed to properly split Interchange payload key!"));
-        return {};
-    }
-
-    FString JSONFileName, LocalTexturePath;
-    if (!RightPart.Split("|", &JSONFileName, &LocalTexturePath))
-    {
-        UE_LOG(LogAxFTranslator, Error, TEXT("Failed to properly split Interchange payload key!"));
-        return {};
-    }
-
-    // TODO: Abstract into separate function
-    FString JSONString;
-    FString JSONFilePath = FPaths::Combine(BasePath, JSONFileName);
+	FString BasePath, RightPart;
+	if (!PayloadKey.Split("|", &BasePath, &RightPart))
+	{
+		UE_LOG(LogAxFTranslator, Error, TEXT("Failed to properly split Interchange payload key!"));
+		return {};
+	}
+
+	FString JSONFileName, LocalTexturePath;
+	if (!RightPart.Split("|", &JSONFileName, &LocalTexturePath))
+	{
+		UE_LOG(LogAxFTranslator, Error, TEXT("Failed to properly split Interchange payload key!"));
+		return {};
+	}
+
+	// TODO: Abstract into separate function
+	FString JSONString;
+	FString JSONFilePath = FPaths::Combine(BasePath, JSONFileName);
 	if (!FPaths::IsUnderDirectory(JSONFilePath, BasePath))
 	{
 		UE_LOG(LogAxFTranslator, Error, TEXT("JSON file %s is not under the directory %s"), *JSONFileName, *BasePath);
 		return {};
 	}
-    if (!FFileHelper::LoadFileToString(JSONString, *JSONFilePath))
-    {
-        UE_LOG(LogAxFTranslator, Error, TEXT("Failed to read JSON file at %s"), *JSONFilePath);
-        return {};
-    }
-
-    TSharedPtr<FJsonObject> JSON;
-    TSharedRef<TJsonReader<>> Reader = TJsonReaderFactory<>::Create(JSONString);
-    if (!FJsonSerializer::Deserialize(Reader, JSON))
-    {
-        UE_LOG(LogAxFTranslator, Error, TEXT("Failed to deserialize JSON read from file at %s"), *JSONFilePath);
-        return {};
-    }
-
-    TSharedPtr<TArray<uint8>> ByteArray(new TArray<uint8>());
-    FString TexturePath = FPaths::Combine(BasePath, LocalTexturePath);
+	if (!FFileHelper::LoadFileToString(JSONString, *JSONFilePath))
+	{
+		UE_LOG(LogAxFTranslator, Error, TEXT("Failed to read JSON file at %s"), *JSONFilePath);
+		return {};
+	}
+
+	TSharedPtr<FJsonObject> JSON;
+	TSharedRef<TJsonReader<>> Reader = TJsonReaderFactory<>::Create(JSONString);
+	if (!FJsonSerializer::Deserialize(Reader, JSON))
+	{
+		UE_LOG(LogAxFTranslator, Error, TEXT("Failed to deserialize JSON read from file at %s"), *JSONFilePath);
+		return {};
+	}
+
+	TSharedPtr<TArray<uint8>> ByteArray(new TArray<uint8>());
+	FString TexturePath = FPaths::Combine(BasePath, LocalTexturePath);
 	if (!FPaths::IsUnderDirectory(TexturePath, BasePath))
 	{
 		UE_LOG(LogAxFTranslator, Error, TEXT("Texture file file %s is not under the directory %s"), *TexturePath, *BasePath);
 		return {};
 	}
-    if (!FFileHelper::LoadFileToArray(*ByteArray, *TexturePath))
-    {
-        UE_LOG(LogAxFTranslator, Error, TEXT("Failed to load texture from file: %s"), *TexturePath);
-        return {};
-    }
-
-    FImage RawImage;
-    if (FImageUtils::DecompressImage(ByteArray->GetData(), ByteArray->Num(), RawImage))
-    {
-        int32 Width = RawImage.GetWidth();
-        int32 Height = RawImage.GetHeight();
-
-        ETextureSourceFormat SourceFormat;
-        if (RawImage.Format == ERawImageFormat::G8)
-        {
-            SourceFormat = TSF_G8;
-        }
-        else if (RawImage.Format == ERawImageFormat::BGRA8)
-        {
-            SourceFormat = TSF_BGRA8;
-        }
-        else if (RawImage.Format == ERawImageFormat::R32F)
-        {
-            SourceFormat = TSF_R32F;
-        }
-        else if (RawImage.Format == ERawImageFormat::RGBA32F)
-        {
-            SourceFormat = TSF_RGBA32F;
-        }
-        else
-        {
-            UE_LOG(LogAxFTranslator, Warning, TEXT("Unsupported number of channels!"));
-            return {};
-        }
+	if (!FFileHelper::LoadFileToArray(*ByteArray, *TexturePath))
+	{
+		UE_LOG(LogAxFTranslator, Error, TEXT("Failed to load texture from file: %s"), *TexturePath);
+		return {};
+	}
+
+	FImage RawImage;
+	if (FImageUtils::DecompressImage(ByteArray->GetData(), ByteArray->Num(), RawImage))
+	{
+		int32 Width = RawImage.GetWidth();
+		int32 Height = RawImage.GetHeight();
+
+		ETextureSourceFormat SourceFormat;
+		if (RawImage.Format == ERawImageFormat::G8)
+		{
+			SourceFormat = TSF_G8;
+		}
+		else if (RawImage.Format == ERawImageFormat::BGRA8)
+		{
+			SourceFormat = TSF_BGRA8;
+		}
+		else if (RawImage.Format == ERawImageFormat::R32F)
+		{
+			SourceFormat = TSF_R32F;
+		}
+		else if (RawImage.Format == ERawImageFormat::RGBA32F)
+		{
+			SourceFormat = TSF_RGBA32F;
+		}
+		else
+		{
+			UE_LOG(LogAxFTranslator, Warning, TEXT("Unsupported number of channels!"));
+			return {};
+		}
 
 		UE::Interchange::FImportImage Image;
-        Image.Init2DWithParams(Width, Height, SourceFormat, false);
-        FMemory::Memcpy(Image.RawData.GetData(), RawImage.RawData.GetData(), RawImage.RawData.NumBytes());
+		Image.Init2DWithParams(Width, Height, SourceFormat, false);
+		FMemory::Memcpy(Image.RawData.GetData(), RawImage.RawData.GetData(), RawImage.RawData.NumBytes());
 
 		return Image;
-    }
-    else
-    {
-        UE_LOG(LogAxFTranslator, Error, TEXT("Texture asset could not be created: %s"), *TexturePath);
+	}
+	else
+	{
+		UE_LOG(LogAxFTranslator, Error, TEXT("Texture asset could not be created: %s"), *TexturePath);
 		return {};
-    }
+	}
 }
 
 bool UAxFTranslator::SupportCompressedTexturePayloadData() const
 {
-    return IInterchangeTexturePayloadInterface::SupportCompressedTexturePayloadData();
+	return IInterchangeTexturePayloadInterface::SupportCompressedTexturePayloadData();
 }
 
 TOptional<UE::Interchange::FImportImage>
-UAxFTranslator::GetCompressedTexturePayloadData(const FString &PayloadKey,
-                                                TOptional<FString> &AlternateTexturePath) const
+UAxFTranslator::GetCompressedTexturePayloadData(const FString& PayloadKey,
+												TOptional<FString>& AlternateTexturePath) const
 {
-    return IInterchangeTexturePayloadInterface::GetCompressedTexturePayloadData(PayloadKey, AlternateTexturePath);
+	return IInterchangeTexturePayloadInterface::GetCompressedTexturePayloadData(PayloadKey, AlternateTexturePath);
 }
\ No newline at end of file
diff --git a/Engine/Plugins/Experimental/InterchangeAxF/Runtime/Source/InterchangeAxF/Public/AxFInterchangePipeline.h b/Engine/Plugins/Experimental/InterchangeAxF/Runtime/Source/InterchangeAxF/Public/AxFInterchangePipeline.h
--- a/Engine/Plugins/Experimental/InterchangeAxF/Runtime/Source/InterchangeAxF/Public/AxFInterchangePipeline.h
+++ b/Engine/Plugins/Experimental/InterchangeAxF/Runtime/Source/InterchangeAxF/Public/AxFInterchangePipeline.h
@@ -3,9 +3,10 @@
 #pragma once
 
 #include "CoreMinimal.h"
+#include "Engine/SpecularProfile.h"
 #include "InterchangeGenericTexturePipeline.h"
 #include "InterchangePipelineBase.h"
-#include "InterchangeSpecularProfileFactoryNode.h"
+
 #include "AxFInterchangePipeline.generated.h"
 
 #define UE_API INTERCHANGEAXF_API
@@ -33,15 +34,15 @@ struct FAxFTextureMap
 UCLASS(BlueprintType)
 class UE_API UAxFInterchangePipeline : public UInterchangePipelineBase
 {
-    GENERATED_BODY()
+	GENERATED_BODY()
 
 public:
-    UAxFInterchangePipeline();
+	UAxFInterchangePipeline();
 
-    UPROPERTY(
-        EditAnywhere, BlueprintReadWrite, Category = "General",
-        meta = (StandAlonePipelineProperty = "True", PipelineInternalEditionData = "True"))
-    FString PipelineDisplayName;
+	UPROPERTY(
+		EditAnywhere, BlueprintReadWrite, Category = "General",
+		meta = (StandAlonePipelineProperty = "True", PipelineInternalEditionData = "True"))
+	FString PipelineDisplayName;
 
 	/* Set the reimport strategy. */
 	UPROPERTY(
@@ -49,50 +50,50 @@ public:
 		meta = (AdjustPipelineAndRefreshDetailOnChange = "True", PipelineInternalEditionData = "True"))
 	EReimportStrategyFlags ReimportStrategy = EReimportStrategyFlags::ApplyNoProperties;
 
-    UPROPERTY(EditAnywhere, BlueprintReadWrite, Category = "AxF")
-    bool bUseTriplanarMappingByDefault = true;
+	UPROPERTY(EditAnywhere, BlueprintReadWrite, Category = "AxF")
+	bool bUseTriplanarMappingByDefault = true;
 
 protected:
-    virtual void ExecutePipeline(
-        UInterchangeBaseNodeContainer* BaseNodeContainer, TArray<UInterchangeSourceData*> const& SourceDatas,
-        FString const& ContentBasePath) override;
+	virtual void ExecutePipeline(
+		UInterchangeBaseNodeContainer* BaseNodeContainer, TArray<UInterchangeSourceData*> const& SourceDatas,
+		FString const& ContentBasePath) override;
 
-    virtual void ExecutePostImportPipeline(
-        UInterchangeBaseNodeContainer const* BaseNodeContainer, FString const& NodeKey, UObject* CreatedAsset,
-        bool bIsAReimport) override;
+	virtual void ExecutePostImportPipeline(
+		UInterchangeBaseNodeContainer const* BaseNodeContainer, FString const& NodeKey, UObject* CreatedAsset,
+		bool bIsAReimport) override;
 
-    virtual void ExecutePostFactoryPipeline(
-        UInterchangeBaseNodeContainer const* BaseNodeContainer, FString const& NodeKey, UObject* CreatedAsset,
-        bool bIsAReimport) override;
+	virtual void ExecutePostFactoryPipeline(
+		UInterchangeBaseNodeContainer const* BaseNodeContainer, FString const& NodeKey, UObject* CreatedAsset,
+		bool bIsAReimport) override;
 
-    virtual void AdjustSettingsForContext(FInterchangePipelineContextParams const& ContextParams) override;
+	virtual void AdjustSettingsForContext(FInterchangePipelineContextParams const& ContextParams) override;
 
-    virtual void PreDialogCleanup(FName const PipelineStackName) override;
+	virtual void PreDialogCleanup(FName const PipelineStackName) override;
 
-    virtual bool IsSettingsAreValid(TOptional<FText>& OutInvalidReason) const override;
+	virtual bool IsSettingsAreValid(TOptional<FText>& OutInvalidReason) const override;
 
-    virtual void FilterPropertiesFromTranslatedData(UInterchangeBaseNodeContainer* InBaseNodeContainer) override;
+	virtual void FilterPropertiesFromTranslatedData(UInterchangeBaseNodeContainer* InBaseNodeContainer) override;
 
-    virtual void GetSupportAssetClasses(TArray<UClass*>& PipelineSupportAssetClasses) const override;
+	virtual void GetSupportAssetClasses(TArray<UClass*>& PipelineSupportAssetClasses) const override;
 
 private:
 	UPROPERTY(Transient)
 	TArray<TObjectPtr<UAxFMaterialObjectNode>> AxFObjectNodes;
 
-    // Texture handling
+	// Texture handling
 	UPROPERTY(Transient)
 	TArray<TObjectPtr<UInterchangeTexture2DFactoryNode>> Texture2DFactoryNodes;
 
 	UPROPERTY()
-    TMap<FString, FAxFTextureMap> TexturesMap;
+	TMap<FString, FAxFTextureMap> TexturesMap;
 
-    // Specular profile handling
+	// Specular profile handling
 	UPROPERTY()
-    TObjectPtr<USpecularProfile> SpecularProfile;
+	TObjectPtr<USpecularProfile> SpecularProfile;
 
-    // Generic texture pipeline to be used to create the necessary texture assets the materials use
+	// Generic texture pipeline to be used to create the necessary texture assets the materials use
 	UPROPERTY()
-    TObjectPtr<UInterchangeGenericTexturePipeline> TexturePipeline;
+	TObjectPtr<UInterchangeGenericTexturePipeline> TexturePipeline;
 };
 
-#undef UE_API
+#undef UE_API
\ No newline at end of file
