[Interchange][FBX] ufbx - override materials(for actor)

- fixes Editor.Interchange.Import.Level.FBX.SKM_Arm_anim1

#Synced-CL 51336880

diff --git a/Engine/Plugins/Interchange/Runtime/Source/Parsers/Fbx/Private/Ufbx/UfbxScene.cpp b/Engine/Plugins/Interchange/Runtime/Source/Parsers/Fbx/Private/Ufbx/UfbxScene.cpp
--- a/Engine/Plugins/Interchange/Runtime/Source/Parsers/Fbx/Private/Ufbx/UfbxScene.cpp
+++ b/Engine/Plugins/Interchange/Runtime/Source/Parsers/Fbx/Private/Ufbx/UfbxScene.cpp
@@ -253,6 +253,28 @@ namespace UE::Interchange::Private
 			{
 				SceneNode->SetCustomAssetInstanceUid(MeshNode->GetUniqueID());
 			}
+
+			// Add override materials
+			int32 MaterialCount = InNode.materials.count;
+			TMap<ufbx_material*, int32> UniqueSlotNames;
+			UniqueSlotNames.Reserve(MaterialCount);
+			for (int32 MaterialIndex = 0; MaterialIndex < MaterialCount; ++MaterialIndex)
+			{
+				if (ufbx_material* SurfaceMaterial = InNode.materials[MaterialIndex])
+				{
+					FString MaterialUid = Parser.GetMaterialUid(*SurfaceMaterial);
+
+					int32& SlotMaterialCount = UniqueSlotNames.FindOrAdd(SurfaceMaterial);
+
+					FString MaterialSlotName = Parser.GetMaterialSlotName(SurfaceMaterial).ToString();
+					if (SlotMaterialCount > 0)
+					{
+						MaterialSlotName += TEXT("_Section") + FString::FromInt(SlotMaterialCount);
+					}
+					SceneNode->SetSlotMaterialDependencyUid(MaterialSlotName, MaterialUid);
+					SlotMaterialCount++;
+				}
+			}
 		}
 		if (UInterchangeJointNode* JointNode = Cast<UInterchangeJointNode>(SceneNode))
 		{
