[Interchange][FBX] ufbx - added animations for attributes, fixes camera parameters animatiton

#Synced-CL 51500539

diff --git a/Engine/Plugins/Interchange/Runtime/Source/Parsers/Fbx/Private/Ufbx/UfbxAnimation.cpp b/Engine/Plugins/Interchange/Runtime/Source/Parsers/Fbx/Private/Ufbx/UfbxAnimation.cpp
--- a/Engine/Plugins/Interchange/Runtime/Source/Parsers/Fbx/Private/Ufbx/UfbxAnimation.cpp
+++ b/Engine/Plugins/Interchange/Runtime/Source/Parsers/Fbx/Private/Ufbx/UfbxAnimation.cpp
@@ -319,8 +319,6 @@ namespace UE::Interchange::Private
 		const TMap<const ufbx_blend_shape*, const ufbx_mesh*> MeshPerBlendShape = Meshes.GetBlendShapes();
 		double FrameRate = Parser.Scene->settings.frames_per_second;
 
-		TArray<FString> TransformAnimTrackNodeUids;
-
 		int32 AnimStackCount = Parser.Scene->anim_stacks.count;
 		for (int32 AnimStackIndex = 0; AnimStackIndex < AnimStackCount; ++AnimStackIndex)
 		{
@@ -597,7 +595,6 @@ namespace UE::Interchange::Private
 							// #ufbx_todo: UE::Interchange::Private::ProcessCustomAttributes
 							// ProcessCustomAttributes(Parser, Node, TransformAnimTrackNode);
 
-							TransformAnimTrackNodeUids.Add(TransformAnimTrackNode->GetUniqueID());
 						}
 
 					}
@@ -606,8 +603,14 @@ namespace UE::Interchange::Private
 			}
 		}
 
+		TArray<FString> AnimTrackNodeUids;
+		NodeContainer.IterateNodesOfType<UInterchangeAnimationTrackNode>([&](const FString& NodeUid, UInterchangeAnimationTrackNode* TransformAnimationTrackNode)
+			{
+				AnimTrackNodeUids.Add(NodeUid);
+			});
+
 		//Only one Track Set Node per fbx file:
-		if (TransformAnimTrackNodeUids.Num() > 0)
+		if (AnimTrackNodeUids.Num() > 0)
 		{
 			FString Name;
 			if (UInterchangeSceneNode** RootSceneNodeFound = Parser.ElementIdToSceneNode.Find(Parser.Scene->root_node->element_id))
@@ -625,12 +628,11 @@ namespace UE::Interchange::Private
 
 			NodeContainer.SetupNode(TrackSetNode, AnimTrackSetNodeUid, AnimTrackSetNodeDisplayLabel, EInterchangeNodeContainerType::TranslatedAsset);
 
-			for (const FString& TransformAnimTrackNodeUid : TransformAnimTrackNodeUids)
+			for (const FString& AimTrackNodeUid : AnimTrackNodeUids)
 			{
-				TrackSetNode->AddCustomAnimationTrackUid(TransformAnimTrackNodeUid);
+				TrackSetNode->AddCustomAnimationTrackUid(AimTrackNodeUid);
 			}
 		}
-
 	}
 
 	bool FUfbxAnimation::AddNodeAttributeCurvesAnimation(FUfbxParser& Parser
@@ -642,7 +644,12 @@ namespace UE::Interchange::Private
 		, FString& OutCurveName)
 	{
 		FString PropertyName = Convert::ToUnrealString(Prop.name);
-		const FString PayLoadKey = NodeUid + PropertyName + TEXT("_AnimationPayloadKey");
+		FString PayLoadKey = NodeUid + PropertyName + TEXT("_AnimationPayloadKey");
+		if (Parser.PayloadContexts.Contains(PayLoadKey))
+		{
+			// Guard against multiple Props with the same name
+			PayLoadKey = NodeUid + PropertyName + FString::FromInt(AnimValue.element_id) + TEXT("_AnimationPayloadKey");
+		}
 		if (!ensure(!Parser.PayloadContexts.Contains(PayLoadKey)))
 		{
 			return false;
diff --git a/Engine/Plugins/Interchange/Runtime/Source/Parsers/Fbx/Private/Ufbx/UfbxScene.cpp b/Engine/Plugins/Interchange/Runtime/Source/Parsers/Fbx/Private/Ufbx/UfbxScene.cpp
--- a/Engine/Plugins/Interchange/Runtime/Source/Parsers/Fbx/Private/Ufbx/UfbxScene.cpp
+++ b/Engine/Plugins/Interchange/Runtime/Source/Parsers/Fbx/Private/Ufbx/UfbxScene.cpp
@@ -591,6 +591,49 @@ namespace UE::Interchange::Private
 				Parser.ConvertProperty(Prop, SceneNode, PayloadKey);
 			}
 		}
+
+		for (const ufbx_element* Attrib : InNode.all_attribs)
+		{
+			for (const ufbx_prop& Prop : Attrib->props.props)
+			{
+				if (SceneUtils::IsPropTypeSupported(Prop.type) && Prop.flags & UFBX_PROP_FLAG_ANIMATED)
+				{
+					for (const ufbx_connection& Connection : Attrib->connections_dst)
+					{
+						// Property and its connection have same name string pointer in ufbx
+						if (Connection.dst_prop.data == Prop.name.data)
+						{
+							if (ufbx_anim_value* AnimValue = ufbx_as_anim_value(Connection.src))
+							{
+								FString CurveName;
+								TOptional<bool> bIsStepCurve;
+								TOptional<FString> PayloadKey;
+								if (FUfbxAnimation::AddNodeAttributeCurvesAnimation(Parser, SceneNode->GetUniqueID(), Prop, *AnimValue, PayloadKey, bIsStepCurve, CurveName))
+								{
+									EInterchangeAnimationPayLoadType PayloadType = bIsStepCurve.GetValue() ? EInterchangeAnimationPayLoadType::STEPCURVE : EInterchangeAnimationPayLoadType::CURVE;
+
+									const UInterchangeFbxSettings* InterchangeFbxSettings = GetDefault<UInterchangeFbxSettings>();
+									if (EInterchangePropertyTracks PropertyTrack = InterchangeFbxSettings->GetPropertyTrack(CurveName); PropertyTrack != EInterchangePropertyTracks::None)
+									{
+										UInterchangeAnimationTrackNode* AnimTrackNode = NewObject< UInterchangeAnimationTrackNode >(&NodeContainer);
+										const FString AnimTrackNodeName = SceneNode->GetDisplayLabel() + CurveName;
+										const FString AnimTrackNodeUid = TEXT("\\AnimationTrack\\") + AnimTrackNodeName;
+
+										NodeContainer.SetupNode(AnimTrackNode, AnimTrackNodeUid, AnimTrackNodeName, EInterchangeNodeContainerType::TranslatedAsset);
+
+										AnimTrackNode->SetCustomActorDependencyUid(SceneNode->GetUniqueID());
+										AnimTrackNode->SetCustomAnimationPayloadKey(PayloadKey.GetValue(), PayloadType);
+										AnimTrackNode->SetCustomPropertyTrack(PropertyTrack);
+									}
+
+									SceneNode->SetAnimationCurveTypeForCurveName(CurveName, PayloadType);
+								}
+							}
+						}
+					}
+				}
+			}
+		}
 	}
 
 	void FUfbxScene::ProcessNodes(UInterchangeBaseNodeContainer& NodeContainer)
