Interchange - Fixed potential ensure during import of rigid animations

#Synced-CL 51561121

diff --git a/Engine/Plugins/Interchange/Runtime/Source/Parsers/Fbx/Private/FbxAnimation.cpp b/Engine/Plugins/Interchange/Runtime/Source/Parsers/Fbx/Private/FbxAnimation.cpp
--- a/Engine/Plugins/Interchange/Runtime/Source/Parsers/Fbx/Private/FbxAnimation.cpp
+++ b/Engine/Plugins/Interchange/Runtime/Source/Parsers/Fbx/Private/FbxAnimation.cpp
@@ -896,13 +896,15 @@ namespace UE::Interchange::Private
 
 	void FFbxAnimation::AddRigidTransformAnimation(FFbxParser& Parser
 		, FbxNode* Node
+		, const int32& AnimationIndex
 		, FbxAnimCurveNode* TranslationCurveNode
 		, FbxAnimCurveNode* RotationCurveNode
 		, FbxAnimCurveNode* ScaleCurveNode
 		, TMap<FString, TSharedPtr<FPayloadContextBase>>& PayloadContexts
 		, TOptional<FString>& OutPayloadKey)
 	{
-		const FString PayLoadKey = Parser.GetFbxHelper()->GetFbxNodeHierarchyName(Node) + TEXT("_RigidAnimationPayloadKey");
+		const FString PayloadKeySuffix = FString::Printf(TEXT("_RigidAnimationPayloadKey_%d"), AnimationIndex);
+		const FString PayLoadKey = Parser.GetFbxHelper()->GetFbxNodeHierarchyName(Node) + PayloadKeySuffix;
 		if (ensure(!PayloadContexts.Contains(PayLoadKey)))
 		{
 			TSharedPtr<FAnimationPayloadContext> AnimPayload = MakeShared<FAnimationPayloadContext>();
diff --git a/Engine/Plugins/Interchange/Runtime/Source/Parsers/Fbx/Private/FbxAnimation.h b/Engine/Plugins/Interchange/Runtime/Source/Parsers/Fbx/Private/FbxAnimation.h
--- a/Engine/Plugins/Interchange/Runtime/Source/Parsers/Fbx/Private/FbxAnimation.h
+++ b/Engine/Plugins/Interchange/Runtime/Source/Parsers/Fbx/Private/FbxAnimation.h
@@ -122,6 +122,7 @@ namespace UE::Interchange::Private
 
 		static void AddRigidTransformAnimation(FFbxParser& Parser
 			, FbxNode* Node
+			, const int32& AnimationIndex
 			, FbxAnimCurveNode* TranslationCurveNode
 			, FbxAnimCurveNode* RotationCurveNode
 			, FbxAnimCurveNode* ScaleCurveNode
diff --git a/Engine/Plugins/Interchange/Runtime/Source/Parsers/Fbx/Private/FbxScene.cpp b/Engine/Plugins/Interchange/Runtime/Source/Parsers/Fbx/Private/FbxScene.cpp
--- a/Engine/Plugins/Interchange/Runtime/Source/Parsers/Fbx/Private/FbxScene.cpp
+++ b/Engine/Plugins/Interchange/Runtime/Source/Parsers/Fbx/Private/FbxScene.cpp
@@ -838,7 +838,8 @@ namespace UE
 			}
 
 			void FFbxScene::AddRigidAnimation(FbxNode* Node
-				, FbxAnimStack* AnimationStack
+				, const FbxScene* SDKScene
+				, const int32& AnimationIndex
 				, const UInterchangeSceneNode* UnrealNode
 				, UInterchangeBaseNodeContainer& NodeContainer
 				, TMap<FString, TSharedPtr<FPayloadContextBase, ESPMode::ThreadSafe>>& PayloadContexts) const
@@ -847,6 +848,8 @@ namespace UE
 				FbxAnimCurveNode* RotationCurveNode = nullptr;
 				FbxAnimCurveNode* ScaleCurveNode = nullptr;
 
+				FbxAnimStack* AnimationStack = (FbxAnimStack*)SDKScene->GetSrcObject<FbxAnimStack>(AnimationIndex);
+
 				FbxProperty Property = Node->GetFirstProperty();
 				while (Property.IsValid())
 				{
@@ -902,7 +905,7 @@ namespace UE
 				{
 					TOptional<FString> PayloadKey;
 
-					FFbxAnimation::AddRigidTransformAnimation(Parser, Node, TranlsationCurveNode, RotationCurveNode, ScaleCurveNode, PayloadContexts, PayloadKey);
+					FFbxAnimation::AddRigidTransformAnimation(Parser, Node, AnimationIndex, TranlsationCurveNode, RotationCurveNode, ScaleCurveNode, PayloadContexts, PayloadKey);
 
 					if (PayloadKey.IsSet())
 					{
@@ -1178,10 +1181,8 @@ namespace UE
 						}
 					}
 
-					FbxAnimStack* CurrentAnimationStack = (FbxAnimStack*)SDKScene->GetSrcObject<FbxAnimStack>(AnimationIndex);
-					AddRigidAnimation(Node, CurrentAnimationStack, UnrealNode, NodeContainer, PayloadContexts);
+					AddRigidAnimation(Node, SDKScene, AnimationIndex, UnrealNode, NodeContainer, PayloadContexts);
 				}
-			
 
 				const int32 ChildCount = Node->GetChildCount();
 				for (int32 ChildIndex = 0; ChildIndex < ChildCount; ++ChildIndex)
diff --git a/Engine/Plugins/Interchange/Runtime/Source/Parsers/Fbx/Private/FbxScene.h b/Engine/Plugins/Interchange/Runtime/Source/Parsers/Fbx/Private/FbxScene.h
--- a/Engine/Plugins/Interchange/Runtime/Source/Parsers/Fbx/Private/FbxScene.h
+++ b/Engine/Plugins/Interchange/Runtime/Source/Parsers/Fbx/Private/FbxScene.h
@@ -81,7 +81,8 @@ namespace UE
 				T* CreateUnrealSceneNodeInternal(UInterchangeBaseNodeContainer& NodeContainer, const FString& NodeName, const FString& NodeUID, const FString& ParentNodeUID);
 
 				void AddRigidAnimation(FbxNode* Node
-					, FbxAnimStack* AnimationStack
+					, const FbxScene* SDKScene
+					, const int32& AnimationIndex
 					, const UInterchangeSceneNode* UnrealNode
 					, UInterchangeBaseNodeContainer& NodeContainer
 					, TMap<FString, TSharedPtr<FPayloadContextBase, ESPMode::ThreadSafe>>& PayloadContexts) const;
