Improved Material import options
- Deprecated 'DoNotSearch' from the EInterchangeMaterialSearchLocation enum
- Added the boolean property, bReuseExistingMaterials, to better reflects that no search for material asset should be performed
- Added EditCondition to properties to better reflect what a user can do and can't do.
- Reordered the Material pipeline properties to group properties by target, i.e. material reuse, material creation.
- Updated the glTF asset pipelines to hide the new property
- Added a new entry to FInterchangeCustomVersion to properly update the bReuseExistingMaterials property based on the previous value stored in SearchLocation

#rnx
#virtualized
#Synced-CL 51542889

diff --git a/Engine/Plugins/Interchange/Editor/Source/InterchangeEditor/Private/InterchangeFbxAssetImportDataConverter.cpp b/Engine/Plugins/Interchange/Editor/Source/InterchangeEditor/Private/InterchangeFbxAssetImportDataConverter.cpp
--- a/Engine/Plugins/Interchange/Editor/Source/InterchangeEditor/Private/InterchangeFbxAssetImportDataConverter.cpp
+++ b/Engine/Plugins/Interchange/Editor/Source/InterchangeEditor/Private/InterchangeFbxAssetImportDataConverter.cpp
@@ -627,23 +627,27 @@ namespace UE::Interchange::Private
 
 		//Material Options
 		DestinationData->bImportMaterials = GenericAssetPipeline->MaterialPipeline->bImportMaterials;
-		switch (GenericAssetPipeline->MaterialPipeline->SearchLocation)
+		if (!GenericAssetPipeline->MaterialPipeline->bReuseExistingMaterials)
 		{
-		case EInterchangeMaterialSearchLocation::Local:
-			DestinationData->TextureImportData->MaterialSearchLocation = EMaterialSearchLocation::Local;
-			break;
-		case EInterchangeMaterialSearchLocation::UnderParent:
-			DestinationData->TextureImportData->MaterialSearchLocation = EMaterialSearchLocation::UnderParent;
-			break;
-		case EInterchangeMaterialSearchLocation::UnderRoot:
-			DestinationData->TextureImportData->MaterialSearchLocation = EMaterialSearchLocation::UnderRoot;
-			break;
-		case EInterchangeMaterialSearchLocation::AllAssets:
-			DestinationData->TextureImportData->MaterialSearchLocation = EMaterialSearchLocation::AllAssets;
-			break;
-		case EInterchangeMaterialSearchLocation::DoNotSearch:
 			DestinationData->TextureImportData->MaterialSearchLocation = EMaterialSearchLocation::DoNotSearch;
-			break;
+		}
+		else
+		{
+			switch (GenericAssetPipeline->MaterialPipeline->SearchLocation)
+			{
+			case EInterchangeMaterialSearchLocation::Local:
+				DestinationData->TextureImportData->MaterialSearchLocation = EMaterialSearchLocation::Local;
+				break;
+			case EInterchangeMaterialSearchLocation::UnderParent:
+				DestinationData->TextureImportData->MaterialSearchLocation = EMaterialSearchLocation::UnderParent;
+				break;
+			case EInterchangeMaterialSearchLocation::UnderRoot:
+				DestinationData->TextureImportData->MaterialSearchLocation = EMaterialSearchLocation::UnderRoot;
+				break;
+			case EInterchangeMaterialSearchLocation::AllAssets:
+				DestinationData->TextureImportData->MaterialSearchLocation = EMaterialSearchLocation::AllAssets;
+				break;
+			}
 		}
 
 		if (GenericAssetPipeline->MaterialPipeline->ParentMaterial.IsAsset())
@@ -819,6 +823,7 @@ namespace UE::Interchange::Private
 
 		//Material Options
 		GenericAssetPipeline->MaterialPipeline->bImportMaterials = FbxImportUI->bImportMaterials;
+		GenericAssetPipeline->MaterialPipeline->bReuseExistingMaterials = true;
 		switch (FbxImportUI->TextureImportData->MaterialSearchLocation)
 		{
 			case EMaterialSearchLocation::Local:
@@ -834,7 +839,12 @@ namespace UE::Interchange::Private
 				GenericAssetPipeline->MaterialPipeline->SearchLocation = EInterchangeMaterialSearchLocation::AllAssets;
 				break;
 			case EMaterialSearchLocation::DoNotSearch:
-				GenericAssetPipeline->MaterialPipeline->SearchLocation = EInterchangeMaterialSearchLocation::DoNotSearch;
+				{
+					GenericAssetPipeline->MaterialPipeline->bReuseExistingMaterials = false;
+				}
+				break;
+			default:
+				ensure(false);
 				break;
 		}
 		if (FbxImportUI->TextureImportData->BaseMaterialName.IsAsset())
diff --git a/Engine/Plugins/Interchange/Runtime/Source/Pipelines/Private/InterchangeGenericMaterialPipeline.cpp b/Engine/Plugins/Interchange/Runtime/Source/Pipelines/Private/InterchangeGenericMaterialPipeline.cpp
--- a/Engine/Plugins/Interchange/Runtime/Source/Pipelines/Private/InterchangeGenericMaterialPipeline.cpp
+++ b/Engine/Plugins/Interchange/Runtime/Source/Pipelines/Private/InterchangeGenericMaterialPipeline.cpp
@@ -274,11 +274,6 @@ namespace UE::Interchange::InterchangeGenericMaterialPipeline::Private
 
 	UMaterialInterface* FindExistingMaterialFromSearchLocation(const FString& MaterialFullName, const FString& ContentPath, EInterchangeMaterialSearchLocation SearchLocation)
 	{
-		if (SearchLocation == EInterchangeMaterialSearchLocation::DoNotSearch)
-		{
-			return nullptr;
-		}
-
 		//Search in memory
 		UMaterialInterface* FoundMaterial = nullptr;
 		//We search only in memory for search in local folder.
@@ -619,6 +614,21 @@ UInterchangeGenericMaterialPipeline::UInterchangeGenericMaterialPipeline()
 	SparseVolumeTexturePipeline = CreateDefaultSubobject<UInterchangeSparseVolumeTexturePipeline>("SparseVolumeTexturePipeline");
 }
 
+void UInterchangeGenericMaterialPipeline::Serialize(FArchive& Ar)
+{
+	Super::Serialize(Ar);
+
+	Ar.UsingCustomVersion(FInterchangeCustomVersion::GUID);
+	if (Ar.IsLoading())
+	{
+		if (Ar.CustomVer(FInterchangeCustomVersion::GUID) < FInterchangeCustomVersion::AdditionOfReuseExistingMaterialsBool)
+		{
+			UpdateSearchLocationProperty();
+		}
+	}
+
+}
+
 FString UInterchangeGenericMaterialPipeline::GetPipelineCategory(UClass* AssetClass)
 {
 	return TEXT("Materials");
@@ -641,6 +651,8 @@ void UInterchangeGenericMaterialPipeline::PreDialogCleanup(const FName PipelineS
 	{
 		SaveSettings(PipelineStackName);
 	}
+
+	UpdateSearchLocationProperty();
 }
 
 bool UInterchangeGenericMaterialPipeline::IsSettingsAreValid(TOptional<FText>& OutInvalidReason) const
@@ -682,7 +694,7 @@ void UInterchangeGenericMaterialPipeline::AdjustSettingsForContext(const FInterc
 	{
 		bImportMaterials = false;
 		HideCategories.Add(UInterchangeGenericMaterialPipeline::GetPipelineCategory(nullptr));
-		SearchLocation = EInterchangeMaterialSearchLocation::DoNotSearch;
+		bReuseExistingMaterials = false;
 	}
 
 	if (UInterchangePipelineBase* OuterMostPipeline = GetMostPipelineOuter())
@@ -1132,57 +1144,68 @@ void UInterchangeGenericMaterialPipeline::ExecutePipeline(UInterchangeBaseNodeCo
 
 		bool bEnableFactoryNode = bCombinedImportMaterials;
 
-		// Note: Even if import materials is fully disabled we still want to search for existing materials to use,
-		// and even potentially re-enable the node if we're reimporting
-		FString MaterialName = BaseMaterialFactoryNode->GetDisplayLabel();
-		if (UMaterialInterface* ExistingMaterial = FindExistingMaterialFromSearchLocation(MaterialName, ContentBasePath, SearchLocation))
+		if (bReuseExistingMaterials)
 		{
-			// Add the search location to notify other pipelines what to do about assets.
-			BaseMaterialFactoryNode->AddInt32Attribute(TEXT("SearchLocation"), int32(SearchLocation));
-			BaseMaterialFactoryNode->SetCustomReferenceObject(ExistingMaterial);
+			FString MaterialName = BaseMaterialFactoryNode->GetDisplayLabel();
+			if (UMaterialInterface* ExistingMaterial = FindExistingMaterialFromSearchLocation(MaterialName, ContentBasePath, SearchLocation))
+			{
+				// Add the search location to notify other pipelines what to do about assets.
+				BaseMaterialFactoryNode->AddInt32Attribute(TEXT("SearchLocation"), int32(SearchLocation));
+				BaseMaterialFactoryNode->SetCustomReferenceObject(ExistingMaterial);
 
-			const bool bIsReimport = CacheContextParam.ReimportAsset != nullptr;
+				const bool bIsReimport = CacheContextParam.ReimportAsset != nullptr;
 
-			// We're reimporting the existing material we found: Enable the factory node if we can.
-			// TODO: Can this potentially lead to issues, since the existing material to reimport is discovered by name and search location?
-			if (bIsReimport && CacheContextParam.ReimportAsset == ExistingMaterial)
-			{
-				const bool bIsMaterial = BaseMaterialFactoryNode->IsA<UInterchangeMaterialFactoryNode>()
-											&& ExistingMaterial->IsA<UMaterial>();
+				// We're reimporting the existing material we found: Enable the factory node if we can.
+				// TODO: Can this potentially lead to issues, since the existing material to reimport is discovered by name and search location?
+				if (bIsReimport && CacheContextParam.ReimportAsset == ExistingMaterial)
+				{
+					const bool bIsMaterial = BaseMaterialFactoryNode->IsA<UInterchangeMaterialFactoryNode>()
+						&& ExistingMaterial->IsA<UMaterial>();
 
-				const bool bIsMaterialinstance = BaseMaterialFactoryNode->IsA<UInterchangeMaterialInstanceFactoryNode>()
-					&& ExistingMaterial->IsA<UMaterialInstance>();
+					const bool bIsMaterialInstance = BaseMaterialFactoryNode->IsA<UInterchangeMaterialInstanceFactoryNode>()
+						&& ExistingMaterial->IsA<UMaterialInstance>();
 
 
-				// Reimport can only be done on material instances
-				bEnableFactoryNode = !bIsMaterial && bIsMaterialinstance;
-			}
+					// Reimport can only be done on material instances
+					bEnableFactoryNode = !bIsMaterial && bIsMaterialInstance;
+				}
 #if WITH_EDITORONLY_DATA
-			// This is a scene reimport of a material instance
-			// check whether or not it is a reimport of the current material instance factory node 
-			else if (bIsReimport && CacheContextParam.ReimportAsset->IsA<UInterchangeSceneImportAsset>() && ExistingMaterial->IsA<UMaterialInstance>())
-			{
-				// Skip it if it is not a reimport of the current material instance factory node
-				bEnableFactoryNode = false;
-
-				if (const UInterchangeAssetImportData* AssetImportData = Cast<UInterchangeAssetImportData>(ExistingMaterial->AssetImportData))
+				// This is a scene reimport of a material instance
+				// check whether or not it is a reimport of the current material instance factory node 
+				else if (bIsReimport && CacheContextParam.ReimportAsset->IsA<UInterchangeSceneImportAsset>() && ExistingMaterial->IsA<UMaterialInstance>())
 				{
-					const UInterchangeSceneImportAsset* SceneImportAsset = Cast<UInterchangeSceneImportAsset>(CacheContextParam.ReimportAsset);
-					if (SceneImportAsset->AssetImportData)
+					// Skip it if it is not a reimport of the current material instance factory node
+					bEnableFactoryNode = false;
+
+					if (const UInterchangeAssetImportData* AssetImportData = Cast<UInterchangeAssetImportData>(ExistingMaterial->AssetImportData))
 					{
-						bEnableFactoryNode = AssetImportData->GetFirstFilename().Equals(SceneImportAsset->AssetImportData->GetFirstFilename());
+						const UInterchangeSceneImportAsset* SceneImportAsset = Cast<UInterchangeSceneImportAsset>(CacheContextParam.ReimportAsset);
+						if (SceneImportAsset->AssetImportData)
+						{
+							bEnableFactoryNode = AssetImportData->GetFirstFilename().Equals(SceneImportAsset->AssetImportData->GetFirstFilename());
+						}
 					}
 				}
-			}
 #endif
-			else
-			{
-				// Disabling the factory node here is important so that we don't overwrite the existing asset.
-				// The existing material will still actually be used for the import's material assignments and etc.
-				bEnableFactoryNode = false;
+				else
+				{
+					// Disabling the factory node here is important so that we don't overwrite the existing asset.
+					// The existing material will still actually be used for the import's material assignments and etc.
+					bEnableFactoryNode = false;
+				}
+
+				BaseMaterialFactoryNode->SetCustomIsMaterialImportEnabled(bEnableFactoryNode);
+				BaseMaterialFactoryNode->SetEnabled(bEnableFactoryNode);
+
+				if (!bEnableFactoryNode)
+				{
+					continue;
+				}
 			}
 		}
 
+		bEnableFactoryNode &= bCreateNewMaterials;
+
 		BaseMaterialFactoryNode->SetCustomIsMaterialImportEnabled(bEnableFactoryNode);
 		BaseMaterialFactoryNode->SetEnabled(bEnableFactoryNode);
 	}
diff --git a/Engine/Plugins/Interchange/Runtime/Source/Pipelines/Private/InterchangeglTFPipeline.cpp b/Engine/Plugins/Interchange/Runtime/Source/Pipelines/Private/InterchangeglTFPipeline.cpp
--- a/Engine/Plugins/Interchange/Runtime/Source/Pipelines/Private/InterchangeglTFPipeline.cpp
+++ b/Engine/Plugins/Interchange/Runtime/Source/Pipelines/Private/InterchangeglTFPipeline.cpp
@@ -193,15 +193,13 @@ UInterchangeMaterialInstanceFactoryNode* UGLTFPipelineSettings::BuildMaterialIns
 
 	FString MaterialFactoryNodeUid = UInterchangeFactoryBaseNode::BuildFactoryNodeUid(ShaderGraphNode->GetUniqueID());
 
-	// Check the search location of the material, if under root, then we shouldn't replace the material factory node and just take the existing one
+	// Check the search location of the material, if no search location is specified, attribute "SearchLocation" is not set,
+	// then we shouldn't replace the material factory node and just take the existing one
 	if (UInterchangeBaseMaterialFactoryNode* MaterialFactoryNode = Cast<UInterchangeBaseMaterialFactoryNode>(NodeContainer->GetFactoryNode(MaterialFactoryNodeUid)))
 	{
-		if (EInterchangeMaterialSearchLocation SearchLocation; MaterialFactoryNode->GetInt32Attribute(TEXT("SearchLocation"), reinterpret_cast<int32&>(SearchLocation)))
+		if (MaterialFactoryNode->HasAttribute(UE::Interchange::FAttributeKey(TEXT("SearchLocation"))))
 		{
-			if (SearchLocation != EInterchangeMaterialSearchLocation::DoNotSearch)
-			{
-				return nullptr;
-			}
+			return nullptr;
 		}
 	}
 
diff --git a/Engine/Plugins/Interchange/Runtime/Source/Pipelines/Public/InterchangeGenericMaterialPipeline.h b/Engine/Plugins/Interchange/Runtime/Source/Pipelines/Public/InterchangeGenericMaterialPipeline.h
--- a/Engine/Plugins/Interchange/Runtime/Source/Pipelines/Public/InterchangeGenericMaterialPipeline.h
+++ b/Engine/Plugins/Interchange/Runtime/Source/Pipelines/Public/InterchangeGenericMaterialPipeline.h
@@ -3,6 +3,7 @@
 #pragma once
 
 #include "CoreMinimal.h"
+#include "InterchangeCustomVersion.h"
 #include "InterchangePipelineBase.h"
 #include "InterchangeSourceData.h"
 #include "Nodes/InterchangeBaseNodeContainer.h"
@@ -53,8 +54,8 @@ enum class EInterchangeMaterialSearchLocation : uint8
 	UnderRoot,
 	/** Search for existing material in all assets folders. */
 	AllAssets,
-	/** Do not search for existing existing materials */
-	DoNotSearch,
+	/** Do not search for existing existing materials. Deprecated replaced with boolean bReuseExistingMaterials */
+	DoNotSearch  UE_DEPRECATED(5.8, "Moved as a boolean, bReuseExistingMaterials in UInterchangeGenericMaterialPipeline") UMETA(Hidden, Deprecated),
 };
 
 UCLASS(MinimalAPI, BlueprintType, editinlinenew)
@@ -65,48 +66,67 @@ class UInterchangeGenericMaterialPipeline : public UInterchangePipelineBase
 public:
 	UE_API UInterchangeGenericMaterialPipeline();
 
+	// Begin UObject interface
+	UE_API virtual void Serialize(FArchive& Ar) override;
+	// End UObject interface
+
 	static UE_API FString GetPipelineCategory(UClass* AssetClass);
 
 	/** The name of the pipeline that will be display in the import dialog. */
-	UPROPERTY(EditAnywhere, BlueprintReadWrite, Category = "Materials", meta = (StandAlonePipelineProperty = "True", PipelineInternalEditionData = "True"))
+	UPROPERTY(EditAnywhere, BlueprintReadWrite, Category = "Materials", Meta = (StandAlonePipelineProperty = "True", PipelineInternalEditionData = "True"))
 	FString PipelineDisplayName;
 
+	/** If enabled, and there is only one asset and one source, the imported asset will be given this name. */
+	UPROPERTY(EditAnywhere, BlueprintReadWrite, Category = "Materials", Meta = (StandAlonePipelineProperty = "True"))
+	FString AssetName;
+
 	/** If enabled, imports the material assets found in the sources. */
 	UPROPERTY(EditAnywhere, BlueprintReadWrite, Category = "Materials")
 	bool bImportMaterials = true;
 
+	/** If enabled, first check for existing materials before creating new material assets. */
+	UPROPERTY(EditAnywhere, BlueprintReadWrite, Category = "Materials", Meta = (EditCondition = "bImportMaterials"))
+	bool bReuseExistingMaterials = true;
+
 	/** Specify where we should search for existing materials when importing.*/
-	UPROPERTY(EditAnywhere, BlueprintReadWrite, Category = "Materials")
+	UPROPERTY(EditAnywhere, BlueprintReadWrite, Category = "Materials", Meta = (EditCondition = "bImportMaterials && bReuseExistingMaterials"))
 	EInterchangeMaterialSearchLocation SearchLocation = EInterchangeMaterialSearchLocation::Local;
 
-	/** If set, and there is only one asset and one source, the imported asset will be given this name. */
-	UPROPERTY(EditAnywhere, BlueprintReadWrite, Category = "Materials", meta = (StandAlonePipelineProperty = "True"))
-	FString AssetName;
+	/** If disabled, no new material asset will be created. */
+	UPROPERTY(EditAnywhere, BlueprintReadWrite, Category = "Materials", Meta = (EditCondition = "bImportMaterials"))
+	bool bCreateNewMaterials = true;
 
-	/** Determines what kind of material assets should be created for the imported materials. */
-	UPROPERTY(EditAnywhere, BlueprintReadWrite, Category = "Materials", Meta=(EditCondition="bImportMaterials"))
+	/** Specifies what type of material assets should be created for the imported materials. */
+	UPROPERTY(EditAnywhere, BlueprintReadWrite, Category = "Materials", Meta=(EditCondition="bImportMaterials && bCreateNewMaterials"))
 	EInterchangeMaterialImportOption MaterialImport = EInterchangeMaterialImportOption::ImportAsMaterials;
-	
-	/** If set, reference materials along with respective material instances are created. */
-	UPROPERTY(EditAnywhere, BlueprintReadWrite, Category = "Materials", Meta = (EditCondition = "bImportMaterials"))
-	bool bIdentifyDuplicateMaterials = false;
 
-	/** If set, additional material instances are created for reference/parent materials. */
-	UPROPERTY(EditAnywhere, BlueprintReadWrite, Category = "Materials", Meta=(EditCondition="bIdentifyDuplicateMaterials"))
+	/** If enabled, material instances are created from a specified parent material. */
+	UPROPERTY(EditAnywhere, BlueprintReadWrite, Category = "Materials", Meta=(DisplayName = "Create MaterialInstance From Parent", EditCondition="bImportMaterials && bCreateNewMaterials && MaterialImport==EInterchangeMaterialImportOption::ImportAsMaterialInstances"))
 	bool bCreateMaterialInstanceForParent = false;
 
 	/** Optional material used as the parent when importing materials as instances. If no parent material is specified, one will be automatically selected during the import process. */
-	UPROPERTY(EditAnywhere, BlueprintReadWrite, Category = "Materials", Meta= (EditCondition="bImportMaterials && MaterialImport==EInterchangeMaterialImportOption::ImportAsMaterialInstances", AllowedClasses="/Script/Engine.MaterialInterface"))
+	UPROPERTY(EditAnywhere, BlueprintReadWrite, Category = "Materials", Meta= (EditCondition="bImportMaterials && bCreateNewMaterials && MaterialImport==EInterchangeMaterialImportOption::ImportAsMaterialInstances && bCreateMaterialInstanceForParent", AllowedClasses="/Script/Engine.MaterialInterface"))
 	FSoftObjectPath ParentMaterial;
 
+	/**
+	 * If enabled and duplicated materials are detected, for each set of duplicated materials,
+	 * a reference material along with respective material instances are created.
+	 */
+	UPROPERTY(EditAnywhere, BlueprintReadWrite, Category = "Materials", Meta = (EditCondition = "bImportMaterials"))
+	bool bIdentifyDuplicateMaterials = false;
+
 	UPROPERTY(VisibleAnywhere, BlueprintReadWrite, Instanced, Category = "Textures")
 	TObjectPtr<UInterchangeGenericTexturePipeline> TexturePipeline;
 
 	UPROPERTY(VisibleAnywhere, BlueprintReadWrite, Instanced, Category = "Textures")
 	TObjectPtr<UInterchangeSparseVolumeTexturePipeline> SparseVolumeTexturePipeline;
 
-	/** If enabled, it will override the displacement center set by shader graph nodes, if any*/
-	UPROPERTY(EditAnywhere, BlueprintReadWrite, Category = "Materials", Meta = (InlineEditConditionToggle))
+	/**
+	 * If enabled, it will override the displacement center set by shader graph nodes, if any
+	 * Note: Some MaterialX Surface may have their preferred displacement center specified.
+	 *       In that case, enabling this option may affect the rendering result. 
+	 */
+	UPROPERTY(EditAnywhere, BlueprintReadWrite, Category = "Materials", Meta = (DisplayName = "Override Displacement Center"))
 	bool bOverrideDisplacement = false;
 
 	/** Set the value of the displacement center. If enabled it will also override any displacement center value set by shader graph nodes*/
@@ -294,6 +314,18 @@ private:
 	using FParameterMaterialInputType = TTuple<FString, EMaterialInputType>;
 
 	const UInterchangeBaseNode * AttributeStorageNode = nullptr;
+
+PRAGMA_DISABLE_DEPRECATION_WARNINGS
+	void UpdateSearchLocationProperty()
+	{
+		// Only process DoNotSearch case since bReuseExistingMaterials is set to true by default
+		if (SearchLocation == EInterchangeMaterialSearchLocation::DoNotSearch)
+		{
+			bReuseExistingMaterials = false;
+			SearchLocation = EInterchangeMaterialSearchLocation::Local;
+		}
+	}
+PRAGMA_ENABLE_DEPRECATION_WARNINGS
 };
 
 #undef UE_API
diff --git a/Engine/Source/Runtime/Interchange/Core/Public/InterchangeCustomVersion.h b/Engine/Source/Runtime/Interchange/Core/Public/InterchangeCustomVersion.h
--- a/Engine/Source/Runtime/Interchange/Core/Public/InterchangeCustomVersion.h
+++ b/Engine/Source/Runtime/Interchange/Core/Public/InterchangeCustomVersion.h
@@ -19,6 +19,9 @@ struct FInterchangeCustomVersion
 		// The change that implemented the previous version had to be backed out to fix a serialization issue
 		MultipleAllocationsPerAttributeInStorageFixed,
 
+		// Replaced 'DoNotSearch' Material's search location enum value with a boolean property
+		AdditionOfReuseExistingMaterialsBool,
+
 		// -----<new versions can be added above this line>-------------------------------------------------
 		VersionPlusOne,
 		LatestVersion = VersionPlusOne - 1
