[Interchange Skeleton Conflict Window] Added ability to expand and collapse trees in the conflicts windows like outliner.

#rnx
#Synced-CL 51767377

diff --git a/Engine/Plugins/Interchange/Editor/Source/Pipelines/Private/InterchangeEditorPipelineDetails.cpp b/Engine/Plugins/Interchange/Editor/Source/Pipelines/Private/InterchangeEditorPipelineDetails.cpp
--- a/Engine/Plugins/Interchange/Editor/Source/Pipelines/Private/InterchangeEditorPipelineDetails.cpp
+++ b/Engine/Plugins/Interchange/Editor/Source/Pipelines/Private/InterchangeEditorPipelineDetails.cpp
@@ -306,40 +306,46 @@ void FInterchangePipelineBaseDetailsCustomization::AddConflictSection()
 		CategoryBuilder.AddCustomRow(ConflictsText).WholeRowContent()
 			[
 				SNew(SHorizontalBox)
-					+ SHorizontalBox::Slot()
-					.VAlign(VAlign_Center)
-					.HAlign(HAlign_Center)
-					.AutoWidth()
-					.Padding(2.0f, 2.0f, 5.0f, 2.0f)
-					[
-						SNew(SImage)
-							.ToolTipText(Conflict_IconTooltip)
-							.Image(ConflictBrush)
-							.ColorAndOpacity(FStyleColors::Warning)
-					]
-					+ SHorizontalBox::Slot()
-					.FillWidth(1.0f)
-					.VAlign(VAlign_Center)
+				+ SHorizontalBox::Slot()
+				.VAlign(VAlign_Center)
+				.HAlign(HAlign_Center)
+				.AutoWidth()
+				.Padding(2.0f, 2.0f, 5.0f, 2.0f)
+				[
+					SNew(SImage)
+					.ToolTipText(Conflict_IconTooltip)
+					.Image(ConflictBrush)
+					.ColorAndOpacity(FStyleColors::Warning)
+				]
+				+ SHorizontalBox::Slot()
+				.FillWidth(1.0f)
+				.VAlign(VAlign_Center)
+				[
+					SNew(STextBlock)
+					.Text(Conflict_NameContent)
+					.Font(IDetailLayoutBuilder::GetDetailFont())
+					.ColorAndOpacity(FStyleColors::Warning)
+				]
+				+ SHorizontalBox::Slot()
+				.VAlign(VAlign_Fill)
+				.AutoWidth()
+				[
+					SNew(SButton)
+					.ToolTipText(Conflict_ButtonTooltip)
+					.OnClicked(this, &FInterchangePipelineBaseDetailsCustomization::ShowConflictDialog, ConflictInfo)
+					.Content()
 					[
-						SNew(STextBlock)
-							.Text(Conflict_NameContent)
+						SNew(SVerticalBox)
+						+ SVerticalBox::Slot()
+						.FillHeight(1.0f)
+						.VAlign(VAlign_Center)
+						[
+							SNew(STextBlock)
+							.Text(Conflict_ButtonText)
 							.Font(IDetailLayoutBuilder::GetDetailFont())
-							.ColorAndOpacity(FStyleColors::Warning)
-					]
-					+ SHorizontalBox::Slot()
-					.AutoWidth()
-					[
-						SNew(SButton)
-							.ToolTipText(Conflict_ButtonTooltip)
-							.OnClicked(this, &FInterchangePipelineBaseDetailsCustomization::ShowConflictDialog, ConflictInfo)
-							.Content()
-							[
-								SNew(STextBlock)
-									.Text(Conflict_ButtonText)
-									.Font(IDetailLayoutBuilder::GetDetailFont())
-							]
+						]
 					]
-
+				]
 			];
 	}
 }
diff --git a/Engine/Plugins/Interchange/Runtime/Source/Pipelines/Private/InterchangeGenericAssetsPipeline.cpp b/Engine/Plugins/Interchange/Runtime/Source/Pipelines/Private/InterchangeGenericAssetsPipeline.cpp
--- a/Engine/Plugins/Interchange/Runtime/Source/Pipelines/Private/InterchangeGenericAssetsPipeline.cpp
+++ b/Engine/Plugins/Interchange/Runtime/Source/Pipelines/Private/InterchangeGenericAssetsPipeline.cpp
@@ -2188,41 +2188,23 @@ void SInterchangeGenericAssetSkeletonConflictWidget::Construct(const FArguments&
 
 namespace UE::Interchange::Private
 {
-	void RecursivelyCollapseTreeItem(TSharedPtr<STreeView<TSharedPtr<UInterchangeGenericAssetsPipeline::FSkeletonJoint>>> CompareTree
-		, TSharedPtr<UInterchangeGenericAssetsPipeline::FSkeletonJoint> JointItem)
+	static bool ExpansionFilter_ConflictFilter(TSharedPtr<UInterchangeGenericAssetsPipeline::FSkeletonJoint> Joint)
 	{
-		CompareTree->SetItemExpansion(JointItem, false);
-		for (TSharedPtr<UInterchangeGenericAssetsPipeline::FSkeletonJoint> ChildJoint : JointItem->Children)
-		{
-			RecursivelyCollapseTreeItem(CompareTree, ChildJoint);
-		}
-	}
-
-	void RecursivelyExpandTreeItem(TSharedPtr<STreeView<TSharedPtr<UInterchangeGenericAssetsPipeline::FSkeletonJoint>>> CompareTree
-		, TSharedPtr<UInterchangeGenericAssetsPipeline::FSkeletonJoint> JointItem)
-	{
-		if (!JointItem->bMatch || !JointItem->bChildConflict)
-		{
-			return;
-		}
-		CompareTree->SetItemExpansion(JointItem, true);
-		for (TSharedPtr<UInterchangeGenericAssetsPipeline::FSkeletonJoint> ChildJoint : JointItem->Children)
-		{
-			RecursivelyExpandTreeItem(CompareTree, ChildJoint);
-		}
+		return Joint.IsValid() && Joint->bMatch && Joint->bChildConflict;
 	}
 }
 
 FReply SInterchangeGenericAssetSkeletonConflictWidget::OnExpandToConflict()
 {
+	using namespace UE::Interchange::Private;
 	for (TSharedPtr<UInterchangeGenericAssetsPipeline::FSkeletonJoint> Joint : Joints)
 	{
-		UE::Interchange::Private::RecursivelyCollapseTreeItem(CompareTree, Joint);
+		OnSetExpansionRecursive(Joint, false);
 	}
 
 	for (TSharedPtr<UInterchangeGenericAssetsPipeline::FSkeletonJoint> Joint : Joints)
 	{
-		UE::Interchange::Private::RecursivelyExpandTreeItem(CompareTree, Joint);
+		OnSetExpansionRecursive(Joint, true, FExpansionPredicate::CreateStatic(&ExpansionFilter_ConflictFilter));
 	}
 	return FReply::Handled();
 }
@@ -2243,20 +2225,46 @@ const FSlateBrush* SInterchangeGenericAssetSkeletonConflictWidget::GetCollapsabl
 	return bShowSectionFlag[(uint8)SectionIndex] ? FAppStyle::GetBrush("Symbols.DownArrow") : FAppStyle::GetBrush("Symbols.RightArrow");
 }
 
+void SInterchangeGenericAssetSkeletonConflictWidget::OnSetExpansionRecursive(TSharedPtr<UInterchangeGenericAssetsPipeline::FSkeletonJoint> SkeletonJoint, bool bShouldBeExpanded, const FExpansionPredicate Predicate)
+{
+	if (!SkeletonJoint.IsValid())
+	{
+		return;
+	}
+
+	if (Predicate.IsBound())
+	{
+		if (!Predicate.Execute(SkeletonJoint))
+		{
+			return;
+		}
+	}
+
+	CompareTree->SetItemExpansion(SkeletonJoint, bShouldBeExpanded);
+	for (TSharedPtr<UInterchangeGenericAssetsPipeline::FSkeletonJoint> ChildJoint : SkeletonJoint->Children)
+	{
+		OnSetExpansionRecursive(ChildJoint, bShouldBeExpanded, Predicate);
+	}
+}
+
 TSharedPtr<SWidget> SInterchangeGenericAssetSkeletonConflictWidget::ConstructSkeletonComparison()
 {
+	using namespace UE::Interchange::Private;
 	FText SkeletonStatus = NSLOCTEXT("SInterchangeGenericAssetSkeletonConflictWidget", "ConstructSkeletonComparison_SkeletonStatus", "The skeleton has some conflicts");
 	
 	CompareTree = SNew(STreeView<TSharedPtr<UInterchangeGenericAssetsPipeline::FSkeletonJoint>>)
 		.SelectionMode(ESelectionMode::None)
 		.TreeItemsSource(&Joints)
 		.OnGenerateRow(this, &SInterchangeGenericAssetSkeletonConflictWidget::OnGenerateRowCompareTreeView)
-		.OnGetChildren(this, &SInterchangeGenericAssetSkeletonConflictWidget::OnGetChildrenRowCompareTreeView);
-	
+		.OnGetChildren(this, &SInterchangeGenericAssetSkeletonConflictWidget::OnGetChildrenRowCompareTreeView)
+		.OnSetExpansionRecursive(
+			this,
+			&SInterchangeGenericAssetSkeletonConflictWidget::OnSetExpansionRecursive, 
+			FExpansionPredicate());
 
 	for (TSharedPtr<UInterchangeGenericAssetsPipeline::FSkeletonJoint> Joint : Joints)
 	{
-		UE::Interchange::Private::RecursivelyExpandTreeItem(CompareTree, Joint);
+		OnSetExpansionRecursive(Joint, true, FExpansionPredicate::CreateStatic(&ExpansionFilter_ConflictFilter));
 	}
 	
 	return SNew(SBox)
@@ -2314,7 +2322,7 @@ TSharedPtr<SWidget> SInterchangeGenericAssetSkeletonConflictWidget::ConstructSke
 							[
 								SNew(SButton)
 								.HAlign(HAlign_Center)
-								.Text(NSLOCTEXT("SInterchangeGenericAssetMaterialConflictWidget", "SInterchangeGenericAssetMaterialConflic_ExpandToConflict", "Expand To Conflict"))
+								.Text(NSLOCTEXT("SInterchangeGenericAssetSkeletonConflictWidget", "SInterchangeGenericAssetSkeletonConflic_ExpandToConflict", "Expand To Conflict"))
 								.OnClicked(this, &SInterchangeGenericAssetSkeletonConflictWidget::OnExpandToConflict)
 							]
 							+SHorizontalBox::Slot()
diff --git a/Engine/Plugins/Interchange/Runtime/Source/Pipelines/Public/InterchangeGenericAssetsPipeline.h b/Engine/Plugins/Interchange/Runtime/Source/Pipelines/Public/InterchangeGenericAssetsPipeline.h
--- a/Engine/Plugins/Interchange/Runtime/Source/Pipelines/Public/InterchangeGenericAssetsPipeline.h
+++ b/Engine/Plugins/Interchange/Runtime/Source/Pipelines/Public/InterchangeGenericAssetsPipeline.h
@@ -290,6 +290,8 @@ enum class EInterchangeSkeletonCompareSection : uint8
 
 class SInterchangeGenericAssetSkeletonConflictWidget : public SInterchangeBaseConflictWidget
 {
+	using FExpansionPredicate = TDelegate<bool(TSharedPtr<UInterchangeGenericAssetsPipeline::FSkeletonJoint>)>;
+
 public:
 
 	SLATE_BEGIN_ARGS(SInterchangeGenericAssetSkeletonConflictWidget)
@@ -335,6 +337,11 @@ private:
 	TArray<TSharedPtr<UInterchangeGenericAssetsPipeline::FSkeletonJoint>> Joints;
 	UObject* ReimportObject;
 
+	//////////////////////////////////////////////////////////////////////////
+	// Recursive Expansion
+	void OnSetExpansionRecursive(TSharedPtr<UInterchangeGenericAssetsPipeline::FSkeletonJoint> SkeletonJoint, bool bShouldBeExpanded, const FExpansionPredicate Predicate = FExpansionPredicate());
+	//////////////////////////////////////////////////////////////////////////
+
 	//////////////////////////////////////////////////////////////////////////
 	//Collapse generic
 	bool bShowSectionFlag[(uint8)EInterchangeSkeletonCompareSection::Count];
