[Interchange][FBX] ufbx bind pose defaults to anim's T0 bone position

- fixes (part of) Interchange.Import.Level.FBX.SKM_Arm_anim1

#Synced-CL 51759865

diff --git a/Engine/Plugins/Interchange/Runtime/Source/Parsers/Fbx/Private/Ufbx/UfbxScene.cpp b/Engine/Plugins/Interchange/Runtime/Source/Parsers/Fbx/Private/Ufbx/UfbxScene.cpp
--- a/Engine/Plugins/Interchange/Runtime/Source/Parsers/Fbx/Private/Ufbx/UfbxScene.cpp
+++ b/Engine/Plugins/Interchange/Runtime/Source/Parsers/Fbx/Private/Ufbx/UfbxScene.cpp
@@ -769,12 +769,24 @@ namespace UE::Interchange::Private
 			Message->Text = LOCTEXT("FFbxJointMeshBindPoseGenerator_NonMatchingBindPoses", "BindPose Matrix generation acquired 2 different Matrices from FbxCluster vs FbxPose.\n If import is incorrect, consider using a different setting for FBX option 'bConsiderClusterBeforePoseForMeshBindPose'.");
 		}
 
-		// For nodes not found in bind pose or clusters set BindPose from node_to_world(initial "pose")
+		// Set BindPose For nodes not found in bind pose or clusters
+		// This is needed in the first place for rigid animation which is converted to skinned.
+		// Since rigid animation has no bind pose or skin clusters to get bone position from.
 		for (const ufbx_node* Node : Parser.Scene->nodes)
 		{
 			if (!bHasPose.Contains(Node))
 			{
-				BindPose.Add(Node, Convert::ConvertMatrix(Node->node_to_world));
+				// #ufbx_todo: Using first frame of animation instead of Node->node_to_world to match legacy behavior but node_to_world might 
+				// be more convenient to use all the time since this is the 'default' position and we are discarding it otherwise.
+				if (Parser.Scene->anim)
+				{
+					FMatrix GlobalFbxMatrix = SceneUtils::EvaluateNodeTransformGlobal(Parser, *Node, 0);
+					BindPose.Add(Node, GlobalFbxMatrix);
+				}
+				else
+				{
+					BindPose.Add(Node, Convert::ConvertMatrix(Node->node_to_world));
+				}
 			}
 		}
 
