[Interchange] ufbx also includes anim stack index into rigid anim track payload key

- also added ensure for the key
- potentyally this also might be useful to make level sequence per anim stack

#Synced-CL 51759865

diff --git a/Engine/Plugins/Interchange/Runtime/Source/Parsers/Fbx/Private/Ufbx/UfbxAnimation.cpp b/Engine/Plugins/Interchange/Runtime/Source/Parsers/Fbx/Private/Ufbx/UfbxAnimation.cpp
--- a/Engine/Plugins/Interchange/Runtime/Source/Parsers/Fbx/Private/Ufbx/UfbxAnimation.cpp
+++ b/Engine/Plugins/Interchange/Runtime/Source/Parsers/Fbx/Private/Ufbx/UfbxAnimation.cpp
@@ -578,25 +578,26 @@ namespace UE::Interchange::Private
 						{
 							UInterchangeSceneNode* SceneNode = *Found; 
 
-							const FString PayloadKey = Parser.GetNodeUid(*Node) + TEXT("_RigidAnimationPayloadKey");
-
-							Parser.PayloadContexts.Add(PayloadKey, RigidAnimationPayloadContext);
+							const FString PayloadKeySuffix = FString::Printf(TEXT("_RigidAnimationPayloadKey_%d"), AnimStackIndex);
+							const FString PayloadKey = Parser.GetNodeUid(*Node) + PayloadKeySuffix;
+							if (ensure(!Parser.PayloadContexts.Contains(PayloadKey)))
+							{
+								Parser.PayloadContexts.Add(PayloadKey, RigidAnimationPayloadContext);
 
-							UInterchangeTransformAnimationTrackNode* TransformAnimTrackNode = NewObject< UInterchangeTransformAnimationTrackNode >(&NodeContainer);
+								UInterchangeTransformAnimationTrackNode* TransformAnimTrackNode = NewObject< UInterchangeTransformAnimationTrackNode >(&NodeContainer);
 
-							const FString TransformAnimTrackNodeName = FString::Printf(TEXT("%s"), *SceneNode->GetDisplayLabel());
-							const FString TransformAnimTrackNodeUid = TEXT("\\AnimationTrack\\") + TransformAnimTrackNodeName;
+								const FString TransformAnimTrackNodeName = FString::Printf(TEXT("%s"), *SceneNode->GetDisplayLabel());
+								const FString TransformAnimTrackNodeUid = TEXT("\\AnimationTrack\\") + TransformAnimTrackNodeName;
 
-							NodeContainer.SetupNode(TransformAnimTrackNode, TransformAnimTrackNodeUid, TransformAnimTrackNodeName, EInterchangeNodeContainerType::TranslatedAsset);
-							TransformAnimTrackNode->SetCustomActorDependencyUid(*SceneNode->GetUniqueID());
-							TransformAnimTrackNode->SetCustomAnimationPayloadKey(PayloadKey, EInterchangeAnimationPayLoadType::CURVE);
-							TransformAnimTrackNode->SetCustomUsedChannels(UsedChannels);
+								NodeContainer.SetupNode(TransformAnimTrackNode, TransformAnimTrackNodeUid, TransformAnimTrackNodeName, EInterchangeNodeContainerType::TranslatedAsset);
+								TransformAnimTrackNode->SetCustomActorDependencyUid(*SceneNode->GetUniqueID());
+								TransformAnimTrackNode->SetCustomAnimationPayloadKey(PayloadKey, EInterchangeAnimationPayLoadType::CURVE);
+								TransformAnimTrackNode->SetCustomUsedChannels(UsedChannels);
+							}
 
 							// #ufbx_todo: UE::Interchange::Private::ProcessCustomAttributes
 							// ProcessCustomAttributes(Parser, Node, TransformAnimTrackNode);
-
 						}
-
 					}
 
 				}
