[Interchange] ufbx - fixed payloads data race with animation fetch

#Synced-CL 51759865

diff --git a/Engine/Plugins/Interchange/Runtime/Source/Parsers/Fbx/Private/Ufbx/UfbxParser.cpp b/Engine/Plugins/Interchange/Runtime/Source/Parsers/Fbx/Private/Ufbx/UfbxParser.cpp
--- a/Engine/Plugins/Interchange/Runtime/Source/Parsers/Fbx/Private/Ufbx/UfbxParser.cpp
+++ b/Engine/Plugins/Interchange/Runtime/Source/Parsers/Fbx/Private/Ufbx/UfbxParser.cpp
@@ -701,7 +701,11 @@ bool FUfbxParser::FetchAnimationBakeTransformPayload(const TArray<UE::Interchang
 			// #ufbx_todo: check, FBX SDK parser somehow checks that ResultPayloads already has QueryHashString(i.e. same query already processed) Why is this?
 			FString PayloadFilepathCopy;
 			{
-				// #ufbx_todo: guard concurrent access to ResultPayloads, see UE::Interchange::Private::FFbxAnimation::FetchAnimationBakeTransformPayload
+				FScopeLock Lock(ResultPayloadsCriticalSection);
+				if (ResultPayloads.Contains(QueryHashString))
+				{
+					continue;
+				}
 
 				FString& PayloadFilepath = ResultPayloads.FindOrAdd(QueryHashString);
 				//To avoid file path with too many character, we hash the payloadKey so we have a deterministic length for the file path.
