[Interchange] Named LOD instancing/handling fixes to align more with expectation and Collision mesh identification systems.

#rnx
#Synced-CL 51760870

diff --git a/Engine/Plugins/Interchange/Runtime/Source/Pipelines/Private/InterchangePipelineMeshesUtilities.cpp b/Engine/Plugins/Interchange/Runtime/Source/Pipelines/Private/InterchangePipelineMeshesUtilities.cpp
--- a/Engine/Plugins/Interchange/Runtime/Source/Pipelines/Private/InterchangePipelineMeshesUtilities.cpp
+++ b/Engine/Plugins/Interchange/Runtime/Source/Pipelines/Private/InterchangePipelineMeshesUtilities.cpp
@@ -99,30 +99,19 @@ namespace UE::Private::InterchangeMeshPipeline
 
 			auto CreateInstantiation = [&](const FString& ParentUid, bool bSetParent = true)
 				{
-					//if parent already has an assetInstanceUid set, then create a new sceneNode parent it to ParentSceneNode and set the AssetInstance on the new SceneNode:
-					FString NewSceneNodeUid = ParentUid + TEXT("_MeshLODContainerInstance_") + LODContainerName;
+					FString NewSceneNodeUid = ParentUid + TEXT("\\MeshLODContainerInstance\\") + LODContainerName;
 
 					if (Cast<UInterchangeSceneNode>(BaseNodeContainer->GetNode(NewSceneNodeUid)) == nullptr)
 					{
 						UInterchangeSceneNode* NewSceneNode = NewObject<UInterchangeSceneNode>(BaseNodeContainer);
-						BaseNodeContainer->SetupNode(NewSceneNode, NewSceneNodeUid, LODContainerName + TEXT("_Instance"), EInterchangeNodeContainerType::TranslatedScene, bSetParent ? ParentUid : TEXT(""));
+						BaseNodeContainer->SetupNode(NewSceneNode, NewSceneNodeUid, LODContainerName, EInterchangeNodeContainerType::TranslatedScene, bSetParent ? ParentUid : TEXT(""));
 						NewSceneNode->SetCustomAssetInstanceUid(MeshLODContainerNodeUid);
 					}
 				};
 
-			//#interchange_LODRefactor_Note: this is not the nicest solution, but we need to be able to set the SceneNode's assetInstanceUID to the newly created MeshLODContainerNode.
-			//									Perhaps it is worth thinking about adding a non-const SceneNode query function on NodeContainer? TBD
-			if (UInterchangeSceneNode* ParentSceneNode = const_cast<UInterchangeSceneNode*>(ConstParentSceneNode))
+			if (ConstParentSceneNode)
 			{
-				FString ParentAssetInstanceUid;
-				if (ParentSceneNode->GetCustomAssetInstanceUid(ParentAssetInstanceUid))
-				{
-					CreateInstantiation(ParentSceneNode->GetUniqueID());
-				}
-				else
-				{
-					ParentSceneNode->SetCustomAssetInstanceUid(MeshLODContainerNodeUid);
-				}
+				CreateInstantiation(ConstParentSceneNode->GetUniqueID());
 			}
 			else
 			{
