Interchange - Added support for import containing both SM and SKM at the same time.

This is mainly to fix UE-355559 (Auto detect mesh type should be less agressive and only convert necessary static meshes to skeletal meshes)
The fix removes the Interchange limitation to import either Static Meshes or Skeletal Meshes.

bAutoDetectMesh has been moved and renamed.
From: `UInterchangeGenericMeshPipeline::CommonMeshesProperties.bAutoDetectMesh`
To:  `UInterchangeGenericMeshPipeline::CommonSkeletalMeshesAndAnimationsProperties.bConvertStaticsWithAnimatedTransformToSkeletals`

Which basically put the setting next to `UInterchangeGenericMeshPipeline::CommonSkeletalMeshesAndAnimationsProperties.bConvertStaticsWithMorphTargetsToSkeletals`

All Interchange tests are passing (except test_materialx_openpbr_chessset_substrate for unrelated issue)

[RN] Ability to import Static Meshes and Skeletal Meshes from the same file at the same time.
#tests All Interchange tests passing except the tests already failing on Main.

#virtualized
#Synced-CL 51881665

diff --git a/Engine/Plugins/Interchange/Runtime/Source/Pipelines/Private/InterchangeGenericGeometryCachePipeline.cpp b/Engine/Plugins/Interchange/Runtime/Source/Pipelines/Private/InterchangeGenericGeometryCachePipeline.cpp
--- a/Engine/Plugins/Interchange/Runtime/Source/Pipelines/Private/InterchangeGenericGeometryCachePipeline.cpp
+++ b/Engine/Plugins/Interchange/Runtime/Source/Pipelines/Private/InterchangeGenericGeometryCachePipeline.cpp
@@ -92,10 +92,10 @@ void UInterchangeGenericMeshPipeline::ExecutePreImportPipelineGeometryCache()
 	// With animated transforms, static/skeletal meshes can get baked into a geometry cache
 	if (AnimPayloads.Num() > 0)
 	{
-		// The bAutoDetectMeshType setting converts static meshes with animated transforms into skeletal meshes with AnimSequence
+		// The bConvertStaticsWithAnimatedTransformToSkeletals setting converts static meshes with animated transforms into skeletal meshes with AnimSequence
 		// For the purpose of geometry cache, both are considered as just meshes
 		TArray<FString> AnimMeshInstanceUids;
-		if (CommonMeshesProperties->bAutoDetectMeshType && CommonMeshesProperties->ForceAllMeshAsType == EInterchangeForceMeshType::IFMT_None)
+		if (CommonSkeletalMeshesAndAnimationsProperties->bConvertStaticsWithAnimatedTransformToSkeletals && CommonMeshesProperties->ForceAllMeshAsType == EInterchangeForceMeshType::IFMT_None)
 		{
 			PipelineMeshesUtilities->GetAllSkinnedMeshInstance(AnimMeshInstanceUids);
 		}
diff --git a/Engine/Plugins/Interchange/Runtime/Source/Pipelines/Private/InterchangeGenericMeshPipeline.cpp b/Engine/Plugins/Interchange/Runtime/Source/Pipelines/Private/InterchangeGenericMeshPipeline.cpp
--- a/Engine/Plugins/Interchange/Runtime/Source/Pipelines/Private/InterchangeGenericMeshPipeline.cpp
+++ b/Engine/Plugins/Interchange/Runtime/Source/Pipelines/Private/InterchangeGenericMeshPipeline.cpp
@@ -94,14 +94,6 @@ void UInterchangeGenericMeshPipeline::AdjustSettingsForContext(const FInterchang
 		return;
 	}
 
-	bool bAutoDetectConvertStaticMeshToSkeletalMesh = false;
-	bool bContainStaticMesh = false;
-	bool bContainSkeletalMesh = false;
-	bool bContainGeometryCache = false;
-	bool bContainStaticMeshAnimationNode = false;
-	bool bIgnoreStaticMesh = false;
-	GetMeshesInformationFromTranslatedData(ContextParams.BaseNodeContainer, bAutoDetectConvertStaticMeshToSkeletalMesh, bContainStaticMesh, bContainSkeletalMesh, bContainGeometryCache, bContainStaticMeshAnimationNode, bIgnoreStaticMesh);
-
 	//Avoid creating physics asset when importing a LOD or the alternate skinning
 	if (ContextParams.ContextType == EInterchangePipelineContext::AssetCustomLODImport
 		|| ContextParams.ContextType == EInterchangePipelineContext::AssetCustomLODReimport
@@ -118,7 +110,6 @@ void UInterchangeGenericMeshPipeline::AdjustSettingsForContext(const FInterchang
 			|| ContextParams.ContextType == EInterchangePipelineContext::AssetAlternateSkinningReimport)
 		{
 			CommonMeshesProperties->ForceAllMeshAsType = EInterchangeForceMeshType::IFMT_SkeletalMesh;
-			CommonMeshesProperties->bAutoDetectMeshType = false;
 			CommonMeshesProperties->bBakeMeshes = true;
 			CommonMeshesProperties->bBakePivotMeshes = false;
 			CommonMeshesProperties->bImportLods = false;
@@ -133,18 +124,19 @@ void UInterchangeGenericMeshPipeline::AdjustSettingsForContext(const FInterchang
 			SkeletalMeshImportContentType = EInterchangeSkeletalMeshContentType::All;
 			CommonSkeletalMeshesAndAnimationsProperties->Skeleton.Reset();
 			CommonSkeletalMeshesAndAnimationsProperties->bImportOnlyAnimations = false;
+			CommonSkeletalMeshesAndAnimationsProperties->bConvertStaticsWithAnimatedTransformToSkeletals = false;
 		}
 		else if (ContextParams.ContextType == EInterchangePipelineContext::AssetCustomMorphTargetImport
 			|| ContextParams.ContextType == EInterchangePipelineContext::AssetCustomMorphTargetReImport)
 		{
 			//Custom morph target are imported has a combined static mesh
 			CommonMeshesProperties->ForceAllMeshAsType = EInterchangeForceMeshType::IFMT_StaticMesh;
-			CommonMeshesProperties->bAutoDetectMeshType = false;
 			CommonMeshesProperties->bBakeMeshes = true;
 			CommonMeshesProperties->bBakePivotMeshes = false;
 			CommonMeshesProperties->bImportLods = true;
 			CommonMeshesProperties->bKeepSectionsSeparate = false;
 			CommonMeshesProperties->VertexColorImportOption = EInterchangeVertexColorImportOption::IVCIO_Ignore;
+			CommonSkeletalMeshesAndAnimationsProperties->bConvertStaticsWithAnimatedTransformToSkeletals = false;
 			bImportSkeletalMeshes = false;
 			bImportStaticMeshes = true;
 			CombineStaticMeshesBehavior = EInterchangeCombineStaticMeshesBehavior::All;
@@ -168,7 +160,7 @@ void UInterchangeGenericMeshPipeline::AdjustSettingsForContext(const FInterchang
 				{
 					bImportStaticMeshes = true;
 					CommonMeshesProperties->ForceAllMeshAsType = EInterchangeForceMeshType::IFMT_StaticMesh;
-					CommonMeshesProperties->bAutoDetectMeshType = false;
+					CommonSkeletalMeshesAndAnimationsProperties->bConvertStaticsWithAnimatedTransformToSkeletals = false;
 					bImportSkeletalMeshes = false;
 					CombineStaticMeshesBehavior = EInterchangeCombineStaticMeshesBehavior::All;
 					LodGroup = NAME_None;
@@ -183,13 +175,13 @@ void UInterchangeGenericMeshPipeline::AdjustSettingsForContext(const FInterchang
 				{
 					bImportSkeletalMeshes = true;
 					CommonMeshesProperties->ForceAllMeshAsType = EInterchangeForceMeshType::IFMT_SkeletalMesh;
-					CommonMeshesProperties->bAutoDetectMeshType = false;
+					CommonSkeletalMeshesAndAnimationsProperties->bConvertStaticsWithAnimatedTransformToSkeletals = false;
 					bCreatePhysicsAsset = false;
 					bImportStaticMeshes = false;
 				}
 				else
 				{
-					CommonMeshesProperties->bAutoDetectMeshType = true;
+					CommonSkeletalMeshesAndAnimationsProperties->bConvertStaticsWithAnimatedTransformToSkeletals = true;
 				}
 			}
 		}
@@ -204,7 +196,13 @@ void UInterchangeGenericMeshPipeline::AdjustSettingsForContext(const FInterchang
 	TArray<FString> HideSubCategories;
 	if (ContextParams.ContextType == EInterchangePipelineContext::AssetReimport)
 	{
-		CommonMeshesProperties->bAutoDetectMeshType = false;
+		bool bContainStaticMesh = false;
+		bool bContainSkeletalMesh = false;
+		bool bContainGeometryCache = false;
+
+		GetMeshesInformationFromTranslatedData(ContextParams.BaseNodeContainer, bContainStaticMesh, bContainSkeletalMesh, bContainGeometryCache);
+
+		CommonSkeletalMeshesAndAnimationsProperties->bConvertStaticsWithAnimatedTransformToSkeletals = false;
 
 		HideSubCategories.Add(TEXT("Build"));
 		if (USkeletalMesh* SkeletalMesh = Cast<USkeletalMesh>(ContextParams.ReimportAsset))
@@ -219,7 +217,9 @@ void UInterchangeGenericMeshPipeline::AdjustSettingsForContext(const FInterchang
 			bImportStaticMeshes = false;
 			HideCategories.Add(StaticMeshesCategory);
 			HideCategories.Add(GeometryCachesCategory);
-			if(!bContainSkeletalMesh
+
+			//Handle the case of reimporting a Skeletal Mesh using a model that contains only Static Meshes
+			if (!bContainSkeletalMesh
 				|| SkeletalMeshImportContentType == EInterchangeSkeletalMeshContentType::Geometry
 				|| CommonMeshesProperties->ForceAllMeshAsType == EInterchangeForceMeshType::IFMT_StaticMesh)
 			{
@@ -232,6 +232,8 @@ void UInterchangeGenericMeshPipeline::AdjustSettingsForContext(const FInterchang
 			HideCategories.Add(GeometryCachesCategory);
 			HideCategories.Add(CommonSkeletalMeshesAndAnimationCategory);
 			bImportSkeletalMeshes = false;
+
+			//Handle the case of reimporting a Static Mesh using a model that contains only Skeletal Meshes
 			if (!bContainStaticMesh
 				|| CommonMeshesProperties->ForceAllMeshAsType == EInterchangeForceMeshType::IFMT_SkeletalMesh)
 			{
@@ -264,11 +266,6 @@ void UInterchangeGenericMeshPipeline::AdjustSettingsForContext(const FInterchang
 
 	if (UInterchangePipelineBase* OuterMostPipeline = GetMostPipelineOuter())
 	{
-		if (bContainGeometryCache)
-		{
-			HideProperty(OuterMostPipeline, CommonMeshesProperties.Get(), GET_MEMBER_NAME_CHECKED(UInterchangeGenericCommonMeshesProperties, ForceAllMeshAsType));
-		}
-
 		constexpr bool bDoTransientSubPipeline = true;
 		if (UInterchangeGenericAssetsPipeline* ParentPipeline = Cast<UInterchangeGenericAssetsPipeline>(OuterMostPipeline))
 		{
@@ -363,77 +360,49 @@ void UInterchangeGenericMeshPipeline::GetSupportAssetClasses(TArray<UClass*>& Pi
 #endif
 
 void UInterchangeGenericMeshPipeline::GetMeshesInformationFromTranslatedData(const UInterchangeBaseNodeContainer* InBaseNodeContainer
-	, bool& bAutoDetectConvertStaticMeshToSkeletalMesh
 	, bool& bContainStaticMesh
 	, bool& bContainSkeletalMesh
-	, bool& bContainGeometryCache
-	, bool& bContainStaticMeshAnimationNode
-	, bool& bIgnoreStaticMeshes) const
+	, bool& bContainGeometryCache) const
 {
+	bContainStaticMesh = false;
+	bContainSkeletalMesh = false;
+	bContainGeometryCache = false;
+
 	//Its valid to call GetMeshesInformationFromTranslatedData with a null container
 	if (!InBaseNodeContainer)
 	{
 		return;
 	}
-	bAutoDetectConvertStaticMeshToSkeletalMesh = false;
-	bContainStaticMesh = false;
-	bContainSkeletalMesh = false;
-	bContainGeometryCache = false;
-	bContainStaticMeshAnimationNode = false;
-	bIgnoreStaticMeshes = false;
-	{
-		TArray<FString> StaticMeshNodeUids;
-		InBaseNodeContainer->IterateNodesOfType<UInterchangeMeshNode>([&bContainSkeletalMesh, &bContainGeometryCache, &StaticMeshNodeUids](const FString& NodeUid, UInterchangeMeshNode* MeshNode)
+	
+	InBaseNodeContainer->BreakableIterateNodesOfType<UInterchangeMeshNode>([&bContainStaticMesh, &bContainSkeletalMesh, &bContainGeometryCache](const FString& NodeUid, UInterchangeMeshNode* MeshNode)
+		{
+			if (!MeshNode->IsMorphTarget())
 			{
-				if (!MeshNode->IsMorphTarget())
+				if (Cast<UInterchangeGeometryCacheNode>(MeshNode))
+				{
+					bContainGeometryCache = true;
+				}
+				else
 				{
-					if (Cast<UInterchangeGeometryCacheNode>(MeshNode))
+					if (MeshNode->IsSkinnedMesh())
 					{
-						bContainGeometryCache = true;
+						bContainSkeletalMesh = true;
 					}
 					else
 					{
-						MeshNode->IsSkinnedMesh() ? bContainSkeletalMesh = true : StaticMeshNodeUids.Add(NodeUid);
-					}
-				}
-			});
-		bContainStaticMesh = !StaticMeshNodeUids.IsEmpty();
-
-		TMap<const UInterchangeSceneNode*, bool> CacheProcessSceneNodes;
-		InBaseNodeContainer->BreakableIterateNodesOfType<UInterchangeTransformAnimationTrackNode>([&InBaseNodeContainer, &bContainStaticMeshAnimationNode, &StaticMeshNodeUids, &CacheProcessSceneNodes](const FString& NodeUid, UInterchangeTransformAnimationTrackNode* AnimationNode)
-			{
-				FString SceneNodeUid;
-				if (AnimationNode->GetCustomActorDependencyUid(SceneNodeUid))
-				{
-					if (const UInterchangeSceneNode* SceneNode = Cast<UInterchangeSceneNode>(InBaseNodeContainer->GetNode(SceneNodeUid)))
-					{
-						if (IsImpactingAnyMeshesRecursive(SceneNode, InBaseNodeContainer, StaticMeshNodeUids, CacheProcessSceneNodes))
-						{
-							bContainStaticMeshAnimationNode = true;
-						}
+						bContainStaticMesh = true;
 					}
 				}
-				return bContainStaticMeshAnimationNode;
-			});
-	}
+			}
 
-	if (CommonMeshesProperties->bAutoDetectMeshType && CommonMeshesProperties->ForceAllMeshAsType == EInterchangeForceMeshType::IFMT_None)
-	{
-		if (!bContainSkeletalMesh && bContainStaticMesh)
-		{
-			//Auto detect some static mesh transform animations, we need to force the skeletal mesh type and recompute
-			bAutoDetectConvertStaticMeshToSkeletalMesh = bContainStaticMeshAnimationNode;
-		}
-		else if (bContainSkeletalMesh)
-		{
-			bIgnoreStaticMeshes = true;
-		}
-	}
+			return bContainStaticMesh && bContainSkeletalMesh && bContainGeometryCache; //If we found one of every sort we can stop the iteration.
+		});
 }
 
 void UInterchangeGenericMeshPipeline::PostLoad()
 {
 	Super::PostLoad();
+
 	if(!bImportCollision_DEPRECATED)
 	{
 		bCollision = bImportCollision_DEPRECATED;
@@ -443,36 +412,21 @@ void UInterchangeGenericMeshPipeline::PostLoad()
 	{
 		CombineStaticMeshesBehavior = EInterchangeCombineStaticMeshesBehavior::All;
 	}
-	
 }
 
 UInterchangePipelineMeshesUtilities* UInterchangeGenericMeshPipeline::CreateMeshPipelineUtilities(UInterchangeBaseNodeContainer* InBaseNodeContainer
 	, const UInterchangeGenericMeshPipeline* Pipeline)
 {
-	UInterchangePipelineMeshesUtilities* CreatedPipelineMeshesUtilities = UInterchangePipelineMeshesUtilities::CreateInterchangePipelineMeshesUtilities(InBaseNodeContainer);
-
-	bool bAutoDetectConvertStaticMeshToSkeletalMesh = false;
-	bool bContainStaticMesh = false;
-	bool bContainSkeletalMesh = false;
-	bool bContainGeometryCache = false;
-	bool bContainStaticMeshAnimationNode = false;
-	bool bIgnoreStaticMeshes = false;
-	Pipeline->GetMeshesInformationFromTranslatedData(InBaseNodeContainer, bAutoDetectConvertStaticMeshToSkeletalMesh, bContainStaticMesh, bContainSkeletalMesh, bContainGeometryCache, bContainStaticMeshAnimationNode, bIgnoreStaticMeshes);
-
-	//Set the context option to use when querying the pipeline mesh utilities
 	FInterchangePipelineMeshesUtilitiesContext DataContext;
-	
-	//We convert to skeletal mesh, only if the translated data do not have skeletal mesh
-	//Rigid mesh import is a fallback when there is no skinned mesh
-	DataContext.bConvertStaticMeshToSkeletalMesh = !bContainSkeletalMesh && (bAutoDetectConvertStaticMeshToSkeletalMesh || (Pipeline->CommonMeshesProperties->ForceAllMeshAsType == EInterchangeForceMeshType::IFMT_SkeletalMesh));
-
-	//Force static mesh convert all mesh to static mesh
-	DataContext.bConvertSkeletalMeshToStaticMesh = (Pipeline->CommonMeshesProperties->ForceAllMeshAsType == EInterchangeForceMeshType::IFMT_StaticMesh);
 
+	//Set the context option to use when querying the pipeline mesh utilities
+	DataContext.ForceMeshType = Pipeline->CommonMeshesProperties->ForceAllMeshAsType;
+	DataContext.bConvertStaticsWithAnimatedTransformToSkeletals = Pipeline->CommonSkeletalMeshesAndAnimationsProperties->bConvertStaticsWithAnimatedTransformToSkeletals;
 	DataContext.bConvertStaticsWithMorphTargetsToSkeletals = Pipeline->CommonSkeletalMeshesAndAnimationsProperties->bConvertStaticsWithMorphTargetsToSkeletals;
 	DataContext.bImportMeshesInBoneHierarchy = Pipeline->CommonSkeletalMeshesAndAnimationsProperties->bImportMeshesInBoneHierarchy;
 	DataContext.bQueryGeometryOnlyIfNoInstance = Pipeline->CommonMeshesProperties->bBakeMeshes || Pipeline->CommonMeshesProperties->bBakePivotMeshes;
-	DataContext.bIgnoreStaticMeshes = bIgnoreStaticMeshes;
+
+	UInterchangePipelineMeshesUtilities* CreatedPipelineMeshesUtilities = UInterchangePipelineMeshesUtilities::CreateInterchangePipelineMeshesUtilities(InBaseNodeContainer);
 	CreatedPipelineMeshesUtilities->SetContext(DataContext);
 	return CreatedPipelineMeshesUtilities;
 }
@@ -622,39 +576,6 @@ bool UInterchangeGenericMeshPipeline::DoClassesIncludeAllEditableStructPropertie
 }
 #endif
 
-bool UInterchangeGenericMeshPipeline::IsImpactingAnyMeshesRecursive(const UInterchangeSceneNode* SceneNode
-	, const UInterchangeBaseNodeContainer* InBaseNodeContainer
-	, const TArray<FString>& StaticMeshNodeUids
-	, TMap<const UInterchangeSceneNode*, bool>& CacheProcessSceneNodes)
-{
-	bool& bIsImpactingCache = CacheProcessSceneNodes.FindOrAdd(SceneNode, false);
-	if (bIsImpactingCache)
-	{
-		return bIsImpactingCache;
-	}
-	FString AssetUid;
-	if (SceneNode->GetCustomAssetInstanceUid(AssetUid))
-	{
-		if (StaticMeshNodeUids.Contains(AssetUid))
-		{
-			bIsImpactingCache = true;
-			return true;
-		}
-	}
-	TArray<FString> Children = InBaseNodeContainer->GetNodeChildrenUids(SceneNode->GetUniqueID());
-	for (const FString& ChildUid : Children)
-	{
-		if (const UInterchangeSceneNode* ChildSceneNode = Cast<UInterchangeSceneNode>(InBaseNodeContainer->GetNode(ChildUid)))
-		{
-			if (IsImpactingAnyMeshesRecursive(ChildSceneNode, InBaseNodeContainer, StaticMeshNodeUids, CacheProcessSceneNodes))
-			{
-				return true;
-			}
-		}
-	}
-	return false;
-}
-
 void UInterchangeGenericMeshPipeline::UpdateAssemblyPartDependencyTable(UInterchangeMeshFactoryNode* MeshFactoryNode, const TMap<int32, TArray<FString>>& NodeUidsPerLodIndex)
 {
 	// Early exit for the common case of no assemblies
diff --git a/Engine/Plugins/Interchange/Runtime/Source/Pipelines/Private/InterchangeGenericStaticMeshPipeline.cpp b/Engine/Plugins/Interchange/Runtime/Source/Pipelines/Private/InterchangeGenericStaticMeshPipeline.cpp
--- a/Engine/Plugins/Interchange/Runtime/Source/Pipelines/Private/InterchangeGenericStaticMeshPipeline.cpp
+++ b/Engine/Plugins/Interchange/Runtime/Source/Pipelines/Private/InterchangeGenericStaticMeshPipeline.cpp
@@ -143,7 +143,7 @@ void UInterchangeGenericMeshPipeline::ExecutePreImportPipelineStaticMesh()
 			// Do not combine static meshes
 
 			//As NodeUids can come from MeshInstances, they can be SceneNodes. but also can be: UInterchangeMeshNode/UInterchangeMeshLODContainerNode/UInterchangeGeometricCaches
-			auto CreateMeshFactoryNodeUnCombined = [this](const TArray<FString>& NodeUids, const bool bInstancedMesh)->bool
+			auto CreateMeshFactoryNodeUnCombined = [this](const TArray<FString>& NodeUids)->bool
 				{
 					bool bFoundMeshes = false;
 
@@ -181,12 +181,13 @@ void UInterchangeGenericMeshPipeline::ExecutePreImportPipelineStaticMesh()
 				};
 
 			TArray<FString> NodeUids;
+
 			PipelineMeshesUtilities->GetAllStaticMeshInstance(NodeUids);
-			if (!CreateMeshFactoryNodeUnCombined(NodeUids, true/*bInstancedMesh*/))
+			if (!CreateMeshFactoryNodeUnCombined(NodeUids))
 			{
 				NodeUids.Reset();
 				PipelineMeshesUtilities->GetAllStaticMeshGeometry(NodeUids);
-				CreateMeshFactoryNodeUnCombined(NodeUids, false/*bInstancedMesh*/);
+				CreateMeshFactoryNodeUnCombined(NodeUids);
 			}
 		}
 	}
diff --git a/Engine/Plugins/Interchange/Runtime/Source/Pipelines/Private/InterchangePipelineMeshesUtilities.cpp b/Engine/Plugins/Interchange/Runtime/Source/Pipelines/Private/InterchangePipelineMeshesUtilities.cpp
--- a/Engine/Plugins/Interchange/Runtime/Source/Pipelines/Private/InterchangePipelineMeshesUtilities.cpp
+++ b/Engine/Plugins/Interchange/Runtime/Source/Pipelines/Private/InterchangePipelineMeshesUtilities.cpp
@@ -152,6 +152,11 @@ static bool IsSceneNodeASocket(const UInterchangeSceneNode* SceneNode)
 	return NodeDisplayName.StartsWith(UInterchangeMeshFactoryNode::GetMeshSocketPrefix());
 }
 
+bool FInterchangeMeshInstance::IsNestedInSkeleton() const
+{
+	return InstancingNode && InstancingNode->IsA<UInterchangeJointNode>();
+}
+
 void FInterchangeMeshGeometry::SetupMeshGeometry(const UInterchangeBaseNodeContainer& BaseNodeContainer, const UInterchangeBaseNode* InBaseNode)
 {
 	BaseNode = InBaseNode;
@@ -329,45 +334,47 @@ void FInterchangeMeshGeometry::SetAssemblyPartDependencies(const UInterchangeBas
 	AssemblyPartDependenciesArray = TSet<FString>(AssemblyPartUids).Array();
 }
 
-
 bool FInterchangePipelineMeshesUtilitiesContext::IsStaticMeshInstance(const FInterchangeMeshInstance& MeshInstance, UInterchangeBaseNodeContainer* BaseNodeContainer)
 {
-	if (bIgnoreStaticMeshes)
+	if (MeshInstance.bIsGeometryCache)
 	{
 		return false;
 	}
-	return !MeshInstance.bIsAnimated && !IsSkeletalMeshInstance(MeshInstance, BaseNodeContainer);
+	return !IsSkeletalMeshInstance(MeshInstance, BaseNodeContainer);
 }
 
 bool FInterchangePipelineMeshesUtilitiesContext::IsSkeletalMeshInstance(const FInterchangeMeshInstance& MeshInstance, UInterchangeBaseNodeContainer* BaseNodeContainer)
 {
-	bool bOutIsStaticMeshNestedInSkeleton = false;
-	return IsSkeletalMeshInstance(MeshInstance, BaseNodeContainer, bOutIsStaticMeshNestedInSkeleton);
-}
-
-bool FInterchangePipelineMeshesUtilitiesContext::IsSkeletalMeshInstance(const FInterchangeMeshInstance& MeshInstance, UInterchangeBaseNodeContainer* BaseNodeContainer, bool& bOutIsStaticMeshNestedInSkeleton)
-{
-	bOutIsStaticMeshNestedInSkeleton = false;
-	if (MeshInstance.bIsAnimated)
+	if (MeshInstance.bIsGeometryCache)
 	{
 		return false;
 	}
-	if (bConvertSkeletalMeshToStaticMesh)
+	if (ForceMeshType == EInterchangeForceMeshType::IFMT_StaticMesh)
 	{
 		return false;
 	}
-	if ((bConvertStaticMeshToSkeletalMesh || MeshInstance.bReferenceSkinnedMesh || (bConvertStaticsWithMorphTargetsToSkeletals && MeshInstance.bHasMorphTargets)) && !MeshInstance.bReferenceMorphTarget)
+	if (bConvertStaticsWithAnimatedTransformToSkeletals && MeshInstance.bHasTransformAnimationTrack)
 	{
 		return true;
 	}
-	else if (bImportMeshesInBoneHierarchy)
+	if (!MeshInstance.bReferenceMorphTarget)
 	{
-		if (const UInterchangeJointNode* SceneNode = Cast<const UInterchangeJointNode>(MeshInstance.InstancingNode))
+		if (ForceMeshType == EInterchangeForceMeshType::IFMT_SkeletalMesh)
+		{
+			return true;
+		}
+		if (MeshInstance.bReferenceSkinnedMesh)
 		{
-			//If current Instancing Node is a JointNode then it is part of a Skeleton.
-			bOutIsStaticMeshNestedInSkeleton = true;
 			return true;
 		}
+		if (bConvertStaticsWithMorphTargetsToSkeletals && MeshInstance.bHasMorphTargets)
+		{
+			return true;
+		}
+	}
+	if (bImportMeshesInBoneHierarchy && MeshInstance.IsNestedInSkeleton())
+	{
+		return true;
 	}
 	return false;
 }
@@ -378,15 +385,11 @@ bool FInterchangePipelineMeshesUtilitiesContext::IsGeometryCacheInstance(const F
 	{
 		return false;
 	}
-	return MeshInstance.bIsAnimated;
+	return MeshInstance.bIsGeometryCache;
 }
 
 bool FInterchangePipelineMeshesUtilitiesContext::IsStaticMeshGeometry(const FInterchangeMeshGeometry& MeshGeometry)
 {
-	if (bIgnoreStaticMeshes)
-	{
-		return false;
-	}
 	if (bQueryGeometryOnlyIfNoInstance && MeshGeometry.UidsOfInstancingNodes.Num() > 0)
 	{
 		return false;
@@ -395,46 +398,39 @@ bool FInterchangePipelineMeshesUtilitiesContext::IsStaticMeshGeometry(const FInt
 	{
 		return false;
 	}
-
 	if (MeshGeometry.BaseNode->IsA<UInterchangeGeometryCacheNode>())
 	{
 		return false;
 	}
-
-	if (bConvertSkeletalMeshToStaticMesh)
+	if (ForceMeshType == EInterchangeForceMeshType::IFMT_StaticMesh)
 	{
 		return true;
 	}
-
 	if (bImportMeshesInBoneHierarchy && MeshGeometry.bIsReferencedBySkeleton)
 	{
 		return false;
 	}
-
 	return !IsSkeletalMeshGeometry(MeshGeometry);
 }
 
 bool FInterchangePipelineMeshesUtilitiesContext::IsSkeletalMeshGeometry(const FInterchangeMeshGeometry& MeshGeometry)
 {
-	if (bConvertSkeletalMeshToStaticMesh)
+	if (ForceMeshType == EInterchangeForceMeshType::IFMT_StaticMesh)
 	{
 		return false;
 	}
-
 	if (bQueryGeometryOnlyIfNoInstance && MeshGeometry.UidsOfInstancingNodes.Num() > 0)
 	{
 		return false;
 	}
-
 	if (MeshGeometry.IsMorphTarget())
 	{
 		return false;
 	}
-
 	if (!MeshGeometry.BaseNode->IsA<UInterchangeGeometryCacheNode>()
-		&& (bConvertStaticMeshToSkeletalMesh
-		|| MeshGeometry.IsSkinnedMesh()
-		|| (bConvertStaticsWithMorphTargetsToSkeletals && MeshGeometry.HasMorphTargetDependencies())))
+		&& (ForceMeshType == EInterchangeForceMeshType::IFMT_SkeletalMesh
+			|| MeshGeometry.IsSkinnedMesh()
+			|| (bConvertStaticsWithMorphTargetsToSkeletals && MeshGeometry.HasMorphTargetDependencies())))
 	{
 		return true;
 	}
@@ -447,15 +443,41 @@ bool FInterchangePipelineMeshesUtilitiesContext::IsGeometryCacheGeometry(const F
 	{
 		return false;
 	}
-
 	if (bQueryGeometryOnlyIfNoInstance && MeshGeometry.UidsOfInstancingNodes.Num() > 0)
 	{
 		return false;
 	}
-
 	return MeshGeometry.BaseNode->IsA<UInterchangeGeometryCacheNode>();
 }
 
+bool IsImpactingMeshRecursive(const UInterchangeSceneNode* AnimatedNode
+	, const UInterchangeBaseNodeContainer* InBaseNodeContainer
+	, const FString& StaticMeshNodeUid)
+{
+	FString AssetUid;
+	if (AnimatedNode->GetCustomAssetInstanceUid(AssetUid))
+	{
+		if (StaticMeshNodeUid == AssetUid)
+		{
+			return true;
+		}
+	}
+
+	TArray<FString> Children = InBaseNodeContainer->GetNodeChildrenUids(AnimatedNode->GetUniqueID());
+	for (const FString& ChildUid : Children)
+	{
+		if (const UInterchangeSceneNode* ChildSceneNode = Cast<UInterchangeSceneNode>(InBaseNodeContainer->GetNode(ChildUid)))
+		{
+			if (IsImpactingMeshRecursive(ChildSceneNode, InBaseNodeContainer, StaticMeshNodeUid))
+			{
+				return true;
+			}
+		}
+	}
+
+	return false;
+}
+
 UInterchangePipelineMeshesUtilities* UInterchangePipelineMeshesUtilities::CreateInterchangePipelineMeshesUtilities(UInterchangeBaseNodeContainer* BaseNodeContainer)
 {
 	if (!ensure(BaseNodeContainer))
@@ -495,64 +517,72 @@ UInterchangePipelineMeshesUtilities* UInterchangePipelineMeshesUtilities::Create
 		}
 	);
 
+	TSet<const UInterchangeSceneNode*> SceneNodesWithTransformAnimationTrackNode;
+
+	BaseNodeContainer->IterateNodesOfType<UInterchangeTransformAnimationTrackNode>(
+		[&PipelineMeshesUtilities, &BaseNodeContainer, &SceneNodesWithTransformAnimationTrackNode](const FString& NodeUid, UInterchangeTransformAnimationTrackNode* AnimationNode)
+		{
+			FString SceneNodeUid;
+			if (AnimationNode->GetCustomActorDependencyUid(SceneNodeUid))
+			{
+				if (const UInterchangeSceneNode* SceneNode = Cast<UInterchangeSceneNode>(BaseNodeContainer->GetNode(SceneNodeUid)))
+				{
+					SceneNodesWithTransformAnimationTrackNode.Add(SceneNode);
+				}
+			}
+		});
+
 	//returns false if MeshUid is not a Mesh.
-	auto ProcessMeshNode = [&PipelineMeshesUtilities, &BaseNodeContainer](const FString& NodeUid, const UInterchangeBaseNode* BaseNode, const FString& AssetUid) -> bool
+	auto ProcessMeshNode = [&PipelineMeshesUtilities, &BaseNodeContainer, &SceneNodesWithTransformAnimationTrackNode](const FString& MeshInstanceUid, const UInterchangeBaseNode* BaseNode, const FString& AssetUid) -> bool
 		{
 			const UInterchangeBaseNode* AssetNode = BaseNodeContainer->GetNode(AssetUid);
 
+			FInterchangeMeshInstance* MeshInstancePtr = nullptr;
+
 			bool bIsVisible = true;
 			if (const UInterchangeSceneNode* SceneNode = Cast<const UInterchangeSceneNode>(BaseNode))
 			{
 				SceneNode->GetCustomActorVisibility(bIsVisible);
 			}
 
-			bool bIsNestedIntoSkeleton = false;
-			if (const UInterchangeJointNode* JointNode = Cast<const UInterchangeJointNode>(BaseNode))
-			{
-				bIsNestedIntoSkeleton = true;
-			}
+			const bool bIsNestedIntoSkeleton = BaseNode && BaseNode->IsA<UInterchangeJointNode>();
 
 			if (AssetNode && (AssetNode->IsA<UInterchangeMeshNode>() || AssetNode->IsA<UInterchangeMeshLODContainerNode>()))
 			{
 				FInterchangeMeshGeometry& MeshGeometry = PipelineMeshesUtilities->MeshGeometriesPerMeshUid.FindChecked(AssetUid);
-				MeshGeometry.UidsOfInstancingNodes.Add(NodeUid);
+				MeshGeometry.UidsOfInstancingNodes.Add(MeshInstanceUid);
 				MeshGeometry.bIsReferencedBySkeleton |= bIsNestedIntoSkeleton;
 
-				const UInterchangeBaseNode* BaseParentNode = BaseNodeContainer->GetNode(BaseNode->GetParentUid());
+				FInterchangeMeshInstance& MeshInstance = PipelineMeshesUtilities->MeshInstancesPerMeshInstanceUid.FindOrAdd(MeshInstanceUid);
+				MeshInstance.InstancingNode = BaseNode;
+				MeshInstance.ReferencingMeshGeometryUids.Add(AssetUid);
+				MeshInstance.bIsVisible = bIsVisible;
 
 				if (MeshGeometry.BaseNode->IsA<UInterchangeGeometryCacheNode>())
 				{
-					FInterchangeMeshInstance& MeshInstance = PipelineMeshesUtilities->MeshInstancesPerMeshInstanceUid.FindOrAdd(NodeUid);
-					MeshInstance.InstancingNode = BaseNode;
-					MeshInstance.ReferencingMeshGeometryUids.Add(AssetUid);
 					MeshInstance.bReferenceSkinnedMesh = false;
 					MeshInstance.bReferenceMorphTarget = false;
 					MeshInstance.bHasMorphTargets = false;
-					MeshInstance.bIsAnimated = true;
-					MeshInstance.bIsVisible = bIsVisible;
+					MeshInstance.bIsGeometryCache = true;	
 				}
 				else
 				{
-					FInterchangeMeshInstance& MeshInstance = PipelineMeshesUtilities->MeshInstancesPerMeshInstanceUid.FindOrAdd(NodeUid);
-					MeshInstance.InstancingNode = BaseNode;
-					MeshInstance.ReferencingMeshGeometryUids.Add(AssetUid);
 					MeshInstance.bReferenceSkinnedMesh |= MeshGeometry.IsSkinnedMesh();
 					MeshInstance.bReferenceMorphTarget |= MeshGeometry.IsMorphTarget();
 					MeshInstance.bHasMorphTargets |= MeshGeometry.HasMorphTargetDependencies();
-					MeshInstance.bIsVisible = bIsVisible;
 				}
 
-				return true;
+				MeshInstancePtr = &MeshInstance;
 			}
 			//#interchange_LODRefactor_Note: we might be able to use UInterchangeMeshLODContainerNode instead of UInterchangeMeshBundleNode completely?
 			else if (const UInterchangeMeshBundleNode* MeshBundle = Cast<UInterchangeMeshBundleNode>(AssetNode))
 			{
-				FInterchangeMeshInstance& MeshInstance = PipelineMeshesUtilities->MeshInstancesPerMeshInstanceUid.FindOrAdd(NodeUid);
+				FInterchangeMeshInstance& MeshInstance = PipelineMeshesUtilities->MeshInstancesPerMeshInstanceUid.FindOrAdd(MeshInstanceUid);
 				MeshInstance.InstancingNode = BaseNode;
 				MeshInstance.bReferenceSkinnedMesh = false;	   // Just static meshes for now
 				MeshInstance.bReferenceMorphTarget = false;
 				MeshInstance.bHasMorphTargets = false;
-				MeshInstance.bIsAnimated = false;
+				MeshInstance.bIsGeometryCache = false;
 				MeshInstance.bIsVisible = bIsVisible;
 
 				// This is a group of many mesh nodes collapsed together, so we need to update all the FInterchangeMeshGeometry entries
@@ -561,16 +591,35 @@ UInterchangePipelineMeshesUtilities* UInterchangePipelineMeshesUtilities::Create
 				for (const FString& ReferenceMeshNodeUid : ReferencedMeshNodeUids)
 				{
 					FInterchangeMeshGeometry& MeshGeometry = PipelineMeshesUtilities->MeshGeometriesPerMeshUid.FindChecked(ReferenceMeshNodeUid);
-					MeshGeometry.UidsOfInstancingNodes.Add(NodeUid);
+					MeshGeometry.UidsOfInstancingNodes.Add(MeshInstanceUid);
 					MeshGeometry.bIsReferencedBySkeleton |= bIsNestedIntoSkeleton;
 				}
 
 				MeshInstance.ReferencingMeshGeometryUids.Add(AssetUid);
 
-				return true;
+				MeshInstancePtr = &MeshInstance;
 			}
 
-			return false;
+			if (MeshInstancePtr == nullptr)
+			{
+				return false;
+			}
+
+			if (AssetNode && (AssetNode->IsA<UInterchangeMeshNode>() || AssetNode->IsA<UInterchangeMeshLODContainerNode>()))
+			{
+				const FString& StaticMeshNodeUid = AssetNode->GetUniqueID();
+
+				for (const UInterchangeSceneNode* AnimatedSceneNode : SceneNodesWithTransformAnimationTrackNode)
+				{
+					if (IsImpactingMeshRecursive(AnimatedSceneNode, BaseNodeContainer, StaticMeshNodeUid))
+					{
+						MeshInstancePtr->bHasTransformAnimationTrack = true;
+						break;
+					}
+				}
+			}
+
+			return true;
 		};
 
 	BaseNodeContainer->BreakableIterateNodesOfType<UInterchangeInstancedStaticMeshComponentNode>([&PipelineMeshesUtilities, &BaseNodeContainer, &ProcessMeshNode](const FString& ISMComponentNodeUid, UInterchangeInstancedStaticMeshComponentNode* ISMComponentNode)
@@ -966,8 +1015,7 @@ void UInterchangePipelineMeshesUtilities::GetCombinedSkinnedMeshInstances(TMap<F
 			continue;
 		}
 
-		bool bIsNestedIntoSkeleton = false;
-		if (!CurrentDataContext.IsSkeletalMeshInstance(MeshInstance, BaseNodeContainer, bIsNestedIntoSkeleton))
+		if (!CurrentDataContext.IsSkeletalMeshInstance(MeshInstance, BaseNodeContainer))
 		{
 			continue;
 		}
@@ -994,7 +1042,9 @@ void UInterchangePipelineMeshesUtilities::GetCombinedSkinnedMeshInstances(TMap<F
 			else if (bIsStaticMesh)
 			{
 				//Create a joint from the instance node (the scene node pointing on the mesh).
+				const bool bIsNestedIntoSkeleton = MeshInstance.IsNestedInSkeleton();
 				SkeletonRootUid = MeshInstanceUidAndMeshInstance.Key;
+
 				if (bIsNestedIntoSkeleton || !MeshInstance.bHasMorphTargets)
 				{
 					//Find the deepest joint node
diff --git a/Engine/Plugins/Interchange/Runtime/Source/Pipelines/Public/InterchangeGenericAssetsPipelineSharedSettings.h b/Engine/Plugins/Interchange/Runtime/Source/Pipelines/Public/InterchangeGenericAssetsPipelineSharedSettings.h
--- a/Engine/Plugins/Interchange/Runtime/Source/Pipelines/Public/InterchangeGenericAssetsPipelineSharedSettings.h
+++ b/Engine/Plugins/Interchange/Runtime/Source/Pipelines/Public/InterchangeGenericAssetsPipelineSharedSettings.h
@@ -78,13 +78,6 @@ public:
 	UPROPERTY(EditAnywhere, BlueprintReadWrite, Category = "Common Meshes", meta = (editcondition = "ForceAllMeshAsType == EInterchangeForceMeshType::IFMT_SkeletalMesh"))
 	bool bSingleBoneSkeleton = false;
 
-	/**
-	 * If enabled, and some static mesh transforms are animated, the pipeline will convert the static mesh into a rigid skeletal mesh.
-	 * This setting is only used if the Force All Meshes As Type setting is set to "None".
-	 */
-	UPROPERTY(EditAnywhere, BlueprintReadWrite, Category = "Common Meshes", meta = (editcondition = "ForceAllMeshAsType == EInterchangeForceMeshType::IFMT_None"))
-	bool bAutoDetectMeshType = true;
-
 	/** If enabled, any existing LODs for meshes are imported. This setting is only used if the Bake Meshes setting is also enabled. */
 	UPROPERTY(EditAnywhere, BlueprintReadWrite, Category = "Common Meshes")
 	bool bImportLods = true;
@@ -148,14 +141,18 @@ public:
 	/** If true, degenerate triangles are removed. */
 	UPROPERTY(EditAnywhere, BlueprintReadWrite, Category = "Common Meshes", meta = (SubCategory = "Build"))
 	bool bRemoveDegenerates = false;
+
+	UE_DEPRECATED(5.8, "Use UInterchangeGenericCommonSkeletalMeshesAndAnimationsProperties::bConvertStaticsWithAnimatedTransformToSkeletals instead")
+	UPROPERTY(meta = (DeprecatedProperty, DeprecationMessage = "bAutoDetectMeshType is deprecated. Use UInterchangeGenericCommonSkeletalMeshesAndAnimationsProperties::bConvertStaticsWithAnimatedTransformToSkeletals instead."))
+	bool bAutoDetectMeshType_DEPRECATED = true;
+
 #if WITH_EDITOR
 	virtual bool IsPropertyChangeNeedRefresh(const FPropertyChangedEvent& PropertyChangedEvent) const override
 	{
 		static const TSet<FName> NeedRefreshProperties =
 		{
 			GET_MEMBER_NAME_CHECKED(UInterchangeGenericCommonMeshesProperties, ForceAllMeshAsType),
-			GET_MEMBER_NAME_CHECKED(UInterchangeGenericCommonMeshesProperties, bSingleBoneSkeleton),
-			GET_MEMBER_NAME_CHECKED(UInterchangeGenericCommonMeshesProperties, bAutoDetectMeshType)
+			GET_MEMBER_NAME_CHECKED(UInterchangeGenericCommonMeshesProperties, bSingleBoneSkeleton)
 		};
 
 		if (NeedRefreshProperties.Contains(PropertyChangedEvent.GetPropertyName()))
@@ -205,6 +202,10 @@ public:
 	UPROPERTY(EditAnywhere, BlueprintReadWrite, Category = "Common Skeletal Meshes and Animations")
 	bool bAddCurveMetadataToSkeleton = true;
 
+	/** If enabled, and some static mesh transforms are animated, the pipeline will convert the static mesh into a rigid skeletal mesh. */
+	UPROPERTY(EditAnywhere, BlueprintReadWrite, Category = "Static Meshes")
+	bool bConvertStaticsWithAnimatedTransformToSkeletals = true;
+
 	/** If enabled, all static meshes that have morph targets will be imported as skeletal meshes instead. */
 	UPROPERTY(EditAnywhere, BlueprintReadWrite, Category = "Static Meshes")
 	bool bConvertStaticsWithMorphTargetsToSkeletals = false;
@@ -214,6 +215,7 @@ public:
 	{
 		static const TSet<FName> NeedRefreshProperties =
 		{
+			GET_MEMBER_NAME_CHECKED(UInterchangeGenericCommonSkeletalMeshesAndAnimationsProperties, bConvertStaticsWithAnimatedTransformToSkeletals),
 			GET_MEMBER_NAME_CHECKED(UInterchangeGenericCommonSkeletalMeshesAndAnimationsProperties, bConvertStaticsWithMorphTargetsToSkeletals),
 			GET_MEMBER_NAME_CHECKED(UInterchangeGenericCommonSkeletalMeshesAndAnimationsProperties, bImportMeshesInBoneHierarchy),
 			GET_MEMBER_NAME_CHECKED(UInterchangeGenericCommonSkeletalMeshesAndAnimationsProperties, Skeleton),
diff --git a/Engine/Plugins/Interchange/Runtime/Source/Pipelines/Public/InterchangeGenericMeshPipeline.h b/Engine/Plugins/Interchange/Runtime/Source/Pipelines/Public/InterchangeGenericMeshPipeline.h
--- a/Engine/Plugins/Interchange/Runtime/Source/Pipelines/Public/InterchangeGenericMeshPipeline.h
+++ b/Engine/Plugins/Interchange/Runtime/Source/Pipelines/Public/InterchangeGenericMeshPipeline.h
@@ -453,12 +453,9 @@ private:
 	 * Fill all reference parameter from the translated node found in the container
 	 */
 	UE_API void GetMeshesInformationFromTranslatedData(const UInterchangeBaseNodeContainer* InBaseNodeContainer
-		, bool& bAutoDetectConvertStaticMeshToSkeletalMesh
 		, bool& bContainStaticMesh
 		, bool& bContainSkeletalMesh
-		, bool& bContainGeometryCache
-		, bool& bContainStaticMeshAnimationNode
-		, bool& bIgnoreStaticMeshes) const;
+		, bool& bContainGeometryCache) const;
 
 	/************************************************************************/
 	/* Geometry cache API BEGIN                                             */
diff --git a/Engine/Plugins/Interchange/Runtime/Source/Pipelines/Public/InterchangePipelineMeshesUtilities.h b/Engine/Plugins/Interchange/Runtime/Source/Pipelines/Public/InterchangePipelineMeshesUtilities.h
--- a/Engine/Plugins/Interchange/Runtime/Source/Pipelines/Public/InterchangePipelineMeshesUtilities.h
+++ b/Engine/Plugins/Interchange/Runtime/Source/Pipelines/Public/InterchangePipelineMeshesUtilities.h
@@ -9,6 +9,7 @@
 #include "InterchangeMaterialFactoryNode.h"
 #include "Nodes/InterchangeBaseNodeContainer.h"
 #include "Nodes/InterchangeFactoryBaseNode.h"
+#include "InterchangeGenericAssetsPipelineSharedSettings.h"
 
 #include "InterchangePipelineMeshesUtilities.generated.h"
 
@@ -33,7 +34,8 @@ struct FInterchangeMeshInstance
 		: bReferenceSkinnedMesh(false)
 		, bReferenceMorphTarget(false)
 		, bHasMorphTargets(false)
-		, bIsAnimated(false)
+		, bIsGeometryCache(false)
+		, bHasTransformAnimationTrack(false)
 		, bIsVisible(false)
 		, InstancingNode(nullptr)
 	{
@@ -49,7 +51,10 @@ struct FInterchangeMeshInstance
 	bool bHasMorphTargets;
 
 	UPROPERTY(EditAnywhere, Category = "Interchange | Pipeline | MeshesInstance")
-	bool bIsAnimated;
+	bool bIsGeometryCache;
+
+	UPROPERTY(EditAnywhere, Category = "Interchange | Pipeline | MeshesInstance")
+	bool bHasTransformAnimationTrack;
 
 	UPROPERTY(EditAnywhere, Category = "Interchange | Pipeline | MeshesInstance")
 	bool bIsVisible;
@@ -67,6 +72,9 @@ struct FInterchangeMeshInstance
 	 */
 	UPROPERTY(EditAnywhere, Category = "Interchange | Pipeline | MeshesInstance")
 	TArray<FString> ReferencingMeshGeometryUids;
+
+	//Returns true if this mesh instance is instanced from a Joint node (i.e. part of a skeleton)
+	bool IsNestedInSkeleton() const;
 };
 
 /*
@@ -135,7 +143,7 @@ private:
 
 	static bool BoolConflictHandler(bool bLeft, bool bRight);
 	/*
-	* Used for handling SkeletonRootUid conclicts, which shouldn't happen.
+	* Used for handling SkeletonRootUid conflicts, which shouldn't happen.
 	*/
 	static FString StringConflictHandler(const FString& Left, const FString& Right);
 
@@ -168,17 +176,25 @@ struct FInterchangePipelineMeshesUtilitiesContext
 {
 	GENERATED_BODY()
 
+	PRAGMA_DISABLE_DEPRECATION_WARNINGS
+	FInterchangePipelineMeshesUtilitiesContext() = default;
+	FInterchangePipelineMeshesUtilitiesContext(const FInterchangePipelineMeshesUtilitiesContext&) = default;
+	FInterchangePipelineMeshesUtilitiesContext(FInterchangePipelineMeshesUtilitiesContext&&) = default;
+	FInterchangePipelineMeshesUtilitiesContext& operator=(const FInterchangePipelineMeshesUtilitiesContext&) = default;
+	FInterchangePipelineMeshesUtilitiesContext& operator=(FInterchangePipelineMeshesUtilitiesContext&&) = default;
+	PRAGMA_ENABLE_DEPRECATION_WARNINGS
+
 	/**
-	 * If enabled, all static meshes are converted to skeletal meshes.
+	 * Whether to force import all meshes as Static Meshes or Skeletal Meshes.
 	 */
 	UPROPERTY(EditAnywhere, BlueprintReadWrite, Category = "Interchange | Pipeline | MeshesContext")
-	bool bConvertStaticMeshToSkeletalMesh = false;
+	EInterchangeForceMeshType ForceMeshType = EInterchangeForceMeshType::IFMT_None;
 
 	/**
-	 * If enabled, all skeletal meshes are converted to static meshes.
+	 * If enabled, all static meshes that have a transform animation will be imported as skeletal meshes instead.
 	 */
 	UPROPERTY(EditAnywhere, BlueprintReadWrite, Category = "Interchange | Pipeline | MeshesContext")
-	bool bConvertSkeletalMeshToStaticMesh = false;
+	bool bConvertStaticsWithAnimatedTransformToSkeletals = false;
 
 	/**
 	 * If enabled, all static meshes that have morph targets will be imported as skeletal meshes instead.
@@ -199,25 +215,37 @@ struct FInterchangePipelineMeshesUtilitiesContext
 	UPROPERTY(EditAnywhere, BlueprintReadWrite, Category = "Interchange | Pipeline | MeshesContext")
 	bool bQueryGeometryOnlyIfNoInstance = true;
 
-	/**
-	 * If enabled, all static meshes will be ignored. The mesh utility will not return any static meshes.
-	 */
-	UPROPERTY(EditAnywhere, BlueprintReadWrite, Category = "Interchange | Pipeline | MeshesContext")
-	bool bIgnoreStaticMeshes = false;
-
 	/**
 	 * If enabled, all geometry caches will be ignored. The mesh utility will not return any geometry caches.
 	 */
 	UPROPERTY(EditAnywhere, BlueprintReadWrite, Category = "Interchange | Pipeline | MeshesContext")
 	bool bIgnoreGeometryCaches = false;
 
+	UE_DEPRECATED(5.8, "Use ForceMeshType instead.")
+	UPROPERTY(meta = (DeprecatedProperty, DeprecationMessage = "Use ForceMeshType instead."))
+	bool bConvertStaticMeshToSkeletalMesh_DEPRECATED = false;
+
+	UE_DEPRECATED(5.8, "Use ForceMeshType instead.")
+	UPROPERTY(meta = (DeprecatedProperty, DeprecationMessage = "Use ForceMeshType instead."))
+	bool bConvertSkeletalMeshToStaticMesh_DEPRECATED = false;
+
+	UE_DEPRECATED(5.8, "Use ForceMeshType instead.")
+	UPROPERTY(meta = (DeprecatedProperty, DeprecationMessage = "Use ForceMeshType instead."))
+	bool bIgnoreStaticMeshes_DEPRECATED = false;
+
 	bool IsStaticMeshInstance(const FInterchangeMeshInstance& MeshInstance, UInterchangeBaseNodeContainer* BaseNodeContainer);
 	bool IsSkeletalMeshInstance(const FInterchangeMeshInstance& MeshInstance, UInterchangeBaseNodeContainer* BaseNodeContainer);
-	bool IsSkeletalMeshInstance(const FInterchangeMeshInstance& MeshInstance, UInterchangeBaseNodeContainer* BaseNodeContainer, bool& bOutIsStaticMeshNestedInSkeleton);
 	bool IsGeometryCacheInstance(const FInterchangeMeshInstance& MeshInstance);
 	bool IsStaticMeshGeometry(const FInterchangeMeshGeometry& MeshGeometry);
 	bool IsSkeletalMeshGeometry(const FInterchangeMeshGeometry& MeshGeometry);
 	bool IsGeometryCacheGeometry(const FInterchangeMeshGeometry& MeshGeometry);
+
+	UE_DEPRECATED(5.8, "Use IsSkeletalMeshInstance(const FInterchangeMeshInstance& MeshInstance, UInterchangeBaseNodeContainer* BaseNodeContainer) instead.")
+	inline bool IsSkeletalMeshInstance(const FInterchangeMeshInstance& MeshInstance, UInterchangeBaseNodeContainer* BaseNodeContainer, bool& bOutIsStaticMeshNestedInSkeleton)
+	{
+		bOutIsStaticMeshNestedInSkeleton = false;
+		return IsSkeletalMeshInstance(MeshInstance, BaseNodeContainer);
+	}
 };
 
 /**/
diff --git a/Engine/Plugins/Tests/InterchangeTests/Source/InterchangeTests/Private/ImportTestFunctions/ActorImportTestFunctions.cpp b/Engine/Plugins/Tests/InterchangeTests/Source/InterchangeTests/Private/ImportTestFunctions/ActorImportTestFunctions.cpp
--- a/Engine/Plugins/Tests/InterchangeTests/Source/InterchangeTests/Private/ImportTestFunctions/ActorImportTestFunctions.cpp
+++ b/Engine/Plugins/Tests/InterchangeTests/Source/InterchangeTests/Private/ImportTestFunctions/ActorImportTestFunctions.cpp
@@ -59,7 +59,7 @@ FInterchangeTestFunctionResult UActorImportTestFunctions::CheckActorClassCount(c
 
 	if (ImportedNumberOfActors != ExpectedNumberOfActors)
 	{
-		Result.AddError(FString::Printf(TEXT("Expected %d actors of class %s, imported %d."), ExpectedNumberOfActors, *Class->GetName(), ExpectedNumberOfActors));
+		Result.AddError(FString::Printf(TEXT("Expected %d actors of class %s, imported %d."), ExpectedNumberOfActors, *Class->GetName(), ImportedNumberOfActors));
 	}
 
 	return Result;
diff --git a/Engine/Plugins/Tests/InterchangeTests/Source/InterchangeTests/Private/ImportTestFunctions/AnimationImportTestFunctions.cpp b/Engine/Plugins/Tests/InterchangeTests/Source/InterchangeTests/Private/ImportTestFunctions/AnimationImportTestFunctions.cpp
--- a/Engine/Plugins/Tests/InterchangeTests/Source/InterchangeTests/Private/ImportTestFunctions/AnimationImportTestFunctions.cpp
+++ b/Engine/Plugins/Tests/InterchangeTests/Source/InterchangeTests/Private/ImportTestFunctions/AnimationImportTestFunctions.cpp
@@ -51,7 +51,7 @@ FInterchangeTestFunctionResult UAnimationImportTestFunctions::CheckImportedAnimS
 	FInterchangeTestFunctionResult Result;
 	if (AnimSequences.Num() != ExpectedNumberOfImportedAnimSequences)
 	{
-		Result.AddError(FString::Printf(TEXT("Expected %d skeletal meshes, imported %d."), ExpectedNumberOfImportedAnimSequences, AnimSequences.Num()));
+		Result.AddError(FString::Printf(TEXT("Expected %d anim sequences, imported %d."), ExpectedNumberOfImportedAnimSequences, AnimSequences.Num()));
 	}
 
 	return Result;
