Fixed remaining issues with Material pipeline revamp
- Updated tooltips to the ones in the JIRA
- Material import is now skipped if 'Material Import' is unchecked
- DIsplacement options are disabled if 'Material Import' is unchecked
- ParentMaterial is taken in account even when translator is forcing the creation of material instances, FBX, OBJ. No check is done on the material instance's input and the parent material.

#rnx
#Synced-CL 51760431

diff --git a/Engine/Plugins/Interchange/Runtime/Source/Pipelines/Private/InterchangeGenericMaterialPipeline.cpp b/Engine/Plugins/Interchange/Runtime/Source/Pipelines/Private/InterchangeGenericMaterialPipeline.cpp
--- a/Engine/Plugins/Interchange/Runtime/Source/Pipelines/Private/InterchangeGenericMaterialPipeline.cpp
+++ b/Engine/Plugins/Interchange/Runtime/Source/Pipelines/Private/InterchangeGenericMaterialPipeline.cpp
@@ -839,6 +839,12 @@ void UInterchangeGenericMaterialPipeline::ExecutePipeline(UInterchangeBaseNodeCo
 		SparseVolumeTexturePipeline->ScriptedExecutePipeline(InBaseNodeContainer, InSourceDatas, ContentBasePath);
 	}
 
+	if (!bImportMaterials)
+	{
+		// Skip import of materials as a whole
+		return;
+	}
+
 	TSet<UInterchangeShaderGraphNode*> ShaderGraphNodes;
 	TSet<UInterchangeMaterialInstanceNode*> MaterialInstanceNodes;
 	TSet<UInterchangeMaterialReferenceNode*> MaterialReferenceNodes;
@@ -986,7 +992,13 @@ void UInterchangeGenericMaterialPipeline::ExecutePipeline(UInterchangeBaseNodeCo
 	for (UInterchangeMaterialInstanceNode* MaterialNode : MaterialInstanceNodes)
 	{
 		FString ParentPathOrMaterialUid;
-        if (!MaterialNode->GetCustomParent(ParentPathOrMaterialUid) || ParentPathOrMaterialUid.IsEmpty())
+		// Override any existing parent material if one has been specified
+		if (UMaterialInterface* ParentMaterialObj = Cast<UMaterialInterface>(ParentMaterial.TryLoad()))
+		{
+			ParentPathOrMaterialUid = ParentMaterialObj->GetPathName();
+			MaterialNode->SetCustomParent(ParentPathOrMaterialUid);
+		}
+        else if (!MaterialNode->GetCustomParent(ParentPathOrMaterialUid) || ParentPathOrMaterialUid.IsEmpty())
         {
             continue;
         }
diff --git a/Engine/Plugins/Interchange/Runtime/Source/Pipelines/Public/InterchangeGenericMaterialPipeline.h b/Engine/Plugins/Interchange/Runtime/Source/Pipelines/Public/InterchangeGenericMaterialPipeline.h
--- a/Engine/Plugins/Interchange/Runtime/Source/Pipelines/Public/InterchangeGenericMaterialPipeline.h
+++ b/Engine/Plugins/Interchange/Runtime/Source/Pipelines/Public/InterchangeGenericMaterialPipeline.h
@@ -84,20 +84,26 @@ public:
 	UPROPERTY(EditAnywhere, BlueprintReadWrite, Category = "Materials")
 	bool bImportMaterials = true;
 
-	/** If enabled, first check for existing materials before creating new material assets. */
+	/**
+	 * When enabled, if the name of a material in the imported source matches a material
+	 * that already exists in the project, the imported assets will reuse that existing material.
+	 */
 	UPROPERTY(EditAnywhere, BlueprintReadWrite, Category = "Materials", Meta = (EditCondition = "bImportMaterials"))
 	bool bReuseExistingMaterials = true;
 
-	/** Specify where we should search for existing materials when importing.*/
+	/** Determines where to look for matching materials.*/
 	UPROPERTY(EditAnywhere, BlueprintReadWrite, Category = "Materials", Meta = (EditCondition = "bImportMaterials && bReuseExistingMaterials"))
 	EInterchangeMaterialSearchLocation SearchLocation = EInterchangeMaterialSearchLocation::Local;
 
-	/** If disabled, no new material asset will be created. */
+	/**
+	 * When enabled, a new material asset will be created for each material
+	 * in the imported source that has not been matched to an existing material.
+	 */
 	UPROPERTY(EditAnywhere, BlueprintReadWrite, Category = "Materials", Meta = (EditCondition = "bImportMaterials"))
 	bool bCreateNewMaterials = true;
 
-	/** Specifies what type of material assets should be created for the imported materials. */
-	UPROPERTY(EditAnywhere, BlueprintReadWrite, Category = "Materials", Meta=(EditCondition="bImportMaterials && bCreateNewMaterials"))
+	/** Determines which kind of material asset to create for all new materials. */
+	UPROPERTY(EditAnywhere, BlueprintReadWrite, Category = "Materials", Meta=(DisplayName="New Material Asset Type", EditCondition = "bImportMaterials && bCreateNewMaterials"))
 	EInterchangeMaterialImportOption MaterialImport = EInterchangeMaterialImportOption::ImportAsMaterials;
 
 	/** If enabled, material instances are created from a specified parent material. */
@@ -126,11 +132,11 @@ public:
 	 * Note: Some MaterialX Surface may have their preferred displacement center specified.
 	 *       In that case, enabling this option may affect the rendering result. 
 	 */
-	UPROPERTY(EditAnywhere, BlueprintReadWrite, Category = "Materials", Meta = (DisplayName = "Override Displacement Center"))
+	UPROPERTY(EditAnywhere, BlueprintReadWrite, Category = "Materials", Meta = (DisplayName = "Override Displacement Center", EditCondition = "bImportMaterials"))
 	bool bOverrideDisplacement = false;
 
 	/** Set the value of the displacement center. If enabled it will also override any displacement center value set by shader graph nodes*/
-	UPROPERTY(EditAnywhere, BlueprintReadWrite, Category = "Materials", Meta = (EditCondition = "bOverrideDisplacement", ClampMin = "0", DisplayName = "Displacement Center"))
+	UPROPERTY(EditAnywhere, BlueprintReadWrite, Category = "Materials", Meta = (EditCondition = "bOverrideDisplacement && bImportMaterials", ClampMin = "0", DisplayName = "Displacement Center"))
 	float OverrideDisplacementCenter = 0.5f;
 
 	/** BEGIN UInterchangePipelineBase overrides */
