USD: Allow controlling the combination/split of skeletal meshes via collapsing rules

#Synced-CL 51466015

diff --git a/Engine/Plugins/Interchange/Extensions/OpenUSD/Source/Import/Private/SchemaHandlers/CollapsingSchemaHandler.cpp b/Engine/Plugins/Interchange/Extensions/OpenUSD/Source/Import/Private/SchemaHandlers/CollapsingSchemaHandler.cpp
--- a/Engine/Plugins/Interchange/Extensions/OpenUSD/Source/Import/Private/SchemaHandlers/CollapsingSchemaHandler.cpp
+++ b/Engine/Plugins/Interchange/Extensions/OpenUSD/Source/Import/Private/SchemaHandlers/CollapsingSchemaHandler.cpp
@@ -9,6 +9,7 @@
 #include "InterchangeUsdTraversalInfo.h"
 #include "SchemaHandlers/SchemaHandlerUtils.h"
 #include "Usd/InterchangeUsdDefinitions.h"
+#include "USDErrorUtils.h"
 #include "USDPrimConversion.h"
 #include "UsdWrappers/UsdAttribute.h"
 #include "UsdWrappers/UsdPrim.h"
@@ -17,6 +18,7 @@
 #include "UsdWrappers/UsdVariantSets.h"
 
 #include "InterchangeMeshBundleNode.h"
+#include "InterchangeMeshNode.h"
 #include "InterchangeSceneNode.h"
 #include "Nodes/InterchangeBaseNodeContainer.h"
 
@@ -148,7 +150,25 @@ namespace UE::CollapsingSchemaHandler::Private
 					{
 						if (UInterchangeMeshNode* MeshNode = Cast<UInterchangeMeshNode>(AssetNode))
 						{
-							CurrentMeshToTransforms.FindOrAdd(MeshNode->GetUniqueID()).Append(TransformsToMeshBundleRoot);
+							// Skip morph target nodes: they're linked to their parent mesh via
+							// SetMorphTargetDependencyUid and should not be part of the bundle's
+							// mesh-to-transforms map
+							if (MeshNode->IsMorphTarget())
+							{
+								continue;
+							}
+
+							// For skinned meshes, use identity transforms: their vertices are already in bind-pose
+							// space relative to the skeleton, so baking the mesh-to-collapsing-root transform into
+							// the vertices would incorrectly deform them.
+							if (MeshNode->IsSkinnedMesh())
+							{
+								CurrentMeshToTransforms.FindOrAdd(MeshNode->GetUniqueID()).Add(FTransform::Identity);
+							}
+							else
+							{
+								CurrentMeshToTransforms.FindOrAdd(MeshNode->GetUniqueID()).Append(TransformsToMeshBundleRoot);
+							}
 						}
 					}
 				}
@@ -315,6 +335,63 @@ namespace UE::Interchange::USD
 		{
 			Bundle->SetAllMeshNodesAndTransforms(CurrentMeshToTransforms);
 
+			// Check if this bundle contains skinned meshes. If so, mark the bundle with the skeleton dependency so
+			// that the pipeline knows to create a skeletal mesh instead of a static mesh.
+			{
+				UInterchangeBaseNodeContainer* NodeContainer = UsdContext.GetNodeContainer();
+				FString BundleSkeletonDependencyUid;
+				bool bHasSkinnedMesh = false;
+				bool bHasStaticMesh = false;
+
+				for (const TPair<FString, TArray<FTransform>>& MeshTransforms : CurrentMeshToTransforms)
+				{
+					if (const UInterchangeMeshNode* MeshNode = Cast<UInterchangeMeshNode>(NodeContainer->GetNode(MeshTransforms.Key)))
+					{
+						// Morph targets should already be filtered out above, but skip defensively
+						if (MeshNode->IsMorphTarget())
+						{
+							ensure(false);
+							continue;
+						}
+
+						if (MeshNode->IsSkinnedMesh())
+						{
+							bHasSkinnedMesh = true;
+
+							FString MeshSkeletonUid;
+							MeshNode->GetSkeletonDependency(0, MeshSkeletonUid);
+
+							if (BundleSkeletonDependencyUid.IsEmpty())
+							{
+								BundleSkeletonDependencyUid = MeshSkeletonUid;
+							}
+							else if (!BundleSkeletonDependencyUid.Equals(MeshSkeletonUid))
+							{
+								USD_LOG_USERWARNING(FText::Format(
+									LOCTEXT("SkinnedMeshMultipleSkeletons",
+										"Collapsing group '{0}' contains skinned meshes bound to different skeletons. "
+										"Skinned mesh collapsing requires all meshes to share the same skeleton. Skipping skinned bundle."),
+									FText::FromString(Prim.GetPrimPath().GetString())
+								));
+								BundleSkeletonDependencyUid.Empty();
+								bHasSkinnedMesh = false;
+								break;
+							}
+						}
+						else
+						{
+							bHasStaticMesh = true;
+						}
+					}
+				}
+
+				// Only mark as skinned if ALL meshes in the bundle are skinned
+				if (bHasSkinnedMesh && !bHasStaticMesh && !BundleSkeletonDependencyUid.IsEmpty())
+				{
+					Bundle->SetCustomSkeletonDependencyUid(BundleSkeletonDependencyUid);
+				}
+			}
+
 			// Only ever setup our scene node if we can collapse stuff. Otherwise we can just yield the creation of the scene
 			// node to the other schema handlers
 			UInterchangeBaseNode* SceneNodeBase = AccumulatedInfo.GetOrCreateMainSceneNode(Prim, TraversalInfo, UsdContext);
diff --git a/Engine/Plugins/Interchange/Extensions/OpenUSD/Source/Import/Private/SchemaHandlers/GprimSchemaHandler.cpp b/Engine/Plugins/Interchange/Extensions/OpenUSD/Source/Import/Private/SchemaHandlers/GprimSchemaHandler.cpp
--- a/Engine/Plugins/Interchange/Extensions/OpenUSD/Source/Import/Private/SchemaHandlers/GprimSchemaHandler.cpp
+++ b/Engine/Plugins/Interchange/Extensions/OpenUSD/Source/Import/Private/SchemaHandlers/GprimSchemaHandler.cpp
@@ -653,13 +653,6 @@ namespace UE::Interchange::USD
 			return false;
 		}
 
-		// For now let's prevent the collapsing of skinned meshes. Maybe we can handle that in the future and bake them into the
-		// collapsed static mesh
-		if (Prim.HasAPI("SkelBindingAPI") && UE::FUsdSkelRoot::Find(Prim))
-		{
-			return false;
-		}
-
 		return {};
 	}
 
diff --git a/Engine/Plugins/Interchange/Runtime/Source/Nodes/Private/InterchangeMeshBundleNode.cpp b/Engine/Plugins/Interchange/Runtime/Source/Nodes/Private/InterchangeMeshBundleNode.cpp
--- a/Engine/Plugins/Interchange/Runtime/Source/Nodes/Private/InterchangeMeshBundleNode.cpp
+++ b/Engine/Plugins/Interchange/Runtime/Source/Nodes/Private/InterchangeMeshBundleNode.cpp
@@ -74,6 +74,22 @@ void UInterchangeMeshBundleNode::ResetMeshNodeUids()
 	MeshNodeToTransforms.Empty();
 }
 
+bool UInterchangeMeshBundleNode::GetCustomSkeletonDependencyUid(FString& AttributeValue) const
+{
+	IMPLEMENT_NODE_ATTRIBUTE_GETTER(SkeletonDependencyUid, FString);
+}
+
+bool UInterchangeMeshBundleNode::SetCustomSkeletonDependencyUid(const FString& AttributeValue)
+{
+	IMPLEMENT_NODE_ATTRIBUTE_SETTER_NODELEGATE(SkeletonDependencyUid, FString);
+}
+
+bool UInterchangeMeshBundleNode::IsSkinnedMeshBundle() const
+{
+	FString SkeletonUid;
+	return GetCustomSkeletonDependencyUid(SkeletonUid) && !SkeletonUid.IsEmpty();
+}
+
 void UInterchangeMeshBundleNode::GetAllMeshNodesAndTransforms(TMap<FString, TArray<FTransform>>& OutMeshNodeToTransforms) const
 {
 	OutMeshNodeToTransforms = MeshNodeToTransforms.ToMap();
diff --git a/Engine/Plugins/Interchange/Runtime/Source/Nodes/Public/InterchangeMeshBundleNode.h b/Engine/Plugins/Interchange/Runtime/Source/Nodes/Public/InterchangeMeshBundleNode.h
--- a/Engine/Plugins/Interchange/Runtime/Source/Nodes/Public/InterchangeMeshBundleNode.h
+++ b/Engine/Plugins/Interchange/Runtime/Source/Nodes/Public/InterchangeMeshBundleNode.h
@@ -11,9 +11,12 @@
 /**
  * This node represents a collection of UInterchangeMeshNode UIDs, each associated with N transforms.
  *
- * This is currently used in the context of collaping, where the transforms represent collapsed instances of the corresponding
- * mesh nodes, and the whole MeshBundle is converted into a single static mesh. In this case, the transforms are relative to 
- * the "collapsing root".
+ * This is used in the context of collapsing, where the transforms represent collapsed instances of the corresponding
+ * mesh nodes. The whole MeshBundle may be converted into a single static mesh or a single skeletal mesh. In both cases,
+ * the transforms are relative to the "collapsing root".
+ *
+ * When the bundle has a valid SkeletonDependencyUid, it represents a group of skinned meshes that should be combined into
+ * a single skeletal mesh asset.
  */
 UCLASS(MinimalAPI, BlueprintType)
 class UInterchangeMeshBundleNode : public UInterchangeBaseNode
@@ -47,12 +50,25 @@ public:
 	UFUNCTION(BlueprintCallable, Category = "Interchange | Scene | Component | MeshBundle")
 	UE_API void ResetMeshNodeUids();
 
+	/** Get the skeleton dependency UID for this bundle. Returns false if not set. */
+	UFUNCTION(BlueprintCallable, Category = "Interchange | Node | MeshBundle")
+	UE_API bool GetCustomSkeletonDependencyUid(FString& AttributeValue) const;
+
+	/** Set the skeleton dependency UID. When set, this bundle represents a group of skinned meshes. */
+	UFUNCTION(BlueprintCallable, Category = "Interchange | Node | MeshBundle")
+	UE_API bool SetCustomSkeletonDependencyUid(const FString& AttributeValue);
+
+	/** Returns true if this bundle has a valid skeleton dependency, meaning it represents skinned meshes. */
+	UE_API bool IsSkinnedMeshBundle() const;
+
 public:
 	// Separate C++ only functions as the signature is not valid for a UFUNCTION
 	UE_API void GetAllMeshNodesAndTransforms(TMap<FString, TArray<FTransform>>& OutMeshNodeToTransforms) const;
 	UE_API void SetAllMeshNodesAndTransforms(const TMap<FString, TArray<FTransform>>& InMeshNodeToTransforms);
 
 private:
+	IMPLEMENT_NODE_ATTRIBUTE_KEY(SkeletonDependencyUid);
+
 	UE::Interchange::TMapAttributeHelper<FString, TArray<FTransform>> MeshNodeToTransforms;
 };
 
diff --git a/Engine/Plugins/Interchange/Runtime/Source/Pipelines/Private/InterchangeGenericSkeletalMeshPipeline.cpp b/Engine/Plugins/Interchange/Runtime/Source/Pipelines/Private/InterchangeGenericSkeletalMeshPipeline.cpp
--- a/Engine/Plugins/Interchange/Runtime/Source/Pipelines/Private/InterchangeGenericSkeletalMeshPipeline.cpp
+++ b/Engine/Plugins/Interchange/Runtime/Source/Pipelines/Private/InterchangeGenericSkeletalMeshPipeline.cpp
@@ -188,6 +188,10 @@ UInterchangeSkeletalMeshFactoryNode* UInterchangeGenericMeshPipeline::CreateSkel
 		return nullptr;
 	}
 	
+	// Note: UInterchangeMeshBundleNode does not need a separate case here. When collapsing produces a
+	// MeshBundle, a scene node is created with AssetInstanceUid pointing to the bundle. The pipeline sees
+	// this scene node (via the UInterchangeSceneNode branch below), and the LOD data parser handles
+	// MeshBundle expansion when building the actual skeletal mesh geometry.
 	auto GetFirstNodeInfo = [this, &NodeUids](FString& OutFirstMeshNodeUid)->const UInterchangeBaseNode*
 	{
 		FString FirstNodeUid = NodeUids.Num() > 0 ? (NodeUids)[0] : FString();
diff --git a/Engine/Plugins/Interchange/Runtime/Source/Pipelines/Private/InterchangePipelineMeshesUtilities.cpp b/Engine/Plugins/Interchange/Runtime/Source/Pipelines/Private/InterchangePipelineMeshesUtilities.cpp
--- a/Engine/Plugins/Interchange/Runtime/Source/Pipelines/Private/InterchangePipelineMeshesUtilities.cpp
+++ b/Engine/Plugins/Interchange/Runtime/Source/Pipelines/Private/InterchangePipelineMeshesUtilities.cpp
@@ -161,7 +161,6 @@ void FInterchangeMeshGeometry::SetupMeshGeometry(const UInterchangeBaseNodeConta
 {
 	BaseNode = InBaseNode;
 
-	SetIsMorphTarget(BaseNodeContainer);
 	SetIsSkinnedMesh(BaseNodeContainer);
 	SetHasMorphTargetDependencies(BaseNodeContainer);
 	SetHasSkeletonDependecies(BaseNodeContainer);
@@ -169,6 +168,16 @@ void FInterchangeMeshGeometry::SetupMeshGeometry(const UInterchangeBaseNodeConta
 	SetAssemblyPartDependencies(BaseNodeContainer);
 }
 
+bool FInterchangeMeshGeometry::IsMorphTarget() const
+{
+	if (const UInterchangeMeshNode* BaseMeshNode = Cast<UInterchangeMeshNode>(BaseNode))
+	{
+		return BaseMeshNode->IsMorphTarget();
+	}
+
+	return false;
+}
+
 FString FInterchangeMeshGeometry::GetSkeletonRootUid(const UInterchangeBaseNodeContainer& BaseNodeContainer, const TArray<FString>& SkeletonRootNodeUids) const
 {
 	//We are not caching this as results depend on input.
@@ -208,35 +217,49 @@ T FInterchangeMeshGeometry::ParseMeshOrLODContainerNode(const UInterchangeBaseNo
 	{
 		return MeshFunction(BaseMeshNode);
 	}
-	else if (const UInterchangeMeshLODContainerNode* LODContainerNode = Cast<UInterchangeMeshLODContainerNode>(BaseNode))
+	else
 	{
-		TSet<FString> ReferencedMeshUids;
+		// Both UInterchangeMeshLODContainerNode and UInterchangeMeshBundleNode contain references to child
+		// UInterchangeMeshNode UIDs. We collect those UIDs and apply the same per-MeshNode logic for both.
+		TArray<FString> ReferencedMeshUidArray;
 
-		LODContainerNode->GetAllReferencedMeshUids(ReferencedMeshUids);
+		if (const UInterchangeMeshLODContainerNode* LODContainerNode = Cast<UInterchangeMeshLODContainerNode>(BaseNode))
+		{
+			TSet<FString> ReferencedMeshUids;
+			LODContainerNode->GetAllReferencedMeshUids(ReferencedMeshUids);
+			ReferencedMeshUidArray = ReferencedMeshUids.Array();
+		}
+		else if (const UInterchangeMeshBundleNode* MeshBundleNode = Cast<UInterchangeMeshBundleNode>(BaseNode))
+		{
+			MeshBundleNode->GetMeshNodeUids(ReferencedMeshUidArray);
+		}
 
-		TOptional<T> OptionalValue;
-		for (const FString& ReferencedMeshUid : ReferencedMeshUids)
+		if (ReferencedMeshUidArray.Num() > 0)
 		{
-			if (const UInterchangeMeshNode* MeshNode = Cast<UInterchangeMeshNode>(BaseNodeContainer.GetNode(ReferencedMeshUid)))
+			TOptional<T> OptionalValue;
+			for (const FString& ReferencedMeshUid : ReferencedMeshUidArray)
 			{
-				if (OptionalValue.IsSet())
+				if (const UInterchangeMeshNode* MeshNode = Cast<UInterchangeMeshNode>(BaseNodeContainer.GetNode(ReferencedMeshUid)))
 				{
-					T CurrentValue = MeshFunction(MeshNode);
-					if (OptionalValue != CurrentValue)
+					if (OptionalValue.IsSet())
 					{
-						OptionalValue = HandleConflictFunction(OptionalValue.GetValue(), CurrentValue);
+						T CurrentValue = MeshFunction(MeshNode);
+						if (OptionalValue != CurrentValue)
+						{
+							OptionalValue = HandleConflictFunction(OptionalValue.GetValue(), CurrentValue);
+						}
+					}
+					else
+					{
+						OptionalValue = MeshFunction(MeshNode);
 					}
-				}
-				else
-				{
-					OptionalValue = MeshFunction(MeshNode);
 				}
 			}
-		}
 
-		if (OptionalValue.IsSet())
-		{
-			return OptionalValue.GetValue();
+			if (OptionalValue.IsSet())
+			{
+				return OptionalValue.GetValue();
+			}
 		}
 	}
 
@@ -254,25 +277,6 @@ FString FInterchangeMeshGeometry::StringConflictHandler(const FString& Left, con
 	return Left;
 }
 
-void FInterchangeMeshGeometry::SetIsMorphTarget(const UInterchangeBaseNodeContainer& BaseNodeContainer)
-{
-	if (const UInterchangeMeshNode* BaseMeshNode = Cast<UInterchangeMeshNode>(BaseNode))
-	{
-		bIsMorphTarget = BaseMeshNode->IsMorphTarget();
-	}
-	else if (const UInterchangeMeshLODContainerNode* LODContainerNode = Cast<UInterchangeMeshLODContainerNode>(BaseNode))
-	{
-		/*
-		* Only making this branch explicitly so that we can add the comment:
-		* #interchange_LODRefactor_Note: Original implementation used an || check for LOD Groups, but I think mesh LOD Container (aka LOD Group) cannot be morph targets?
-		*									Only place where the results of this function is used is to deny a MeshInstance to be a SkinnedMeshInstance.
-		*									And there is no scenario where a MeshLODContainerNode should be denied to be a SkinnedMeshInstance due to a sub Mesh being a morph target.
-		*/
-
-		bIsMorphTarget = false;
-	}
-}
-
 void FInterchangeMeshGeometry::SetIsSkinnedMesh(const UInterchangeBaseNodeContainer& BaseNodeContainer)
 {
 	bIsSkinnedMesh = ParseMeshOrLODContainerNode<bool>(BaseNodeContainer,
@@ -500,7 +504,7 @@ UInterchangePipelineMeshesUtilities* UInterchangePipelineMeshesUtilities::Create
 		{
 			if (Node->GetNodeContainerType() == EInterchangeNodeContainerType::TranslatedAsset)
 			{
-				if (Node->IsA<UInterchangeMeshNode>() || Node->IsA<UInterchangeMeshLODContainerNode>())
+				if (Node->IsA<UInterchangeMeshNode>() || Node->IsA<UInterchangeMeshLODContainerNode>() || Node->IsA<UInterchangeMeshBundleNode>())
 				{
 					FInterchangeMeshGeometry& MeshGeometry = PipelineMeshesUtilities->MeshGeometriesPerMeshUid.FindOrAdd(NodeUid);
 					MeshGeometry.SetupMeshGeometry(*BaseNodeContainer, Node);
@@ -579,13 +583,13 @@ UInterchangePipelineMeshesUtilities* UInterchangePipelineMeshesUtilities::Create
 			{
 				FInterchangeMeshInstance& MeshInstance = PipelineMeshesUtilities->MeshInstancesPerMeshInstanceUid.FindOrAdd(MeshInstanceUid);
 				MeshInstance.InstancingNode = BaseNode;
-				MeshInstance.bReferenceSkinnedMesh = false;	   // Just static meshes for now
+				MeshInstance.bReferenceSkinnedMesh = MeshBundle->IsSkinnedMeshBundle();
 				MeshInstance.bReferenceMorphTarget = false;
-				MeshInstance.bHasMorphTargets = false;
 				MeshInstance.bIsGeometryCache = false;
 				MeshInstance.bIsVisible = bIsVisible;
 
 				// This is a group of many mesh nodes collapsed together, so we need to update all the FInterchangeMeshGeometry entries
+				bool bBundleHasMorphTargets = false;
 				TArray<FString> ReferencedMeshNodeUids;
 				MeshBundle->GetMeshNodeUids(ReferencedMeshNodeUids);
 				for (const FString& ReferenceMeshNodeUid : ReferencedMeshNodeUids)
@@ -593,7 +597,9 @@ UInterchangePipelineMeshesUtilities* UInterchangePipelineMeshesUtilities::Create
 					FInterchangeMeshGeometry& MeshGeometry = PipelineMeshesUtilities->MeshGeometriesPerMeshUid.FindChecked(ReferenceMeshNodeUid);
 					MeshGeometry.UidsOfInstancingNodes.Add(MeshInstanceUid);
 					MeshGeometry.bIsReferencedBySkeleton |= bIsNestedIntoSkeleton;
+					bBundleHasMorphTargets |= MeshGeometry.HasMorphTargetDependencies();
 				}
+				MeshInstance.bHasMorphTargets = bBundleHasMorphTargets;
 
 				MeshInstance.ReferencingMeshGeometryUids.Add(AssetUid);
 
diff --git a/Engine/Plugins/Interchange/Runtime/Source/Pipelines/Public/InterchangePipelineMeshesUtilities.h b/Engine/Plugins/Interchange/Runtime/Source/Pipelines/Public/InterchangePipelineMeshesUtilities.h
--- a/Engine/Plugins/Interchange/Runtime/Source/Pipelines/Public/InterchangePipelineMeshesUtilities.h
+++ b/Engine/Plugins/Interchange/Runtime/Source/Pipelines/Public/InterchangePipelineMeshesUtilities.h
@@ -124,7 +124,7 @@ struct FInterchangeMeshGeometry
 
 	void SetupMeshGeometry(const UInterchangeBaseNodeContainer& BaseNodeContainer, const UInterchangeBaseNode* InBaseNode);
 
-	bool IsMorphTarget() const { return bIsMorphTarget; };
+	bool IsMorphTarget() const;
 	bool IsSkinnedMesh() const { return bIsSkinnedMesh; };
 	bool HasMorphTargetDependencies() const { return bHasMorphTargetDependencies; };
 	bool HasSkeletonDependecies() const { return bHasSkeletonDependencies; };
@@ -151,7 +151,6 @@ private:
 	/**
 	* Helper functions to set values
 	*/
-	void SetIsMorphTarget(const UInterchangeBaseNodeContainer& BaseNodeContainer);
 	void SetIsSkinnedMesh(const UInterchangeBaseNodeContainer& BaseNodeContainer);
 	void SetHasMorphTargetDependencies(const UInterchangeBaseNodeContainer& BaseNodeContainer);
 	void SetHasSkeletonDependecies(const UInterchangeBaseNodeContainer& BaseNodeContainer);
