Interchange - Removed unused function declaration.

Forgot to remove it when submitting CL-51887232.
Thanks [REDACTED_HANDLE] for bringing this to my attention.

#Synced-CL 51881665

diff --git a/Engine/Plugins/Interchange/Runtime/Source/Pipelines/Public/InterchangeGenericMeshPipeline.h b/Engine/Plugins/Interchange/Runtime/Source/Pipelines/Public/InterchangeGenericMeshPipeline.h
--- a/Engine/Plugins/Interchange/Runtime/Source/Pipelines/Public/InterchangeGenericMeshPipeline.h
+++ b/Engine/Plugins/Interchange/Runtime/Source/Pipelines/Public/InterchangeGenericMeshPipeline.h
@@ -364,11 +364,6 @@ private:
 	/** Table to access the corresponding factory node for the given mesh uid*/
 	TMap<FString, UInterchangeMeshFactoryNode*> AssemblyPartMeshUidToFactoryNodeTable;
 
-	static UE_API bool IsImpactingAnyMeshesRecursive(const UInterchangeSceneNode* SceneNode
-		, const UInterchangeBaseNodeContainer* InBaseNodeContainer
-		, const TArray<FString>& StaticMeshNodeUids
-		, TMap<const UInterchangeSceneNode*, bool>& CacheProcessSceneNodes);
-
 	/**
 	* Update the MeshUid to MeshFactoryNode table with the given node uids
 	* @param NodeUidsPerLodIndex - The NodeUids can be a UInterchangeSceneNode or a UInterchangeMeshNode.
