Fixed crash caused by array index out of bounds
Corrected major issues reported through ECA

#rnx
#Synced-CL 51760431

diff --git a/Engine/Plugins/Interchange/Runtime/Source/Import/Private/Mesh/InterchangeOBJTranslator.cpp b/Engine/Plugins/Interchange/Runtime/Source/Import/Private/Mesh/InterchangeOBJTranslator.cpp
--- a/Engine/Plugins/Interchange/Runtime/Source/Import/Private/Mesh/InterchangeOBJTranslator.cpp
+++ b/Engine/Plugins/Interchange/Runtime/Source/Import/Private/Mesh/InterchangeOBJTranslator.cpp
@@ -47,7 +47,7 @@ static FAutoConsoleVariableRef CCvarInterchangeEnableOBJLevelImport(
 #define INTERCHANGE_OBJ_TRANSLATOR_LOG_ONCE(Verbosity, Format, ...) \
 { \
 	static bool bLogged = false; \
-	UE_CLOG(!bLogged, Verbosity, Format, ##__VA_ARGS__); \
+	UE_CLOG(!bLogged, LogInterchangeObjImport, Verbosity, Format, ##__VA_ARGS__); \
 	bLogged = true; \
 }
 /**
@@ -85,9 +85,10 @@ struct FObjData
 		/** Return the FVertexData of this face which contains the given vertex index */
 		const FVertexData& GetVertexDataContainingVertexIndex(int32 VertexIndex) const
 		{
+			static FVertexData InvalidVertexData { INDEX_NONE, INDEX_NONE, INDEX_NONE };
+
 			const FVertexData* VertexData = Algo::FindBy(Vertices, VertexIndex, &FVertexData::VertexIndex);
-			check(VertexData);
-			return *VertexData;
+			return ensure(VertexData) ? *VertexData : InvalidVertexData;
 		}
 	};
 
@@ -147,9 +148,9 @@ struct FObjData
 		FTextureParameter BumpmapTexture;
 		FTextureParameter NormalmapTexture;
 
-		float ClearCoatThickness;
-		float ClearCoatRoughness;
-		float Anisotropy;
+		float ClearCoatThickness = -1.f;
+		float ClearCoatRoughness = -1.f;
+		float Anisotropy = -1.f;
 	};
 
 	/** All the materials defined by the .obj */
@@ -223,7 +224,7 @@ TArray<int32> FObjData::GetUVIndicesUsedByGroup(const FGroupData& GroupData) con
 		{
 			for (const FVertexData& VertexData : FaceData.Vertices)
 			{
-				if (VertexData.UVIndex != INDEX_NONE)
+				if (ensure(UVs.IsValidIndex(VertexData.UVIndex)))
 				{
 					UVIndexSet.Add(VertexData.UVIndex);
 				}
@@ -256,7 +257,7 @@ FBox FObjData::GetGroupBoundingBox(const FString& GroupName) const
 		}
 		else
 		{
-			INTERCHANGE_OBJ_TRANSLATOR_LOG_ONCE(LogTemp, Warning, TEXT("FObjData::GetGroupBoundingBox: One or more Position index were invalid, skipping those position to compute the bounding box"));
+			INTERCHANGE_OBJ_TRANSLATOR_LOG_ONCE(Warning, TEXT("FObjData::GetGroupBoundingBox: One or more Position index were invalid, skipping those position to compute the bounding box"));
 		}
 	}
 
@@ -295,6 +296,8 @@ int32 FObjData::GetGroupPolygonCount(const FString& GroupName) const
 
 FMeshDescription FObjData::MakeMeshDescriptionForGroup(const FString& GroupName, const FTransform& MeshGlobalTransform) const
 {
+	static FString DefaultMaterialName = UMaterial::GetDefaultMaterial(MD_Surface)->GetName();
+
 	FMeshDescription MeshDescription;
 
 	const FGroupData* GroupDataPtr = Groups.Find(GroupName);
@@ -353,7 +356,7 @@ FMeshDescription FObjData::MakeMeshDescriptionForGroup(const FString& GroupName,
 		}
 		else
 		{
-			INTERCHANGE_OBJ_TRANSLATOR_LOG_ONCE(LogTemp, Warning, TEXT("FObjData::MakeMeshDescriptionForGroup: One or more vertex position index are not valid, skipping those vertex position"));
+			INTERCHANGE_OBJ_TRANSLATOR_LOG_ONCE(Warning, TEXT("FObjData::MakeMeshDescriptionForGroup: One or more vertex position index are not valid, skipping those vertex position"));
 		}
 		
 	}
@@ -377,8 +380,7 @@ FMeshDescription FObjData::MakeMeshDescriptionForGroup(const FString& GroupName,
 		
 		// Create polygon group
 		FPolygonGroupID PolygonGroupIndex = MeshDescription.CreatePolygonGroup();
-		const FString MaterialName = FaceGroup.MaterialName.IsEmpty() ? UMaterial::GetDefaultMaterial(MD_Surface)->GetName() : FaceGroup.MaterialName;
-		ensure(!MaterialName.IsEmpty());
+		const FString MaterialName = FaceGroup.MaterialName.IsEmpty() ? DefaultMaterialName : FaceGroup.MaterialName;
 		Attributes.GetPolygonGroupMaterialSlotNames()[PolygonGroupIndex] = FName(MaterialName);
 
 		// Create faces.
@@ -396,23 +398,29 @@ FMeshDescription FObjData::MakeMeshDescriptionForGroup(const FString& GroupName,
 			for (const FVertexData& VertexData : FaceData.Vertices)
 			{
 				FVertexID VertexID = Algo::BinarySearch(VertexIndexMapping, VertexData.VertexIndex);
-				FVertexInstanceID VertexInstanceID = MeshDescription.CreateVertexInstance(VertexID);
-				VertexInstanceIDs.Add(VertexInstanceID);
-
-				if (VertexData.NormalIndex != INDEX_NONE && Normals.IsValidIndex(VertexData.NormalIndex))
+				if (ensure(VertexID != INDEX_NONE))
 				{
-					FVector3f& Normal = Attributes.GetVertexInstanceNormals()[VertexInstanceID];
-					Normal = PositionToUEBasis(Normals[VertexData.NormalIndex]);
-					TransformPosition(TotalMatrixForNormal, Normal);
-				}
+					FVertexInstanceID VertexInstanceID = MeshDescription.CreateVertexInstance(VertexID);
+					VertexInstanceIDs.Add(VertexInstanceID);
 
-				if (VertexData.UVIndex != INDEX_NONE && UVs.IsValidIndex(VertexData.UVIndex))
-				{
-					Attributes.GetVertexInstanceUVs()[VertexInstanceID] = UVToUEBasis(UVs[VertexData.UVIndex]);
+					if (Normals.IsValidIndex(VertexData.NormalIndex))
+					{
+						FVector3f& Normal = Attributes.GetVertexInstanceNormals()[VertexInstanceID];
+						Normal = PositionToUEBasis(Normals[VertexData.NormalIndex]);
+						TransformPosition(TotalMatrixForNormal, Normal);
+					}
+
+					if (UVs.IsValidIndex(VertexData.UVIndex))
+					{
+						Attributes.GetVertexInstanceUVs()[VertexInstanceID] = UVToUEBasis(UVs[VertexData.UVIndex]);
+					}
 				}
 			}
 
-			Polygons.Add(MeshDescription.CreatePolygon(PolygonGroupIndex, VertexInstanceIDs), &FaceData);
+			if (ensure(VertexInstanceIDs.Num() > 2))
+			{
+				Polygons.Add(MeshDescription.CreatePolygon(PolygonGroupIndex, VertexInstanceIDs), &FaceData);
+			}
 		}
 	}
 
@@ -443,6 +451,10 @@ FMeshDescription FObjData::MakeMeshDescriptionForGroup(const FString& GroupName,
 					// For the vertex we are considering, look at the first face adjacent to the edge, and find the vertex data which includes it
 					// This is a baseline we will use to compare all the other adjacent faces
 					const FVertexData& VertexData = Polygons[EdgePolygonIDs[0]]->GetVertexDataContainingVertexIndex(ObjVertexIndex);
+					if (VertexData.NormalIndex == INDEX_NONE)
+					{
+						continue;
+					}
 
 					for (int32 Index = 1; Index < EdgePolygonIDs.Num(); Index++)
 					{
@@ -450,7 +462,7 @@ FMeshDescription FObjData::MakeMeshDescriptionForGroup(const FString& GroupName,
 						// If the normal index is not the same as the baseline, we know this must be a hard edge.
 						const FVertexData& VertexDataToCompare = Polygons[EdgePolygonIDs[Index]]->GetVertexDataContainingVertexIndex(ObjVertexIndex);
 
-						if (VertexData.NormalIndex != VertexDataToCompare.NormalIndex)
+						if (VertexDataToCompare.NormalIndex != INDEX_NONE && VertexData.NormalIndex != VertexDataToCompare.NormalIndex)
 						{
 							bAllNormalsEqual = false;
 							break;
@@ -635,7 +647,7 @@ namespace ObjParserUtils
 			// @todo: emit warning through Interchange results container
 			// @todo: get new API working through InterchangeManager (keyed multiple times on translator, pipeline, factory objects)
 			// We have a pointer to the Translator inside the FObjData, so we can get the results container from there.
-			INTERCHANGE_OBJ_TRANSLATOR_LOG_ONCE(LogTemp, Warning, TEXT("Interchange Obj translator: Unexpected extra arguments on v keyword"));
+			INTERCHANGE_OBJ_TRANSLATOR_LOG_ONCE(Warning, TEXT("Interchange Obj translator: Unexpected extra arguments on v keyword"));
 		}
 
 		return bSuccess;
@@ -651,7 +663,7 @@ namespace ObjParserUtils
 
 		if (!Line.IsEmpty())
 		{
-			INTERCHANGE_OBJ_TRANSLATOR_LOG_ONCE(LogTemp, Warning, TEXT("Interchange Obj translator: Unexpected extra arguments on vt keyword"));
+			INTERCHANGE_OBJ_TRANSLATOR_LOG_ONCE(Warning, TEXT("Interchange Obj translator: Unexpected extra arguments on vt keyword"));
 		}
 
 		return bSuccess;
@@ -667,7 +679,7 @@ namespace ObjParserUtils
 
 		if (!Line.IsEmpty())
 		{
-			INTERCHANGE_OBJ_TRANSLATOR_LOG_ONCE(LogTemp, Warning, TEXT("Interchange Obj translator: Unexpected extra arguments on vn keyword"));
+			INTERCHANGE_OBJ_TRANSLATOR_LOG_ONCE(Warning, TEXT("Interchange Obj translator: Unexpected extra arguments on vn keyword"));
 		}
 
 		return bSuccess;
@@ -682,9 +694,9 @@ namespace ObjParserUtils
 		{
 			FStringView FaceToken = GetToken(Line);
 
-			auto GetFaceIndex = [](FStringView& Token, int32 NumElements)
+			auto GetFaceIndex = [](FStringView& Token, int32 NumElements, int32& Index)
 			{
-				int32 Index = INDEX_NONE;
+				Index = INDEX_NONE;
 				if (GetInteger(Token, Index, TEXT('/')))
 				{
 					if (Index < 0)
@@ -695,18 +707,24 @@ namespace ObjParserUtils
 					else
 					{
 						// Indices start at one, so adjust for zero-base
-						Index--;
+						--Index;
 					}
 				}
 
-				return Index;
+				return Index >= 0 && Index < NumElements;
 			};
 
-			int32 VertexIndex = GetFaceIndex(FaceToken, ObjData.Positions.Num());
-			int32 UVIndex = GetFaceIndex(FaceToken, ObjData.UVs.Num());
-			int32 NormalIndex = GetFaceIndex(FaceToken, ObjData.Normals.Num());
+			int32 VertexIndex = INDEX_NONE;
+			int32 UVIndex = INDEX_NONE;
+			int32 NormalIndex = INDEX_NONE;
+			const bool bCanAddVertex = GetFaceIndex(FaceToken, ObjData.Positions.Num(), VertexIndex)
+				&& GetFaceIndex(FaceToken, ObjData.UVs.Num(), UVIndex)
+				&& GetFaceIndex(FaceToken, ObjData.Normals.Num(), NormalIndex);
 
-			FaceData.Vertices.Emplace(VertexIndex, UVIndex, NormalIndex);
+			if (bCanAddVertex)
+			{
+				FaceData.Vertices.Emplace(VertexIndex, UVIndex, NormalIndex);
+			}
 		}
 
 		if (FaceData.Vertices.Num() > 2)
@@ -723,7 +741,7 @@ namespace ObjParserUtils
 		}
 		else
 		{
-			INTERCHANGE_OBJ_TRANSLATOR_LOG_ONCE(LogTemp, Warning, TEXT("Interchange Obj translator: Too few vertices on f keyword"));
+			INTERCHANGE_OBJ_TRANSLATOR_LOG_ONCE(Warning, TEXT("Interchange Obj translator: Too few vertices on f keyword"));
 			return false;
 		}
 
@@ -737,12 +755,12 @@ namespace ObjParserUtils
 
 		if (MaterialName.IsEmpty())
 		{
-			INTERCHANGE_OBJ_TRANSLATOR_LOG_ONCE(LogTemp, Warning, TEXT("Interchange Obj translator: Missing material name on usemtl keyword"));
+			INTERCHANGE_OBJ_TRANSLATOR_LOG_ONCE(Warning, TEXT("Interchange Obj translator: Missing material name on usemtl keyword"));
 		}
 
 		if (!Line.IsEmpty())
 		{
-			INTERCHANGE_OBJ_TRANSLATOR_LOG_ONCE(LogTemp, Warning, TEXT("Interchange Obj translator: Unexpected extra arguments on newmtl keyword"));
+			INTERCHANGE_OBJ_TRANSLATOR_LOG_ONCE(Warning, TEXT("Interchange Obj translator: Unexpected extra arguments on usemtl keyword"));
 		}
 
 		ObjData.CurrentMaterial = FString(MaterialName);
@@ -755,12 +773,12 @@ namespace ObjParserUtils
 		FStringView ObjectName = GetToken(Line);
 		if (ObjectName.IsEmpty())
 		{
-			INTERCHANGE_OBJ_TRANSLATOR_LOG_ONCE(LogTemp, Warning, TEXT("Interchange Obj translator: Missing object name on o keyword"));
+			INTERCHANGE_OBJ_TRANSLATOR_LOG_ONCE(Warning, TEXT("Interchange Obj translator: Missing object name on o keyword"));
 		}
 
 		if (!Line.IsEmpty())
 		{
-			INTERCHANGE_OBJ_TRANSLATOR_LOG_ONCE(LogTemp, Warning, TEXT("Interchange Obj translator: Unexpected extra arguments on o keyword"));
+			INTERCHANGE_OBJ_TRANSLATOR_LOG_ONCE(Warning, TEXT("Interchange Obj translator: Unexpected extra arguments on o keyword"));
 		}
 
 		// Use Object Name as Current Group and if there is no Group underneath 
@@ -776,12 +794,12 @@ namespace ObjParserUtils
 		FStringView GroupName = GetToken(Line);
 		if (GroupName.IsEmpty())
 		{
-			INTERCHANGE_OBJ_TRANSLATOR_LOG_ONCE(LogTemp, Warning, TEXT("Interchange Obj translator: Missing group name on g keyword"));
+			INTERCHANGE_OBJ_TRANSLATOR_LOG_ONCE(Warning, TEXT("Interchange Obj translator: Missing group name on g keyword"));
 		}
 
 		if (!Line.IsEmpty())
 		{
-			INTERCHANGE_OBJ_TRANSLATOR_LOG_ONCE(LogTemp, Warning, TEXT("Interchange Obj translator: Unexpected extra arguments on g keyword"));
+			INTERCHANGE_OBJ_TRANSLATOR_LOG_ONCE(Warning, TEXT("Interchange Obj translator: Unexpected extra arguments on g keyword"));
 		}
 
 		ObjData.CurrentGroup = GroupName;
@@ -790,9 +808,6 @@ namespace ObjParserUtils
 
 	static bool ParseLine(FObjData& ObjData, FStringView Line, const FKeywordMap& KeywordMap)
 	{
-		//Issue unknown keyword only once per keyword
-		static TArray<FString> UnknownKeywords;
-
 		// Dispatch by keyword to appropriate parsing function
 		FString Keyword = FString(GetToken(Line));
 		if (Keyword.IsEmpty())
@@ -804,11 +819,10 @@ namespace ObjParserUtils
 		{
 			return (*Handler)(ObjData, Line);
 		}
-		else if (!UnknownKeywords.Contains(Keyword))
+		else
 		{
-			// Unknown keyword, report it once
-			UE_LOG(LogInterchangeObjImport, Display, TEXT("Interchange Obj translator: Unknown keyword: %s"), *Keyword);
-			UnknownKeywords.Add(Keyword);
+			// Unknown keyword
+			UE_LOG(LogInterchangeObjImport, Warning, TEXT("Interchange Obj translator: Unknown keyword: %s"), *Keyword);
 		}
 
 		return true;
@@ -841,12 +855,12 @@ namespace ObjParserUtils
 
 		if (MaterialName.IsEmpty())
 		{
-			INTERCHANGE_OBJ_TRANSLATOR_LOG_ONCE(LogTemp, Warning, TEXT("Interchange Obj translator: Missing material name on newmtl keyword"));
+			INTERCHANGE_OBJ_TRANSLATOR_LOG_ONCE(Warning, TEXT("Interchange Obj translator: Missing material name on newmtl keyword"));
 		}
 
 		if (!Line.IsEmpty())
 		{
-			INTERCHANGE_OBJ_TRANSLATOR_LOG_ONCE(LogTemp, Warning, TEXT("Interchange Obj translator: Unexpected extra arguments on newmtl keyword"));
+			INTERCHANGE_OBJ_TRANSLATOR_LOG_ONCE(Warning, TEXT("Interchange Obj translator: Unexpected extra arguments on newmtl keyword"));
 		}
 
 		ObjData.Materials.Add(FString(MaterialName));
@@ -872,14 +886,14 @@ namespace ObjParserUtils
 			}
 			else
 			{
-				INTERCHANGE_OBJ_TRANSLATOR_LOG_ONCE(LogTemp, Warning, TEXT("Interchange Obj translator: Missing newmtl keyword"));
+				INTERCHANGE_OBJ_TRANSLATOR_LOG_ONCE(Warning, TEXT("Interchange Obj translator: Missing newmtl keyword"));
 				return false;
 			}
 		}
 
 		if (!Line.IsEmpty())
 		{
-			INTERCHANGE_OBJ_TRANSLATOR_LOG_ONCE(LogTemp, Warning, TEXT("Interchange Obj translator: Unexpected extra arguments on material property"));
+			INTERCHANGE_OBJ_TRANSLATOR_LOG_ONCE(Warning, TEXT("Interchange Obj translator: Unexpected extra arguments on material property"));
 		}
 
 		return true;
@@ -891,7 +905,7 @@ namespace ObjParserUtils
 	{
 		if (Line.IsEmpty())
 		{
-			INTERCHANGE_OBJ_TRANSLATOR_LOG_ONCE(LogTemp, Warning, TEXT("Interchange Obj translator: Expecting more parameters(including texture filename) on a map keyword"));
+			INTERCHANGE_OBJ_TRANSLATOR_LOG_ONCE(Warning, TEXT("Interchange Obj translator: Expecting more parameters(including texture filename) on a map keyword"));
 		}
 
 		if (FObjData::FMaterialData* MaterialData = ObjData.Materials.Find(ObjData.MaterialBeingDefined))
@@ -900,7 +914,7 @@ namespace ObjParserUtils
 		}
 		else
 		{
-			INTERCHANGE_OBJ_TRANSLATOR_LOG_ONCE(LogTemp, Warning, TEXT("Interchange Obj translator: Missing newmtl keyword"));
+			INTERCHANGE_OBJ_TRANSLATOR_LOG_ONCE(Warning, TEXT("Interchange Obj translator: Missing newmtl keyword"));
 			return false;
 		}
 
@@ -912,7 +926,7 @@ namespace ObjParserUtils
 	{
 		if (Line.IsEmpty())
 		{
-			INTERCHANGE_OBJ_TRANSLATOR_LOG_ONCE(LogTemp, Warning, TEXT("Interchange Obj translator: Expecting more parameters(including texture filename) on a map keyword"));
+			INTERCHANGE_OBJ_TRANSLATOR_LOG_ONCE(Warning, TEXT("Interchange Obj translator: Expecting more parameters(including texture filename) on a map keyword"));
 		}
 
 		if (FObjData::FMaterialData* MaterialData = ObjData.Materials.Find(ObjData.MaterialBeingDefined))
@@ -958,7 +972,7 @@ namespace ObjParserUtils
 		}
 		else
 		{
-			INTERCHANGE_OBJ_TRANSLATOR_LOG_ONCE(LogTemp, Warning, TEXT("Interchange Obj translator: Missing newmtl keyword"));
+			INTERCHANGE_OBJ_TRANSLATOR_LOG_ONCE(Warning, TEXT("Interchange Obj translator: Missing newmtl keyword"));
 			return false;
 		}
 
@@ -1009,7 +1023,7 @@ namespace ObjParserUtils
 		FStringView MtlFilename = Line; // (instead of GetToken(Line);)
 		if (MtlFilename.IsEmpty())
 		{
-			INTERCHANGE_OBJ_TRANSLATOR_LOG_ONCE(LogTemp, Warning, TEXT("Interchange Obj translator: Missing filename on mtllib keyword"));
+			INTERCHANGE_OBJ_TRANSLATOR_LOG_ONCE(Warning, TEXT("Interchange Obj translator: Missing filename on mtllib keyword"));
 		}
 
 		// Parse the file referenced by the .obj
@@ -1018,14 +1032,14 @@ namespace ObjParserUtils
 
 		bool bSuccess = true;
 
+		// Only .mtl material files are supported
 		if (FPaths::GetExtension(FileToParse).Equals(TEXT("mtl"), ESearchCase::IgnoreCase))
 		{
 			bSuccess &= ParseFile(ObjData, *FileToParse, MtlKeywordMap);
 		}
 		else
 		{
-			//We support only .mtl material file
-			UE_LOG(LogInterchangeObjImport, Warning, TEXT("Interchange Obj translator: Unsupported material file: %s"), *FileToParse);
+			UE_LOG(LogInterchangeObjImport, Warning, TEXT("Interchange Obj translator: Unsupported material file extension: %s"), *FileToParse);
 		}
 
 		// After parsing, clear the current material name
@@ -1090,7 +1104,7 @@ namespace ObjTranslatorUtils
 	static const UInterchangeTexture2DNode* CreateTexture2DNode(UInterchangeBaseNodeContainer& BaseNodeContainer, const FString& TexturePath, FString& InOutTextureName) 
 	{
 		FString TextureName = InOutTextureName;
-		for (int32 Index = 0;;Index ++)
+		for (int32 Index = 0;; ++Index)
 		{
 			FString NormalizedTexturePath(TexturePath);
 			FPaths::NormalizeFilename(NormalizedTexturePath);
@@ -1305,7 +1319,7 @@ bool UInterchangeOBJTranslator::Translate(UInterchangeBaseNodeContainer& BaseNod
 			UInterchangeSceneNode* InterchangeSceneNode = NewObject<UInterchangeSceneNode>(&BaseNodeContainer);
 			BaseNodeContainer.SetupNode(InterchangeSceneNode, ActorNodeUid, GroupName, EInterchangeNodeContainerType::TranslatedScene, SceneNodeUid);
 			InterchangeSceneNode->SetCustomAssetInstanceUid(MeshNodeUid);
-			
+
 			UInterchangeMeshNode* MeshNode = NewObject<UInterchangeMeshNode>(&BaseNodeContainer);
 			BaseNodeContainer.SetupNode(MeshNode, MeshNodeUid, GroupName, EInterchangeNodeContainerType::TranslatedAsset);
 
@@ -1576,12 +1590,6 @@ TOptional<UE::Interchange::FMeshPayloadData> UInterchangeOBJTranslator::GetMeshP
 
 TOptional<UE::Interchange::FImportImage> UInterchangeOBJTranslator::GetTexturePayloadData(const FString& PayLoadKey, TOptional<FString>& AlternateTexturePath) const
 {
-	// @TODO: API design
-	// This method is copied verbatim from InterchangeFBXTranslator.
-	// We need a better way of allowing a translator to delegate translation to other translator.
-	// Here, we have a texture source referenced by the .obj which is ideally handled by the existing texture translators.
-	// In this case, we can just implement the texture payload interface here and use a temporary texture translator instance
-	// to generate the payload, but what if the translator needed to do non-trivial work during the Translate() phase?
 	UE::Interchange::Private::FScopedTranslator ScopedTranslator(PayLoadKey, Results, AnalyticsHandler);
 	const IInterchangeTexturePayloadInterface* TextureTranslator = ScopedTranslator.GetPayLoadInterface<IInterchangeTexturePayloadInterface>();
 	if (!ensure(TextureTranslator))
