[Interchange] ufbx imports all textures

- not just textures on used materials. Same as FBX parser does

#Synced-CL 51759865

diff --git a/Engine/Plugins/Interchange/Runtime/Source/Parsers/Fbx/Private/Ufbx/UfbxMaterial.cpp b/Engine/Plugins/Interchange/Runtime/Source/Parsers/Fbx/Private/Ufbx/UfbxMaterial.cpp
--- a/Engine/Plugins/Interchange/Runtime/Source/Parsers/Fbx/Private/Ufbx/UfbxMaterial.cpp
+++ b/Engine/Plugins/Interchange/Runtime/Source/Parsers/Fbx/Private/Ufbx/UfbxMaterial.cpp
@@ -60,7 +60,7 @@ namespace UE::Interchange::Private
 		{
 		}
 
-		FString FindTextureFile(const ufbx_texture& Texture)
+		static FString FindTextureFile(const FUfbxParser& Parser, const ufbx_texture& Texture)
 		{
 			// #ufbx_todo: maybe just for each of "filename" always check absolute/relative? Just need to decide on the order...
 
@@ -103,12 +103,18 @@ namespace UE::Interchange::Private
 		}
 
 		const UInterchangeTexture2DNode* CreateTexture2DNode(const ufbx_texture& Texture)
+		{
+			return CreateTexture2DNode(Parser, NodeContainer, Texture, &MaterialInstance);
+		}
+
+		static const UInterchangeTexture2DNode* CreateTexture2DNode(FUfbxParser& Parser, UInterchangeBaseNodeContainer& NodeContainer, const ufbx_texture& Texture, UInterchangeBaseNode* InMaterialInstance=nullptr)
 		{
 			if (Texture.type != UFBX_TEXTURE_FILE)
 			{
 				return nullptr;
 			}
-			const FString TextureFilename = FindTextureFile(Texture);
+			const FString TextureFilename = FindTextureFile(Parser, Texture);
+
 			if (TextureFilename.IsEmpty() || !FPaths::FileExists(TextureFilename))
 			{
 				// Texture might be embedded
@@ -124,20 +130,15 @@ namespace UE::Interchange::Private
 					{
 						UInterchangeResultTextureDisplay_TextureFileDoNotExist* Message = Parser.AddMessage<UInterchangeResultTextureDisplay_TextureFileDoNotExist>();
 						Message->TextureName = TextureFilename;
-						Message->MaterialName = MaterialInstance.GetDisplayLabel();
+						Message->MaterialName = InMaterialInstance ? InMaterialInstance->GetDisplayLabel() : FString();
 					}
 					return nullptr;
 				}
 			}
-			return CreateTexture2DNode(TextureFilename);
-		}
 
-		// Expected normalized full path for existing file textures
-		const UInterchangeTexture2DNode* CreateTexture2DNode(const FString& TextureFilePath)
-		{
 			FString TextureName;
 			FString TextureNodeID;
-			ensure(UInterchangeTextureNode::ExtractNodeUidAndNameFromFilePath(TextureFilePath, TextureNodeID, TextureName));
+			ensure(UInterchangeTextureNode::ExtractNodeUidAndNameFromFilePath(TextureFilename, TextureNodeID, TextureName));
 
 			if (const UInterchangeTexture2DNode* TextureNode = Cast<const UInterchangeTexture2DNode>(NodeContainer.GetNode(TextureNodeID)))
 			{
@@ -147,7 +148,7 @@ namespace UE::Interchange::Private
 			UInterchangeTexture2DNode* NewTextureNode = UInterchangeTexture2DNode::Create(&NodeContainer, TextureNodeID, TextureName);
 
 			//All texture translator expect a file as the payload key
-			NewTextureNode->SetPayLoadKey(TextureFilePath);
+			NewTextureNode->SetPayLoadKey(TextureFilename);
 
 			return NewTextureNode;
 		}
@@ -165,6 +166,14 @@ namespace UE::Interchange::Private
 		return MaterialInstanceNode;
 	}
 
+	void FUfbxMaterial::AddAllTextures(UInterchangeBaseNodeContainer& NodeContainer)
+	{
+		for (const ufbx_texture* Texture : Parser.Scene->textures)
+		{
+			FMaterialMapConverter::CreateTexture2DNode(Parser, NodeContainer, *Texture);
+		}
+	}
+
 	void FUfbxMaterial::AddMaterials(UInterchangeBaseNodeContainer& NodeContainer)
 	{
 		using namespace UE::Interchange::Materials;
diff --git a/Engine/Plugins/Interchange/Runtime/Source/Parsers/Fbx/Private/Ufbx/UfbxMaterial.h b/Engine/Plugins/Interchange/Runtime/Source/Parsers/Fbx/Private/Ufbx/UfbxMaterial.h
--- a/Engine/Plugins/Interchange/Runtime/Source/Parsers/Fbx/Private/Ufbx/UfbxMaterial.h
+++ b/Engine/Plugins/Interchange/Runtime/Source/Parsers/Fbx/Private/Ufbx/UfbxMaterial.h
@@ -19,6 +19,7 @@ namespace UE::Interchange::Private
 		}
 
 		void AddMaterials(UInterchangeBaseNodeContainer& NodeContainer);
+		void AddAllTextures(UInterchangeBaseNodeContainer& NodeContainer);
 
 	private:
 		FUfbxParser& Parser;
diff --git a/Engine/Plugins/Interchange/Runtime/Source/Parsers/Fbx/Private/Ufbx/UfbxParser.cpp b/Engine/Plugins/Interchange/Runtime/Source/Parsers/Fbx/Private/Ufbx/UfbxParser.cpp
--- a/Engine/Plugins/Interchange/Runtime/Source/Parsers/Fbx/Private/Ufbx/UfbxParser.cpp
+++ b/Engine/Plugins/Interchange/Runtime/Source/Parsers/Fbx/Private/Ufbx/UfbxParser.cpp
@@ -220,7 +220,9 @@ void FUfbxParser::FillContainerWithFbxScene(UInterchangeBaseNodeContainer& NodeC
 	SceneConverter.InitHierarchy(NodeContainer);
 
 	// #ufbx_todo: CleanupFbxData();
-	FUfbxMaterial(*this).AddMaterials(NodeContainer);
+	FUfbxMaterial UfbxMaterial(*this);
+	UfbxMaterial.AddMaterials(NodeContainer);
+	UfbxMaterial.AddAllTextures(NodeContainer);
 
 	FUfbxMesh Meshes(*this);
 	// Get meshes first, nodes reference them later in ProcessNodes
