Interchange GLTF:
- Added a material function that implements glTF materials using Substrate (also based on MaterialX for compatibility)
- Added support specularColor
- Added support for volume extensions (most likely to be updated)
- Added a parent material using the material function, keeping the same architecture as the Legacy materials

#virtualized
#Synced-CL 51585235

diff --git a/Engine/Plugins/Interchange/Runtime/Source/Import/Private/Gltf/InterchangeGLTFMaterial.cpp b/Engine/Plugins/Interchange/Runtime/Source/Import/Private/Gltf/InterchangeGLTFMaterial.cpp
--- a/Engine/Plugins/Interchange/Runtime/Source/Import/Private/Gltf/InterchangeGLTFMaterial.cpp
+++ b/Engine/Plugins/Interchange/Runtime/Source/Import/Private/Gltf/InterchangeGLTFMaterial.cpp
@@ -591,11 +591,26 @@ namespace UE::Interchange::GLTFMaterials
 
 					//Configuration
 					SetBool(Inputs::Configuration::bHasSpecularTexture, bHasSpecularTexture);
+
+					//SpecularColorTexture
+					//SpecularColorTexture_OffsetScale
+					//SpecularColorTexture_Rotation
+					//SpecularColorTexture_TexCoord
+					//SpecularColorTexture_TilingMethod
+					bool bHasSpecularColorTexture = SetMap(Inputs::SpecularColorTexture, GltfMaterial.Specular.SpecularColorMap);
+
+					//SpecularColorFactor
+					const FVector3f SpecularColorFactor(GltfMaterial.Specular.SpecularColorFactor.X, GltfMaterial.Specular.SpecularColorFactor.Y, GltfMaterial.Specular.SpecularColorFactor.Z);
+					SetVec3(Inputs::SpecularColorFactor, SpecularColorFactor, { 1, 1, 1});
+
+					//Configuration
+					SetBool(Inputs::Configuration::bHasSpecularColorTexture, bHasSpecularColorTexture);
 				}
 				else
 				{
 					//Configuration
 					SetBool(Inputs::Configuration::bHasSpecularTexture, false);
+					SetBool(Inputs::Configuration::bHasSpecularColorTexture, false);
 				}
 
 				if (GltfMaterial.Iridescence.bHasIridescence)
@@ -670,6 +685,34 @@ namespace UE::Interchange::GLTFMaterials
 				SetScalar(Inputs::IOR, GltfMaterial.IOR, 1.5f);
 			}
 
+			if (GltfMaterial.Volume.bHasVolume)
+			{
+				//ThicknessTexture
+				//ThicknessTexture_OffsetScale
+				//ThicknessTexture_Rotation
+				//ThicknessTexture_TexCoord
+				//ThicknessTexture_TilingMethod
+				bool bHasThicknessTexture = SetMap(Inputs::ThicknessTexture, GltfMaterial.Volume.ThicknessTexture);
+
+				//ThicknessFactor
+				SetScalar(Inputs::ThicknessFactor, GltfMaterial.Volume.Thickness, 0.f);
+
+				//AttenuationDistance
+				SetScalar(Inputs::AttenuationDistance, GltfMaterial.Volume.AttenuationDistance, 10.f);
+
+				//AttenuationColor
+				const FVector3f AttenuationColor(GltfMaterial.Volume.AttenuationColor.X, GltfMaterial.Volume.AttenuationColor.Y, GltfMaterial.Volume.AttenuationColor.Z);
+				SetVec3(Inputs::AttenuationColor, AttenuationColor, FVector3f::OneVector);
+
+				//Configuration
+				SetBool(Inputs::Configuration::bHasThicknessTexture, bHasThicknessTexture);
+			}
+			else
+			{
+				//Configuration
+				SetBool(Inputs::Configuration::bHasThicknessTexture, false);
+			}
+
 			//Transmission ClearCoat and Sheen are Model specific
 			// Meaning we don't have to do exclusion configs.
 			if (GltfMaterial.bHasTransmission)
diff --git a/Engine/Plugins/Interchange/Runtime/Source/Import/Private/Gltf/InterchangeGltfTranslator.cpp b/Engine/Plugins/Interchange/Runtime/Source/Import/Private/Gltf/InterchangeGltfTranslator.cpp
--- a/Engine/Plugins/Interchange/Runtime/Source/Import/Private/Gltf/InterchangeGltfTranslator.cpp
+++ b/Engine/Plugins/Interchange/Runtime/Source/Import/Private/Gltf/InterchangeGltfTranslator.cpp
@@ -71,6 +71,7 @@ static const TArray<FString> ImporterSupportedExtensions = {
 	GLTF::ToString(GLTF::EExtension::KHR_MaterialsPbrSpecularGlossiness),
 	GLTF::ToString(GLTF::EExtension::KHR_MaterialsEmissiveStrength),
 	GLTF::ToString(GLTF::EExtension::KHR_MaterialsIridescence),
+	GLTF::ToString(GLTF::EExtension::KHR_MaterialsVolume),
 	GLTF::ToString(GLTF::EExtension::MSFT_PackingOcclusionRoughnessMetallic),
 	GLTF::ToString(GLTF::EExtension::MSFT_PackingNormalRoughnessMetallic),
 	/* Textures */
@@ -750,6 +751,9 @@ bool UInterchangeGLTFTranslator::Translate( UInterchangeBaseNodeContainer& NodeC
 			SetTextureSRGB(NodeContainer, GltfMaterial.ClearCoat.RoughnessMap, false);
 			SetTextureSRGB(NodeContainer, GltfMaterial.Transmission.TransmissionMap, false);
 
+			// TODO: Spec specifies linear space, but it seems to give better results in sRGB especially on the Dragon
+			SetTextureSRGB(NodeContainer, GltfMaterial.Volume.ThicknessTexture, true);
+
 			const FString ShaderGraphNodeUid = UInterchangeShaderGraphNode::MakeNodeUid(GltfMaterial.UniqueId);
 			bool bUseVertexColor = MaterialsUsedOnMeshesWithVertexColor.Contains(ShaderGraphNodeUid);
 
diff --git a/Engine/Plugins/Interchange/Runtime/Source/Import/Public/Gltf/InterchangeGLTFMaterial.h b/Engine/Plugins/Interchange/Runtime/Source/Import/Public/Gltf/InterchangeGLTFMaterial.h
--- a/Engine/Plugins/Interchange/Runtime/Source/Import/Public/Gltf/InterchangeGLTFMaterial.h
+++ b/Engine/Plugins/Interchange/Runtime/Source/Import/Public/Gltf/InterchangeGLTFMaterial.h
@@ -15,15 +15,15 @@ INTERCHANGEIMPORT_API extern const FString InterchangeGltfMaterialAttributeIdent
 //OffsetScale for MaterialInstances
 //Offset_X Offset_Y Scale_X Scale_Y for Materials
 #define INTERCHANGE_GLTF_STRINGIFY(x) #x
-#define DECLARE_INTERCHANGE_GLTF_MI_MAP(MapName) const FString MapName##Texture = TEXT(INTERCHANGE_GLTF_STRINGIFY(MapName ## Texture)); \
-												 const FString MapName##Texture_OffsetScale = TEXT(INTERCHANGE_GLTF_STRINGIFY(MapName ## Texture_OffsetScale)); \
-												 const FString MapName##Texture_Offset_X = TEXT(INTERCHANGE_GLTF_STRINGIFY(MapName ## Texture_Offset_X)); \
-												 const FString MapName##Texture_Offset_Y = TEXT(INTERCHANGE_GLTF_STRINGIFY(MapName ## Texture_Offset_Y)); \
-												 const FString MapName##Texture_Scale_X = TEXT(INTERCHANGE_GLTF_STRINGIFY(MapName ## Texture_Scale_X)); \
-												 const FString MapName##Texture_Scale_Y = TEXT(INTERCHANGE_GLTF_STRINGIFY(MapName ## Texture_Scale_Y)); \
-												 const FString MapName##Texture_Rotation = TEXT(INTERCHANGE_GLTF_STRINGIFY(MapName ## Texture_Rotation)); \
-												 const FString MapName##Texture_TexCoord = TEXT(INTERCHANGE_GLTF_STRINGIFY(MapName ## Texture_TexCoord)); \
-												 const FString MapName##Texture_TilingMethod = TEXT(INTERCHANGE_GLTF_STRINGIFY(MapName ## Texture_TilingMethod));
+#define DECLARE_INTERCHANGE_GLTF_MI_MAP(MapName) inline const FString MapName##Texture = TEXT(INTERCHANGE_GLTF_STRINGIFY(MapName ## Texture)); \
+												 inline const FString MapName##Texture_OffsetScale = TEXT(INTERCHANGE_GLTF_STRINGIFY(MapName ## Texture_OffsetScale)); \
+												 inline const FString MapName##Texture_Offset_X = TEXT(INTERCHANGE_GLTF_STRINGIFY(MapName ## Texture_Offset_X)); \
+												 inline const FString MapName##Texture_Offset_Y = TEXT(INTERCHANGE_GLTF_STRINGIFY(MapName ## Texture_Offset_Y)); \
+												 inline const FString MapName##Texture_Scale_X = TEXT(INTERCHANGE_GLTF_STRINGIFY(MapName ## Texture_Scale_X)); \
+												 inline const FString MapName##Texture_Scale_Y = TEXT(INTERCHANGE_GLTF_STRINGIFY(MapName ## Texture_Scale_Y)); \
+												 inline const FString MapName##Texture_Rotation = TEXT(INTERCHANGE_GLTF_STRINGIFY(MapName ## Texture_Rotation)); \
+												 inline const FString MapName##Texture_TexCoord = TEXT(INTERCHANGE_GLTF_STRINGIFY(MapName ## Texture_TexCoord)); \
+												 inline const FString MapName##Texture_TilingMethod = TEXT(INTERCHANGE_GLTF_STRINGIFY(MapName ## Texture_TilingMethod));
 
 namespace UE::Interchange::GLTFMaterials
 {
@@ -33,120 +33,130 @@ namespace UE::Interchange::GLTFMaterials
 		///PostFixes
 		namespace PostFix
 		{
-			const FString Color_RGB = TEXT("_RGB");
-			const FString Color_A = TEXT("_A");
+			inline const FString Color_RGB = TEXT("_RGB");
+			inline const FString Color_A = TEXT("_A");
 
-			const FString TexCoord = TEXT("_TexCoord");
+			inline const FString TexCoord = TEXT("_TexCoord");
 
-			const FString OffsetX = TEXT("_Offset_X");
-			const FString OffsetY = TEXT("_Offset_Y");
-			const FString ScaleX = TEXT("_Scale_X");
-			const FString ScaleY = TEXT("_Scale_Y");
+			inline const FString OffsetX = TEXT("_Offset_X");
+			inline const FString OffsetY = TEXT("_Offset_Y");
+			inline const FString ScaleX = TEXT("_Scale_X");
+			inline const FString ScaleY = TEXT("_Scale_Y");
 
-			const FString OffsetScale = TEXT("_OffsetScale");
+			inline const FString OffsetScale = TEXT("_OffsetScale");
 
-			const FString Rotation = TEXT("_Rotation");
+			inline const FString Rotation = TEXT("_Rotation");
 
-			const FString TilingMethod = TEXT("_TilingMethod");
+			inline const FString TilingMethod = TEXT("_TilingMethod");
 		}
 
 		//MetalRoughness specific:
 		DECLARE_INTERCHANGE_GLTF_MI_MAP(BaseColor)
-		const FString BaseColorFactor = TEXT("BaseColorFactor");
-		const FString BaseColorFactor_RGB = BaseColorFactor + PostFix::Color_RGB; //Connection to inputs from BaseColorFactor.RGB
-		const FString BaseColorFactor_A = BaseColorFactor + PostFix::Color_A; //Connection to inputs from BaseColorFactor.A
+		inline const FString BaseColorFactor = TEXT("BaseColorFactor");
+		inline const FString BaseColorFactor_RGB = BaseColorFactor + PostFix::Color_RGB; //Connection to inputs from BaseColorFactor.RGB
+		inline const FString BaseColorFactor_A = BaseColorFactor + PostFix::Color_A; //Connection to inputs from BaseColorFactor.A
 
 		DECLARE_INTERCHANGE_GLTF_MI_MAP(MetallicRoughness)
-		const FString MetallicFactor = TEXT("MetallicFactor");
-		const FString RoughnessFactor = TEXT("RoughnessFactor");
+		inline const FString MetallicFactor = TEXT("MetallicFactor");
+		inline const FString RoughnessFactor = TEXT("RoughnessFactor");
 
 		DECLARE_INTERCHANGE_GLTF_MI_MAP(Specular)
-		const FString SpecularFactor = TEXT("SpecularFactor");
+		inline const FString SpecularFactor = TEXT("SpecularFactor");
 
+		DECLARE_INTERCHANGE_GLTF_MI_MAP(SpecularColor)
+		inline const FString SpecularColorFactor = TEXT("SpecularColorFactor");
 
 		//SpecularGlossiness specific
 		DECLARE_INTERCHANGE_GLTF_MI_MAP(Diffuse)
-		const FString DiffuseFactor = TEXT("DiffuseFactor");
-		const FString DiffuseFactor_RGB = DiffuseFactor + PostFix::Color_RGB; //Connection to inputs from BaseColorFactor.RGB
-		const FString DiffuseFactor_A = DiffuseFactor + PostFix::Color_A; //Connection to inputs from BaseColorFactor.A
+		inline const FString DiffuseFactor = TEXT("DiffuseFactor");
+		inline const FString DiffuseFactor_RGB = DiffuseFactor + PostFix::Color_RGB; //Connection to inputs from BaseColorFactor.RGB
+		inline const FString DiffuseFactor_A = DiffuseFactor + PostFix::Color_A; //Connection to inputs from BaseColorFactor.A
 
 		DECLARE_INTERCHANGE_GLTF_MI_MAP(SpecularGlossiness)
-		const FString SpecFactor = TEXT("SpecFactor");
-		const FString GlossinessFactor = TEXT("GlossinessFactor");
+		inline const FString SpecFactor = TEXT("SpecFactor");
+		inline const FString GlossinessFactor = TEXT("GlossinessFactor");
 
 
 		//Generic:
 		DECLARE_INTERCHANGE_GLTF_MI_MAP(Normal)
-		const FString NormalScale = TEXT("NormalScale");
+		inline const FString NormalScale = TEXT("NormalScale");
 
 		DECLARE_INTERCHANGE_GLTF_MI_MAP(Emissive)
-		const FString EmissiveFactor = TEXT("EmissiveFactor");
-		const FString EmissiveStrength = TEXT("EmissiveStrength");
+		inline const FString EmissiveFactor = TEXT("EmissiveFactor");
+		inline const FString EmissiveStrength = TEXT("EmissiveStrength");
 
 		DECLARE_INTERCHANGE_GLTF_MI_MAP(Occlusion)
-		const FString OcclusionStrength = TEXT("OcclusionStrength");
+		inline const FString OcclusionStrength = TEXT("OcclusionStrength");
 
-		const FString IOR = TEXT("IOR");
+		inline const FString IOR = TEXT("IOR");
 
-		const FString AlphaCutoff = TEXT("AlphaCutoff");
-		const FString AlphaMode = TEXT("AlphaMode"); //Specifically for Transmission Materials
+		inline const FString AlphaCutoff = TEXT("AlphaCutoff");
+		inline const FString AlphaMode = TEXT("AlphaMode"); //Specifically for Transmission Materials
 
 		//ClearCoat specific:
 		DECLARE_INTERCHANGE_GLTF_MI_MAP(ClearCoat)
-		const FString ClearCoatFactor = TEXT("ClearCoatFactor");
+		inline const FString ClearCoatFactor = TEXT("ClearCoatFactor");
 
 		DECLARE_INTERCHANGE_GLTF_MI_MAP(ClearCoatRoughness)
-		const FString ClearCoatRoughnessFactor = TEXT("ClearCoatRoughnessFactor");
+		inline const FString ClearCoatRoughnessFactor = TEXT("ClearCoatRoughnessFactor");
 
 		DECLARE_INTERCHANGE_GLTF_MI_MAP(ClearCoatNormal)
-		const FString ClearCoatNormalScale = TEXT("ClearCoatNormalScale");
+		inline const FString ClearCoatNormalScale = TEXT("ClearCoatNormalScale");
 
 
 		//Sheen specific:
 		DECLARE_INTERCHANGE_GLTF_MI_MAP(SheenColor)
-		const FString SheenColorFactor = TEXT("SheenColorFactor");
+		inline const FString SheenColorFactor = TEXT("SheenColorFactor");
 
 		DECLARE_INTERCHANGE_GLTF_MI_MAP(SheenRoughness)
-		const FString SheenRoughnessFactor = TEXT("SheenRoughnessFactor");
+		inline const FString SheenRoughnessFactor = TEXT("SheenRoughnessFactor");
 
 
 		//Transmission specific:
 		DECLARE_INTERCHANGE_GLTF_MI_MAP(Transmission)
-		const FString TransmissionFactor = TEXT("TransmissionFactor");
+		inline const FString TransmissionFactor = TEXT("TransmissionFactor");
 
 
 		//Iridescence specific:
-		const FString IridescenceIOR = TEXT("IridescenceIOR");
+		inline const FString IridescenceIOR = TEXT("IridescenceIOR");
 		
 		DECLARE_INTERCHANGE_GLTF_MI_MAP(Iridescence)
-		const FString IridescenceFactor = TEXT("IridescenceFactor");
+		inline const FString IridescenceFactor = TEXT("IridescenceFactor");
 
 		DECLARE_INTERCHANGE_GLTF_MI_MAP(IridescenceThickness)
-		const FString IridescenceThicknessMinimum = TEXT("IridescenceThicknessMinimum");
-		const FString IridescenceThicknessMaximum = TEXT("IridescenceThicknessMaximum");
+		inline const FString IridescenceThicknessMinimum = TEXT("IridescenceThicknessMinimum");
+		inline const FString IridescenceThicknessMaximum = TEXT("IridescenceThicknessMaximum");
 
 
 		//Anisotropy Specific:
 		DECLARE_INTERCHANGE_GLTF_MI_MAP(Anisotropy);
-		const FString AnisotropyStrength = TEXT("AnisotropyStrength");
-		const FString AnisotropyRotation = TEXT("AnisotropyRotation");
+		inline const FString AnisotropyStrength = TEXT("AnisotropyStrength");
+		inline const FString AnisotropyRotation = TEXT("AnisotropyRotation");
+
+		//Volume Specific:
+		DECLARE_INTERCHANGE_GLTF_MI_MAP(Thickness);
+		inline const FString ThicknessFactor = TEXT("ThicknessFactor");
+		inline const FString AttenuationDistance = TEXT("AttenuationDistance");
+		inline const FString AttenuationColor = TEXT("AttenuationColor");
 
 		namespace Configuration
 		{
 			//For StaticSwitch optimizations:
 			// Note: only usable in Editor, the Default values are true so that Runtime will still use the complete graph (without optimization).
-			const FString bHasBaseColorTexture = TEXT("bHasBaseColorTexture");
-			const FString bHasMetallicRoughnessTexture = TEXT("bHasMetallicRoughnessTexture");
-			const FString bHasDiffuseSpecGlossTexture = TEXT("bHasDiffuseSpecGlossTexture");
-			const FString bHasEmissiveTexture = TEXT("bHasEmissiveTexture");
-			const FString bHasNormalTexture = TEXT("bHasNormalTexture");
-			const FString bHasOcclusionTexture = TEXT("bHasOcclusionTexture");
-			const FString bHasSpecularTexture = TEXT("bHasSpecularTexture");
-			const FString bHasClearCoatTexture = TEXT("bHasClearCoatTexture");
-			const FString bHasSheenTexture = TEXT("bHasSheenTexture");
-			const FString bHasTransmissionTexture = TEXT("bHasTransmissionTexture");
-			const FString bHasIridescence = TEXT("bHasIridescence");
-			const FString bHasAnisotropyTextureAndOrRotation = TEXT("bHasAnisotropyTextureAndOrRotation"); //AnistropyRotation can affect the Tangents so we check
+			inline const FString bHasBaseColorTexture = TEXT("bHasBaseColorTexture");
+			inline const FString bHasMetallicRoughnessTexture = TEXT("bHasMetallicRoughnessTexture");
+			inline const FString bHasDiffuseSpecGlossTexture = TEXT("bHasDiffuseSpecGlossTexture");
+			inline const FString bHasEmissiveTexture = TEXT("bHasEmissiveTexture");
+			inline const FString bHasNormalTexture = TEXT("bHasNormalTexture");
+			inline const FString bHasOcclusionTexture = TEXT("bHasOcclusionTexture");
+			inline const FString bHasSpecularTexture = TEXT("bHasSpecularTexture");
+			inline const FString bHasSpecularColorTexture = TEXT("bHasSpecularColorTexture");
+			inline const FString bHasClearCoatTexture = TEXT("bHasClearCoatTexture");
+			inline const FString bHasSheenTexture = TEXT("bHasSheenTexture");
+			inline const FString bHasTransmissionTexture = TEXT("bHasTransmissionTexture");
+			inline const FString bHasIridescence = TEXT("bHasIridescence");
+			inline const FString bHasAnisotropyTextureAndOrRotation = TEXT("bHasAnisotropyTextureAndOrRotation"); //AnistropyRotation can affect the Tangents so we check
+			inline const FString bHasThicknessTexture = TEXT("bHasThicknessTexture");
 		}
 	}
 
diff --git a/Engine/Plugins/Interchange/Runtime/Source/Parsers/GLTFCore/Private/GLTF/ExtensionsHandler.cpp b/Engine/Plugins/Interchange/Runtime/Source/Parsers/GLTFCore/Private/GLTF/ExtensionsHandler.cpp
--- a/Engine/Plugins/Interchange/Runtime/Source/Parsers/GLTFCore/Private/GLTF/ExtensionsHandler.cpp
+++ b/Engine/Plugins/Interchange/Runtime/Source/Parsers/GLTFCore/Private/GLTF/ExtensionsHandler.cpp
@@ -291,6 +291,7 @@ namespace GLTF
 			EExtension::KHR_MaterialsEmissiveStrength,
 			EExtension::KHR_MaterialsIridescence,
 			EExtension::KHR_MaterialsAnisotropy,
+			EExtension::KHR_MaterialsVolume,
 			EExtension::MSFT_PackingOcclusionRoughnessMetallic,
 			EExtension::MSFT_PackingNormalRoughnessMetallic 
 		};
@@ -393,7 +394,7 @@ namespace GLTF
 
 					Material.bHasSpecular = true;
 					Material.Specular.SpecularFactor = GetScalar(Specular, TEXT("specularFactor"), 1.0f);
-					Material.Specular.SpecularColorFactor = GetVec3(Specular, TEXT("specularColorFactor"));
+					Material.Specular.SpecularColorFactor = GetVec3(Specular, TEXT("specularColorFactor"), FVector::OneVector);
 					GLTF::SetTextureMap(Specular, TEXT("specularTexture"), nullptr, Asset->Textures, Material.Specular.SpecularMap, Messages);
 					GLTF::SetTextureMap(Specular, TEXT("specularColorTexture"), nullptr, Asset->Textures, Material.Specular.SpecularColorMap, Messages);
 
@@ -455,6 +456,7 @@ namespace GLTF
 					Material.Iridescence.Thickness.Minimum = GetScalar(Iridescence, TEXT("iridescenceThicknessMinimum"), 100.0f);
 					Material.Iridescence.Thickness.Maximum = GetScalar(Iridescence, TEXT("iridescenceThicknessMaximum"), 400.0f);
 					GLTF::SetTextureMap(Iridescence, TEXT("iridescenceThicknessTexture"), nullptr, Asset->Textures, Material.Iridescence.Thickness.Texture, Messages);
+					Asset->ProcessedExtensions.Add(EExtension::KHR_MaterialsIridescence);
 				}
 				break;
 				case EExtension::KHR_MaterialsAnisotropy:
@@ -466,6 +468,19 @@ namespace GLTF
 					Material.Anisotropy.Strength = GetScalar(Anisotropy, TEXT("anisotropyStrength"), 0.0f);
 					Material.Anisotropy.Rotation = GetScalar(Anisotropy, TEXT("anisotropyRotation"), 0.0f);
 					GLTF::SetTextureMap(Anisotropy, TEXT("anisotropyTexture"), nullptr, Asset->Textures, Material.Anisotropy.Texture, Messages);
+					Asset->ProcessedExtensions.Add(EExtension::KHR_MaterialsAnisotropy);
+				}
+				break;
+				case EExtension::KHR_MaterialsVolume:
+				{
+					const FJsonObject& Volume = ExtObj;
+					Material.Volume.bHasVolume = true;
+
+					Material.Volume.Thickness = GetScalar(Volume, TEXT("thicknessFactor"));
+					Material.Volume.AttenuationDistance = GetScalar(Volume, TEXT("attenuationDistance"), 1.f);
+					Material.Volume.AttenuationColor = GetVec3(Volume, TEXT("attenuationColor"), FVector::OneVector);
+					GLTF::SetTextureMap(Volume, TEXT("thicknessTexture"), nullptr, Asset->Textures, Material.Volume.ThicknessTexture, Messages);
+					Asset->ProcessedExtensions.Add(EExtension::KHR_MaterialsVolume);
 				}
 				break;
 				default:
diff --git a/Engine/Plugins/Interchange/Runtime/Source/Parsers/GLTFCore/Private/GLTF/GLTFAsset.cpp b/Engine/Plugins/Interchange/Runtime/Source/Parsers/GLTFCore/Private/GLTF/GLTFAsset.cpp
--- a/Engine/Plugins/Interchange/Runtime/Source/Parsers/GLTFCore/Private/GLTF/GLTFAsset.cpp
+++ b/Engine/Plugins/Interchange/Runtime/Source/Parsers/GLTFCore/Private/GLTF/GLTFAsset.cpp
@@ -475,6 +475,8 @@ namespace GLTF
 				return TEXT("KHR_materials_iridescence");
 			case GLTF::EExtension::KHR_MaterialsAnisotropy:
 				return TEXT("KHR_materials_anisotropy");
+			case GLTF::EExtension::KHR_MaterialsVolume:
+				return TEXT("KHR_materials_volume");
 			case GLTF::EExtension::KHR_TextureTransform:
 				return TEXT("KHR_texture_transform");
 			case GLTF::EExtension::KHR_DracoMeshCompression:
diff --git a/Engine/Plugins/Interchange/Runtime/Source/Parsers/GLTFCore/Public/GLTFAsset.h b/Engine/Plugins/Interchange/Runtime/Source/Parsers/GLTFCore/Public/GLTFAsset.h
--- a/Engine/Plugins/Interchange/Runtime/Source/Parsers/GLTFCore/Public/GLTFAsset.h
+++ b/Engine/Plugins/Interchange/Runtime/Source/Parsers/GLTFCore/Public/GLTFAsset.h
@@ -39,6 +39,7 @@ namespace GLTF
 		KHR_MaterialsEmissiveStrength,
 		KHR_MaterialsIridescence,
 		KHR_MaterialsAnisotropy,
+		KHR_MaterialsVolume,
 		KHR_TextureTransform,
 		KHR_DracoMeshCompression,
 		KHR_LightsPunctual,
diff --git a/Engine/Plugins/Interchange/Runtime/Source/Parsers/GLTFCore/Public/GLTFMaterial.h b/Engine/Plugins/Interchange/Runtime/Source/Parsers/GLTFCore/Public/GLTFMaterial.h
--- a/Engine/Plugins/Interchange/Runtime/Source/Parsers/GLTFCore/Public/GLTFMaterial.h
+++ b/Engine/Plugins/Interchange/Runtime/Source/Parsers/GLTFCore/Public/GLTFMaterial.h
@@ -156,6 +156,15 @@ namespace GLTF
 			bool bHasAnisotropy = false;
 		};
 
+		struct FVolume
+		{
+			FTextureMap ThicknessTexture;
+			float Thickness = 0.f;
+			float AttenuationDistance = 10.f; // Theoretically +∞, should definitely not be 0, 10.f if a good compromise as it renders no attenuation
+			FVector AttenuationColor{ 1, 1, 1 };
+			bool bHasVolume = false;
+		};
+
 		FString Name;
 
 		// PBR properties
@@ -195,6 +204,7 @@ namespace GLTF
 		bool            bHasEmissiveStrength;
 		FIridescence    Iridescence;
 		FAnisotropy		Anisotropy;
+		FVolume         Volume;
 
 		TMap<FString, FString> Extras;
 
diff --git a/Engine/Plugins/Interchange/Runtime/Source/Pipelines/Private/InterchangeglTFPipeline.cpp b/Engine/Plugins/Interchange/Runtime/Source/Pipelines/Private/InterchangeglTFPipeline.cpp
--- a/Engine/Plugins/Interchange/Runtime/Source/Pipelines/Private/InterchangeglTFPipeline.cpp
+++ b/Engine/Plugins/Interchange/Runtime/Source/Pipelines/Private/InterchangeglTFPipeline.cpp
@@ -2,23 +2,26 @@
 
 #include "InterchangeglTFPipeline.h"
 
-#include "InterchangePipelineHelper.h"
-#include "InterchangePipelineLog.h"
+#include "Gltf/InterchangeGLTFMaterial.h"
 
+#include "InterchangeGenericMaterialPipeline.h"
+#include "InterchangeImportModule.h"
 #include "InterchangeManager.h"
 #include "InterchangeMaterialFactoryNode.h"
-#include "InterchangeGenericMaterialPipeline.h"
 #include "InterchangeShaderGraphNode.h"
-#include "Nodes/InterchangeSourceNode.h"
 
-#include "Gltf/InterchangeGLTFMaterial.h"
+#include "InterchangePipelineHelper.h"
+#include "InterchangePipelineLog.h"
 
+#include "Materials/Material.h"
 #include "Materials/MaterialInstanceConstant.h"
 #include "Materials/MaterialInstanceDynamic.h"
 
+#include "Nodes/InterchangeSourceNode.h"
 #include UE_INLINE_GENERATED_CPP_BY_NAME(InterchangeglTFPipeline)
 
-const TArray<FString> UGLTFPipelineSettings::ExpectedMaterialInstanceIdentifiers = {TEXT("MI_Default_Opaque"), TEXT("MI_Default_Mask"), TEXT("MI_Default_Blend"), 
+const TArray<FString> UGLTFPipelineSettings::ExpectedMaterialInstanceIdentifiers = {TEXT("M_GLTF"), // Substrate material that replaces all of the material instances below
+																					TEXT("MI_Default_Opaque"), TEXT("MI_Default_Mask"), TEXT("MI_Default_Blend"), 
 																					TEXT("MI_Unlit_Opaque"), TEXT("MI_Unlit_Mask"), TEXT("MI_Unlit_Blend"), 
 																					TEXT("MI_ClearCoat_Opaque"), TEXT("MI_ClearCoat_Mask"), TEXT("MI_ClearCoat_Blend"),
 																					TEXT("MI_Sheen_Opaque"), TEXT("MI_Sheen_Mask"), TEXT("MI_Sheen_Blend"), 
@@ -48,11 +51,11 @@ TArray<FString> UGLTFPipelineSettings::ValidateMaterialInstancesAndParameters()
 		NotCoveredIdentifiersParameters.Add(TEXT("[") + ExpectedIdentifier + TEXT("]: MaterialInstance not found for Identifier."));
 	}
 
-	for (const TPair<FString, FSoftObjectPath>& MaterialParent : MaterialParents)
+	auto ValidateParameters = [this, &NotCoveredIdentifiersParameters]<typename T>(const FString& Identifier, const FSoftObjectPath& ParentPath)
 	{
-		TSet<FString> ExpectedParameters = GenerateExpectedParametersList(MaterialParent.Key);
+		TSet<FString> ExpectedParameters = GenerateExpectedParametersList(Identifier);
 
-		if (UMaterialInstance* ParentMaterialInstance = Cast<UMaterialInstance>(MaterialParent.Value.TryLoad()))
+		if (T* ParentMaterialInstance = Cast<T>(ParentPath.TryLoad()))
 		{
 			TArray<FGuid> ParameterIds;
 			TArray<FMaterialParameterInfo> ScalarParameterInfos;
@@ -78,7 +81,26 @@ TArray<FString> UGLTFPipelineSettings::ValidateMaterialInstancesAndParameters()
 
 		for (const FString& ExpectedParameter : ExpectedParameters)
 		{
-			NotCoveredIdentifiersParameters.Add(TEXT("[") + MaterialParent.Key + TEXT("]: Does not cover expected parameter: ") + ExpectedParameter + TEXT("."));
+			NotCoveredIdentifiersParameters.Add(TEXT("[") + Identifier + TEXT("]: Does not cover expected parameter: ") + ExpectedParameter + TEXT("."));
+		}
+	};
+
+	// For Substrate we only need to validate the first Material
+	if (IInterchangeImportModule::Get().IsSubstrateEnabled())
+	{
+		const FSoftObjectPath& SubstrateMaterialPath = MaterialParents[TEXT("M_GLTF")];
+		ValidateParameters.operator()<UMaterial>({}, SubstrateMaterialPath);
+	}
+	else
+	{
+		for (const TPair<FString, FSoftObjectPath>& MaterialParent : MaterialParents)
+		{
+			if (MaterialParent.Key == TEXT("M_GLTF"))
+			{
+				continue;
+			}
+
+			ValidateParameters.operator()<UMaterialInstance>(MaterialParent.Key, MaterialParent.Value);
 		}
 	}
 
@@ -100,43 +122,25 @@ TSet<FString> UGLTFPipelineSettings::GenerateExpectedParametersList(const FStrin
 			ExpectedParameters.Add(TextureName + Inputs::PostFix::TilingMethod);
 		};
 
-	if (Identifier.Contains(TEXT("_Unlit")))
-	{
-		AddTextureAndRelated(Inputs::BaseColorTexture);
-		ExpectedParameters.Add(Inputs::BaseColorFactor);
-
-		return ExpectedParameters;
-	}
-
-	//Generic ones:
+	if (IInterchangeImportModule::Get().IsSubstrateEnabled())
 	{
 		AddTextureAndRelated(Inputs::NormalTexture);
 		ExpectedParameters.Add(Inputs::NormalScale);
 
-		if (!Identifier.Contains(TEXT("Transmission")))
-		{
-			AddTextureAndRelated(Inputs::EmissiveTexture);
-			ExpectedParameters.Add(Inputs::EmissiveFactor);
-			ExpectedParameters.Add(Inputs::EmissiveStrength);
-		}
+		AddTextureAndRelated(Inputs::EmissiveTexture);
+		ExpectedParameters.Add(Inputs::EmissiveFactor);
+		ExpectedParameters.Add(Inputs::EmissiveStrength);
 
 		AddTextureAndRelated(Inputs::OcclusionTexture);
 		ExpectedParameters.Add(Inputs::OcclusionStrength);
 
-		if (!Identifier.Contains(TEXT("SpecularGlossiness")))
-		{
-			ExpectedParameters.Add(Inputs::IOR);
+		ExpectedParameters.Add(Inputs::IOR);
 
-			AddTextureAndRelated(Inputs::SpecularTexture);
-			ExpectedParameters.Add(Inputs::SpecularFactor);
-		}
-	}
+		AddTextureAndRelated(Inputs::SpecularTexture);
+		ExpectedParameters.Add(Inputs::SpecularFactor);
 
-	//Based on ShadingModel:
-
-	if (Identifier.Contains(TEXT("Default")))
-	{
-		//MetalRoughness Specific:
+		AddTextureAndRelated(Inputs::SpecularColorTexture);
+		ExpectedParameters.Add(Inputs::SpecularColorFactor);
 
 		AddTextureAndRelated(Inputs::BaseColorTexture);
 		ExpectedParameters.Add(Inputs::BaseColorFactor);
@@ -144,9 +148,7 @@ TSet<FString> UGLTFPipelineSettings::GenerateExpectedParametersList(const FStrin
 		AddTextureAndRelated(Inputs::MetallicRoughnessTexture);
 		ExpectedParameters.Add(Inputs::MetallicFactor);
 		ExpectedParameters.Add(Inputs::RoughnessFactor);
-	}
-	else if (Identifier.Contains(TEXT("ClearCoat")))
-	{
+
 		AddTextureAndRelated(Inputs::ClearCoatTexture);
 		ExpectedParameters.Add(Inputs::ClearCoatFactor);
 
@@ -155,30 +157,109 @@ TSet<FString> UGLTFPipelineSettings::GenerateExpectedParametersList(const FStrin
 
 		AddTextureAndRelated(Inputs::ClearCoatNormalTexture);
 		ExpectedParameters.Add(Inputs::ClearCoatNormalScale);
-	}
-	else if (Identifier.Contains(TEXT("Sheen")))
-	{
+
 		AddTextureAndRelated(Inputs::SheenColorTexture);
 		ExpectedParameters.Add(Inputs::SheenColorFactor);
 
 		AddTextureAndRelated(Inputs::SheenRoughnessTexture);
 		ExpectedParameters.Add(Inputs::SheenRoughnessFactor);
-	}
-	else if (Identifier.Contains(TEXT("Transmission")))
-	{
+
 		AddTextureAndRelated(Inputs::TransmissionTexture);
 		ExpectedParameters.Add(Inputs::TransmissionFactor);
-	}
-	else if (Identifier.Contains(TEXT("SpecularGlossiness")))
-	{
+
 		AddTextureAndRelated(Inputs::DiffuseTexture);
 		ExpectedParameters.Add(Inputs::DiffuseFactor);
 
 		AddTextureAndRelated(Inputs::SpecularGlossinessTexture);
 		ExpectedParameters.Add(Inputs::SpecFactor);
 		ExpectedParameters.Add(Inputs::GlossinessFactor);
+
+		AddTextureAndRelated(Inputs::ThicknessTexture);
+		ExpectedParameters.Add(Inputs::ThicknessFactor);
+		ExpectedParameters.Add(Inputs::AttenuationDistance);
+		ExpectedParameters.Add(Inputs::AttenuationColor);
 	}
+	else
+	{
+		if (Identifier.Contains(TEXT("_Unlit")))
+		{
+			AddTextureAndRelated(Inputs::BaseColorTexture);
+			ExpectedParameters.Add(Inputs::BaseColorFactor);
+
+			return ExpectedParameters;
+		}
+
+		//Generic ones:
+		{
+			AddTextureAndRelated(Inputs::NormalTexture);
+			ExpectedParameters.Add(Inputs::NormalScale);
+
+			if (!Identifier.Contains(TEXT("Transmission")))
+			{
+				AddTextureAndRelated(Inputs::EmissiveTexture);
+				ExpectedParameters.Add(Inputs::EmissiveFactor);
+				ExpectedParameters.Add(Inputs::EmissiveStrength);
+			}
+
+			AddTextureAndRelated(Inputs::OcclusionTexture);
+			ExpectedParameters.Add(Inputs::OcclusionStrength);
+
+			if (!Identifier.Contains(TEXT("SpecularGlossiness")))
+			{
+				ExpectedParameters.Add(Inputs::IOR);
+
+				AddTextureAndRelated(Inputs::SpecularTexture);
+				ExpectedParameters.Add(Inputs::SpecularFactor);
+			}
+		}
+
+		//Based on ShadingModel:
+
+		if (Identifier.Contains(TEXT("Default")))
+		{
+			//MetalRoughness Specific:
 
+			AddTextureAndRelated(Inputs::BaseColorTexture);
+			ExpectedParameters.Add(Inputs::BaseColorFactor);
+
+			AddTextureAndRelated(Inputs::MetallicRoughnessTexture);
+			ExpectedParameters.Add(Inputs::MetallicFactor);
+			ExpectedParameters.Add(Inputs::RoughnessFactor);
+		}
+		else if (Identifier.Contains(TEXT("ClearCoat")))
+		{
+			AddTextureAndRelated(Inputs::ClearCoatTexture);
+			ExpectedParameters.Add(Inputs::ClearCoatFactor);
+
+			AddTextureAndRelated(Inputs::ClearCoatRoughnessTexture);
+			ExpectedParameters.Add(Inputs::ClearCoatRoughnessFactor);
+
+			AddTextureAndRelated(Inputs::ClearCoatNormalTexture);
+			ExpectedParameters.Add(Inputs::ClearCoatNormalScale);
+		}
+		else if (Identifier.Contains(TEXT("Sheen")))
+		{
+			AddTextureAndRelated(Inputs::SheenColorTexture);
+			ExpectedParameters.Add(Inputs::SheenColorFactor);
+
+			AddTextureAndRelated(Inputs::SheenRoughnessTexture);
+			ExpectedParameters.Add(Inputs::SheenRoughnessFactor);
+		}
+		else if (Identifier.Contains(TEXT("Transmission")))
+		{
+			AddTextureAndRelated(Inputs::TransmissionTexture);
+			ExpectedParameters.Add(Inputs::TransmissionFactor);
+		}
+		else if (Identifier.Contains(TEXT("SpecularGlossiness")))
+		{
+			AddTextureAndRelated(Inputs::DiffuseTexture);
+			ExpectedParameters.Add(Inputs::DiffuseFactor);
+
+			AddTextureAndRelated(Inputs::SpecularGlossinessTexture);
+			ExpectedParameters.Add(Inputs::SpecFactor);
+			ExpectedParameters.Add(Inputs::GlossinessFactor);
+		}
+	}
 	return ExpectedParameters;
 }
 
@@ -226,7 +307,22 @@ UInterchangeMaterialInstanceFactoryNode* UGLTFPipelineSettings::BuildMaterialIns
 	}
 	GltfAttributeKeys.Remove(TEXT("ParentIdentifier"));
 
+	bool bSubstrate = IInterchangeImportModule::Get().IsSubstrateEnabled();
+	bool bTwoSided = ParentIdentifier.Contains(TEXT("DS"), ESearchCase::CaseSensitive, ESearchDir::FromEnd) ? true : false;
+	EBlendMode BlendMode = 
+		ParentIdentifier.Contains(TEXT("Blend"), ESearchCase::CaseSensitive, ESearchDir::FromEnd)        ||
+		ParentIdentifier.Contains(TEXT("Mask"), ESearchCase::CaseSensitive, ESearchDir::FromEnd)         ||
+		ParentIdentifier.Contains(TEXT("Transmission"), ESearchCase::CaseSensitive, ESearchDir::FromEnd) ?
+		EBlendMode::BLEND_TranslucentColoredTransmittance                                                :
+		EBlendMode::BLEND_Opaque;
+	bool bSpecGloss = ParentIdentifier.Contains(TEXT("SpecularGlossiness"), ESearchCase::CaseSensitive) ? true : false;
+
 	FString Parent;
+	if(bSubstrate)
+	{
+		ParentIdentifier = TEXT("M_GLTF");
+	}
+
 	if (const FSoftObjectPath* ObjectPath = MaterialParents.Find(ParentIdentifier))
 	{
 		Parent = ObjectPath->GetAssetPathString();
@@ -256,6 +352,13 @@ UInterchangeMaterialInstanceFactoryNode* UGLTFPipelineSettings::BuildMaterialIns
 		UInterchangeEditorUtilitiesBase* EditorUtilities = UInterchangeManager::GetInterchangeManager().GetEditorUtilities();
 		const UClass* MaterialClass = (EditorUtilities && EditorUtilities->IsRuntimeOrPIE()) ? UMaterialInstanceDynamic::StaticClass() : UMaterialInstanceConstant::StaticClass();
 		MaterialInstanceFactoryNode->SetCustomInstanceClassName(MaterialClass->GetPathName());
+		
+		if (bSubstrate)
+		{
+			MaterialInstanceFactoryNode->SetCustomBlendMode(BlendMode);
+			MaterialInstanceFactoryNode->SetCustomTwoSided(bTwoSided);
+			MaterialInstanceFactoryNode->AddBooleanAttribute(UInterchangeShaderPortsAPI::MakeInputValueKey(TEXT("bHasDiffuseSpecGloss")), bSpecGloss);
+		}
 
 		for (const TPair<FString, UE::Interchange::FAttributeKey>& GltfAttributeKey : GltfAttributeKeys)
 		{
