Interchange glTF:
- Set the AlphaMode for the parent material as it is now used by Substrate

#virtualized
#Synced-CL 51585235

diff --git a/Engine/Plugins/Interchange/Runtime/Source/Import/Private/Gltf/InterchangeGLTFMaterial.cpp b/Engine/Plugins/Interchange/Runtime/Source/Import/Private/Gltf/InterchangeGLTFMaterial.cpp
--- a/Engine/Plugins/Interchange/Runtime/Source/Import/Private/Gltf/InterchangeGLTFMaterial.cpp
+++ b/Engine/Plugins/Interchange/Runtime/Source/Import/Private/Gltf/InterchangeGLTFMaterial.cpp
@@ -431,6 +431,10 @@ namespace UE::Interchange::GLTFMaterials
 		////////////////////////////////////////////////////////////////////////////////////
 		void ParseMaterialValues()
 		{
+			// AlphaMode used now by Substrate to handle every case
+			SetScalar(Inputs::AlphaMode, float(GltfMaterial.AlphaMode), 0.f);
+
+
 			if (GltfMaterial.bIsUnlitShadingModel)
 			{
 				//BaseColorTexture
diff --git a/Engine/Plugins/Interchange/Runtime/Source/Import/Public/Gltf/InterchangeGLTFMaterial.h b/Engine/Plugins/Interchange/Runtime/Source/Import/Public/Gltf/InterchangeGLTFMaterial.h
--- a/Engine/Plugins/Interchange/Runtime/Source/Import/Public/Gltf/InterchangeGLTFMaterial.h
+++ b/Engine/Plugins/Interchange/Runtime/Source/Import/Public/Gltf/InterchangeGLTFMaterial.h
@@ -91,7 +91,7 @@ namespace UE::Interchange::GLTFMaterials
 		inline const FString IOR = TEXT("IOR");
 
 		inline const FString AlphaCutoff = TEXT("AlphaCutoff");
-		inline const FString AlphaMode = TEXT("AlphaMode"); //Specifically for Transmission Materials
+		inline const FString AlphaMode = TEXT("AlphaMode");
 
 		//ClearCoat specific:
 		DECLARE_INTERCHANGE_GLTF_MI_MAP(ClearCoat)
