MaterialX:
- Add support for gltf_pbr

#Synced-CL 51585235

diff --git a/Engine/Plugins/Interchange/Runtime/Source/Common/Public/MaterialX/InterchangeMaterialXDefinitions.h b/Engine/Plugins/Interchange/Runtime/Source/Common/Public/MaterialX/InterchangeMaterialXDefinitions.h
--- a/Engine/Plugins/Interchange/Runtime/Source/Common/Public/MaterialX/InterchangeMaterialXDefinitions.h
+++ b/Engine/Plugins/Interchange/Runtime/Source/Common/Public/MaterialX/InterchangeMaterialXDefinitions.h
@@ -53,6 +53,9 @@ enum class EInterchangeMaterialXShaders : uint8
 	/** Construct a volume shader describing a participating medium. */
 	Volume,
 
+	/** Default settings for glTF Surface shader. */
+	GltfPBR,
+
 	MaxShaderCount UMETA(hidden)
 };
 
@@ -631,6 +634,42 @@ MATERIALX_NAMESPACE_BEGIN
 		}
 	}
 
+	namespace GltfPbr
+	{
+		namespace Input
+		{
+			inline constexpr const char* BaseColor = "base_color";
+			inline constexpr const char* Metallic = "metallic";
+			inline constexpr const char* Roughness = "roughness";
+			inline constexpr const char* Normal = "normal";
+			inline constexpr const char* Tangent = "tangent";
+			inline constexpr const char* Occlusion = "occlusion";
+			inline constexpr const char* Transmission = "transmission";
+			inline constexpr const char* Specular = "specular";
+			inline constexpr const char* SpecularColor = "specular_color";
+			inline constexpr const char* Ior = "ior";
+			inline constexpr const char* Alpha = "alpha";
+			inline constexpr const char* AlphaMode = "alpha_mode";
+			inline constexpr const char* AlphaCutoff = "alpha_cutoff";
+			inline constexpr const char* Iridescence = "iridescence";
+			inline constexpr const char* IridescenceIor = "iridescence_ior";
+			inline constexpr const char* IridescenceThickness = "iridescence_thickness";
+			inline constexpr const char* SheenColor = "sheen_color";
+			inline constexpr const char* SheenRoughness = "sheen_roughness";
+			inline constexpr const char* Clearcoat = "clearcoat";
+			inline constexpr const char* ClearcoatRoughness = "clearcoat_roughness";
+			inline constexpr const char* ClearcoatNormal = "clearcoat_normal";
+			inline constexpr const char* Emissive = "emissive";
+			inline constexpr const char* EmissiveStrength = "emissive_strength";
+			inline constexpr const char* Thickness = "thickness";
+			inline constexpr const char* AttenuationDistance = "attenuation_distance";
+			inline constexpr const char* AttenuationColor = "attenuation_color";
+			inline constexpr const char* AnisotropyStrength = "anisotropy_strength";
+			inline constexpr const char* AnisotropyRotation = "anisotropy_rotation";
+			inline constexpr const char* Dispersion = "dispersion";
+		}
+	}
+
 	namespace Lights
 	{
 		//There's no input per se in a Light, but we can find some common inputs among those lights
@@ -894,6 +933,7 @@ MATERIALX_NAMESPACE_BEGIN
 		inline constexpr const char* StandardSurface = "ND_standard_surface_surfaceshader";
 		inline constexpr const char* SurfaceUnlit = "ND_surface_unlit";
 		inline constexpr const char* UsdPreviewSurface = "ND_UsdPreviewSurface_surfaceshader";
+		inline constexpr const char* GltfPbr = "ND_gltf_pbr_surfaceshader";
 		inline constexpr const char* PointLight = "ND_point_light";
 		inline constexpr const char* DirectionalLight = "ND_directional_light";
 		inline constexpr const char* SpotLight = "ND_spot_light";
diff --git a/Engine/Plugins/Interchange/Runtime/Source/Import/Private/MaterialX/MaterialXUtils/MaterialXGltfPbr.cpp b/Engine/Plugins/Interchange/Runtime/Source/Import/Private/MaterialX/MaterialXUtils/MaterialXGltfPbr.cpp
new file mode 100644
--- /dev/null
+++ b/Engine/Plugins/Interchange/Runtime/Source/Import/Private/MaterialX/MaterialXUtils/MaterialXGltfPbr.cpp
@@ -0,0 +1,50 @@
+// Copyright Epic Games, Inc. All Rights Reserved.
+
+#include "MaterialXGltfPbr.h"
+
+#if WITH_EDITOR
+
+namespace mx = MaterialX;
+
+FMaterialXGltfPbr::FMaterialXGltfPbr(UInterchangeBaseNodeContainer& BaseNodeContainer)
+	: FMaterialXSurfaceShaderAbstract(BaseNodeContainer)
+{
+	NodeDefinition = mx::NodeDefinition::GltfPbr;
+}
+
+TSharedRef<FMaterialXBase> FMaterialXGltfPbr::MakeInstance(UInterchangeBaseNodeContainer& BaseNodeContainer)
+{
+	TSharedRef<FMaterialXGltfPbr> Result = MakeShared<FMaterialXGltfPbr>(BaseNodeContainer);
+	Result->RegisterConnectNodeOutputToInputDelegates();
+	return Result;
+}
+
+UInterchangeBaseNode* FMaterialXGltfPbr::Translate(MaterialX::NodePtr GltfPbrNode)
+{
+	this->SurfaceShaderNode = GltfPbrNode;
+
+	using namespace UE::Interchange::Materials;
+
+	UInterchangeShaderNode* GltfPbrShaderNode = FMaterialXSurfaceShaderAbstract::Translate(EInterchangeMaterialXShaders::GltfPBR);
+
+	UInterchangeShaderPortsAPI::ConnectOuputToInputByName(ShaderGraphNode, SubstrateMaterial::Parameters::FrontMaterial.ToString(), GltfPbrShaderNode->GetUniqueID(), SubstrateMaterial::Parameters::FrontMaterial.ToString());
+	UInterchangeShaderPortsAPI::ConnectOuputToInputByName(ShaderGraphNode, Common::Parameters::Refraction.ToString(), GltfPbrShaderNode->GetUniqueID(), Common::Parameters::Refraction.ToString());
+	UInterchangeShaderPortsAPI::ConnectOuputToInputByName(ShaderGraphNode, SubstrateMaterial::Parameters::Occlusion.ToString(), GltfPbrShaderNode->GetUniqueID(), SubstrateMaterial::Parameters::Occlusion.ToString());
+
+	return GltfPbrShaderNode;
+}
+
+MaterialX::InputPtr FMaterialXGltfPbr::GetInputNormal(MaterialX::NodePtr GltfPbrNode, const char*& InputNormal) const
+{
+	InputNormal = mx::GltfPbr::Input::Normal;
+	mx::InputPtr Input = GltfPbrNode->getActiveInput(InputNormal);
+
+	if (!Input)
+	{
+		Input = GltfPbrNode->getNodeDef(mx::EMPTY_STRING, true)->getActiveInput(InputNormal);
+	}
+
+	return Input;
+}
+
+#endif // WITH_EDITOR
\ No newline at end of file
diff --git a/Engine/Plugins/Interchange/Runtime/Source/Import/Private/MaterialX/MaterialXUtils/MaterialXGltfPbr.h b/Engine/Plugins/Interchange/Runtime/Source/Import/Private/MaterialX/MaterialXUtils/MaterialXGltfPbr.h
new file mode 100644
--- /dev/null
+++ b/Engine/Plugins/Interchange/Runtime/Source/Import/Private/MaterialX/MaterialXUtils/MaterialXGltfPbr.h
@@ -0,0 +1,25 @@
+// Copyright Epic Games, Inc. All Rights Reserved.
+
+#pragma once
+
+#if WITH_EDITOR
+
+#include "MaterialX/MaterialXUtils/MaterialXSurfaceShaderAbstract.h"
+
+class FMaterialXGltfPbr : public FMaterialXSurfaceShaderAbstract
+{
+	template<typename T, ESPMode>
+	friend class SharedPointerInternals::TIntrusiveReferenceController;
+
+	FMaterialXGltfPbr(UInterchangeBaseNodeContainer& BaseNodeContainer);
+
+public:
+
+	static TSharedRef<FMaterialXBase> MakeInstance(UInterchangeBaseNodeContainer& BaseNodeContainer);
+
+	virtual UInterchangeBaseNode* Translate(MaterialX::NodePtr GltfPbrNode) override;
+
+	virtual MaterialX::InputPtr GetInputNormal(MaterialX::NodePtr GltfPbrNode, const char*& InputNormal) const override;
+};
+
+#endif // WITH_EDITOR
\ No newline at end of file
diff --git a/Engine/Plugins/Interchange/Runtime/Source/Import/Private/MaterialX/MaterialXUtils/MaterialXManager.cpp b/Engine/Plugins/Interchange/Runtime/Source/Import/Private/MaterialX/MaterialXUtils/MaterialXManager.cpp
--- a/Engine/Plugins/Interchange/Runtime/Source/Import/Private/MaterialX/MaterialXUtils/MaterialXManager.cpp
+++ b/Engine/Plugins/Interchange/Runtime/Source/Import/Private/MaterialX/MaterialXUtils/MaterialXManager.cpp
@@ -12,6 +12,7 @@
 #include "MaterialX/MaterialXUtils/MaterialXOpenPBRSurfaceShader.h"
 #include "MaterialX/MaterialXUtils/MaterialXSurfaceShader.h"
 #include "MaterialX/MaterialXUtils/MaterialXStandardSurfaceShader.h"
+#include "MaterialX/MaterialXUtils/MaterialXGltfPbr.h"
 #include "MaterialX/MaterialXUtils/MaterialXUsdPreviewSurfaceShader.h"
 #include "MaterialX/MaterialXUtils/MaterialXSurfaceUnlitShader.h"
 #include "MaterialX/MaterialXUtils/MaterialXPointLightShader.h"
@@ -205,6 +206,7 @@ FMaterialXManager::FMaterialXManager()
 		{mx::Category::Surface, FMaterialXManager::FOnGetMaterialXInstance::CreateStatic(&FMaterialXSurfaceShader::MakeInstance)},
 		{mx::Category::SurfaceUnlit, FMaterialXManager::FOnGetMaterialXInstance::CreateStatic(&FMaterialXSurfaceUnlitShader::MakeInstance)},
 		{mx::Category::StandardSurface, FMaterialXManager::FOnGetMaterialXInstance::CreateStatic(&FMaterialXStandardSurfaceShader::MakeInstance)},
+		{mx::Category::GltfPbr, FMaterialXManager::FOnGetMaterialXInstance::CreateStatic(&FMaterialXGltfPbr::MakeInstance)},
 		{mx::Category::UsdPreviewSurface, FMaterialXManager::FOnGetMaterialXInstance::CreateStatic(&FMaterialXUsdPreviewSurfaceShader::MakeInstance)},
 		{mx::Category::Displacement, FMaterialXManager::FOnGetMaterialXInstance::CreateStatic(&FMaterialXDisplacementShader::MakeInstance)},
 		{mx::Category::Mix, FMaterialXManager::FOnGetMaterialXInstance::CreateStatic(&FMaterialXMixShader::MakeInstance)},
@@ -297,6 +299,7 @@ FMaterialXManager::FMaterialXManager()
 		mx::Category::Extract,
 		mx::Category::Fractal2D,
 		mx::Category::Fractal3D,
+		mx::Category::GltfPbr,
 		mx::Category::GoochShade,
 		mx::Category::HsvAdjust,
 		mx::Category::LatLongImage,
diff --git a/Engine/Plugins/Interchange/Runtime/Source/Import/Private/MaterialX/MaterialXUtils/MaterialXSurfaceShaderAbstract.cpp b/Engine/Plugins/Interchange/Runtime/Source/Import/Private/MaterialX/MaterialXUtils/MaterialXSurfaceShaderAbstract.cpp
--- a/Engine/Plugins/Interchange/Runtime/Source/Import/Private/MaterialX/MaterialXUtils/MaterialXSurfaceShaderAbstract.cpp
+++ b/Engine/Plugins/Interchange/Runtime/Source/Import/Private/MaterialX/MaterialXUtils/MaterialXSurfaceShaderAbstract.cpp
@@ -201,6 +201,12 @@ bool FMaterialXSurfaceShaderAbstract::AddLinearColorAttribute(MaterialX::InputPt
 
 			UInterchangeShaderNode* VectorParameterNode = CreateShaderNode(Input, NodeName, VectorParameter::Name.ToString());
 			VectorParameterNode->AddLinearColorAttribute(UInterchangeShaderPortsAPI::MakeInputParameterKey(VectorParameter::Attributes::DefaultValue.ToString()), Value);
+			
+			if (Input->getType() == mx::Type::Color4)
+			{
+				OutputIndex = 5; // RGBA output on a VectorParameter is the 6th output
+			}
+			
 			return UInterchangeShaderPortsAPI::ConnectOuputToInputByIndex(ShaderNode, InputChannelName, VectorParameterNode->GetUniqueID(), OutputIndex);
 		}
 	}
@@ -224,6 +230,12 @@ bool FMaterialXSurfaceShaderAbstract::AddVectorAttribute(MaterialX::InputPtr Inp
 
 			UInterchangeShaderNode* VectorParameterNode = CreateShaderNode(Input, NodeName, VectorParameter::Name.ToString());
 			VectorParameterNode->AddLinearColorAttribute(UInterchangeShaderPortsAPI::MakeInputParameterKey(VectorParameter::Attributes::DefaultValue.ToString()), Value);
+			
+			if (Input->getType() == mx::Type::Vector4)
+			{
+				OutputIndex = 5; // RGBA output on a VectorParameter is the 6th output
+			}
+			
 			return UInterchangeShaderPortsAPI::ConnectOuputToInputByIndex(ShaderNode, InputChannelName, VectorParameterNode->GetUniqueID(), OutputIndex);
 		}
 	}
@@ -462,6 +474,10 @@ bool FMaterialXSurfaceShaderAbstract::ConnectNodeNameOutputToInput(MaterialX::In
 	UInterchangeShaderNode* ParentShaderNode = ShaderNode;
 	FString InputChannelName = ParentInputName;
 
+	// for a direct node connection without nodegraph, we need the output of the first node connection
+	// for example connecting directly a separate2/3/4 node on a surfaceshader input might miss an entire shader graph
+	FString OutputString = !InputToConnectedNode->getOutputString().empty() ? ANSI_TO_TCHAR(InputToConnectedNode->getOutputString().c_str()) : TEXT("out");
+
 	mx::Edge Edge(nullptr, InputToConnectedNode, ConnectedNode);
 
 	TArray<mx::Edge> Stack{ Edge };
@@ -472,7 +488,7 @@ bool FMaterialXSurfaceShaderAbstract::ConnectNodeNameOutputToInput(MaterialX::In
 
 		if(Edge.getUpstreamElement())
 		{
-			ConnectNodeCategoryOutputToInput(Edge, ShaderNode, ParentInputName);
+			ConnectNodeCategoryOutputToInput(Edge, ShaderNode, ParentInputName, OutputString);
 			ConnectedNode = Edge.getUpstreamElement()->asA<mx::Node>();
 			for(mx::InputPtr Input : ConnectedNode->getInputs())
 			{
diff --git a/Engine/Plugins/Interchange/Runtime/Source/Pipelines/Private/InterchangeMaterialXPipeline.cpp b/Engine/Plugins/Interchange/Runtime/Source/Pipelines/Private/InterchangeMaterialXPipeline.cpp
--- a/Engine/Plugins/Interchange/Runtime/Source/Pipelines/Private/InterchangeMaterialXPipeline.cpp
+++ b/Engine/Plugins/Interchange/Runtime/Source/Pipelines/Private/InterchangeMaterialXPipeline.cpp
@@ -29,13 +29,14 @@ namespace
 	constexpr const TCHAR* OpenPBRSurfaceSubstratePath = TEXT("/Engine/Functions/Substrate/MF_Substrate_OpenPBR_Opaque.MF_Substrate_OpenPBR_Opaque");
 	constexpr const TCHAR* OpenPBRTransmissionSurfaceFunctionsPath = TEXT("/InterchangeAssets/Functions/MX_OpenPBR_Translucent.MX_OpenPBR_Translucent");
 	constexpr const TCHAR* OpenPBRTransmissionSurfaceSubstratePath = TEXT("/Engine/Functions/Substrate/MF_Substrate_OpenPBR_Translucent.MF_Substrate_OpenPBR_Translucent");
-
 	constexpr const TCHAR* StandardSurfaceFunctionsPath = TEXT("/InterchangeAssets/Functions/MX_StandardSurface.MX_StandardSurface");
 	constexpr const TCHAR* StandardSurfaceSubstratePath = TEXT("/Engine/Functions/Substrate/Substrate-StandardSurface-Opaque.Substrate-StandardSurface-Opaque");
 	constexpr const TCHAR* TransmissionSurfaceFunctionsPath = TEXT("/InterchangeAssets/Functions/MX_TransmissionSurface.MX_TransmissionSurface");
 	constexpr const TCHAR* TransmissionSurfaceSubstratePath = TEXT("/Engine/Functions/Substrate/Substrate-StandardSurface-Translucent.Substrate-StandardSurface-Translucent");
 	MATERIALX_FUNCTIONS_SUBSTRATE_PATH(SurfaceUnlit);
 	MATERIALX_FUNCTIONS_SUBSTRATE_PATH(Surface);
+	constexpr const TCHAR* GltfPBRFunctionsPath = TEXT("/InterchangeAssets/gltf/Substrate/MF_GLTFPBR.MF_GLTFPBR");
+	constexpr const TCHAR* GltfPBRSubstratePath = TEXT("/InterchangeAssets/gltf/Substrate/MF_GLTFPBR.MF_GLTFPBR");
 	constexpr const TCHAR* UsdPreviewSurfaceFunctionsPath = TEXT("/InterchangeAssets/Functions/MX_UsdPreviewSurface.MX_UsdPreviewSurface");
 	constexpr const TCHAR* UsdPreviewSurfaceSubstratePath = TEXT("/Engine/Functions/Substrate/MF_Substrate_UsdPreviewSurface.MF_Substrate_UsdPreviewSurface");
 	constexpr const TCHAR* DisplacementFunctionsPath = TEXT("/InterchangeAssets/Functions/MX_Displacement.MX_Displacement");
@@ -386,6 +387,49 @@ UMaterialXPipelineSettings::UMaterialXPipelineSettings()
 					}
 				}
 			},
+			{
+				UMaterialXPipelineSettings::ToEnumKey(EInterchangeMaterialXShaders::GltfPBR),
+				{
+					// GltfPBR Inputs
+					TSet<FName>{
+						mx::GltfPbr::Input::Alpha,
+						mx::GltfPbr::Input::AlphaCutoff,
+						mx::GltfPbr::Input::AlphaMode,
+						mx::GltfPbr::Input::AnisotropyRotation,
+						mx::GltfPbr::Input::AnisotropyStrength,
+						mx::GltfPbr::Input::AttenuationColor,
+						mx::GltfPbr::Input::AttenuationDistance,
+						mx::GltfPbr::Input::BaseColor,
+						mx::GltfPbr::Input::Clearcoat,
+						mx::GltfPbr::Input::ClearcoatNormal,
+						mx::GltfPbr::Input::ClearcoatRoughness,
+						mx::GltfPbr::Input::Dispersion,
+						mx::GltfPbr::Input::Emissive,
+						mx::GltfPbr::Input::EmissiveStrength,
+						mx::GltfPbr::Input::Ior,
+						mx::GltfPbr::Input::Iridescence,
+						mx::GltfPbr::Input::IridescenceIor,
+						mx::GltfPbr::Input::IridescenceThickness,
+						mx::GltfPbr::Input::Metallic,
+						mx::GltfPbr::Input::Normal,
+						mx::GltfPbr::Input::Occlusion,
+						mx::GltfPbr::Input::Roughness,
+						mx::GltfPbr::Input::SheenColor,
+						mx::GltfPbr::Input::SheenRoughness,
+						mx::GltfPbr::Input::Specular,
+						mx::GltfPbr::Input::SpecularColor,
+						mx::GltfPbr::Input::Tangent,
+						mx::GltfPbr::Input::Thickness,
+						mx::GltfPbr::Input::Transmission
+					},
+					// GltfPbr Outputs					
+					TSet<FName>{
+						SubstrateMaterial::Parameters::FrontMaterial,
+						Common::Parameters::Refraction,
+						Common::Parameters::Occlusion
+						}
+				}
+			},
 			{
 				UMaterialXPipelineSettings::ToEnumKey(EInterchangeMaterialXShaders::UsdPreviewSurface),
 				{
@@ -721,6 +765,7 @@ UInterchangeMaterialXPipeline::UInterchangeMaterialXPipeline()
 			{MATERIALX_MATERIALFUNCTION_PATH(SurfaceUnlit),				  UMaterialXPipelineSettings::ToEnumKey(EInterchangeMaterialXShaders::SurfaceUnlit)},
 			{MATERIALX_MATERIALFUNCTION_PATH(Surface),					  UMaterialXPipelineSettings::ToEnumKey(EInterchangeMaterialXShaders::Surface)},
 			{MATERIALX_MATERIALFUNCTION_PATH(Volume),					  UMaterialXPipelineSettings::ToEnumKey(EInterchangeMaterialXShaders::Volume)},
+			{MATERIALX_MATERIALFUNCTION_PATH(GltfPBR),					  UMaterialXPipelineSettings::ToEnumKey(EInterchangeMaterialXShaders::GltfPBR)},
 			{MATERIALX_MATERIALFUNCTION_PATH(UsdPreviewSurface),		  UMaterialXPipelineSettings::ToEnumKey(EInterchangeMaterialXShaders::UsdPreviewSurface)},
 			{MATERIALX_MATERIALFUNCTION_PATH(Displacement),				  UMaterialXPipelineSettings::ToEnumKey(EInterchangeMaterialXShaders::Displacement)},
 
@@ -947,6 +992,7 @@ void UMaterialXPipelineSettings::InitPredefinedAssets()
 			{UMaterialXPipelineSettings::ToEnumKey(EInterchangeMaterialXShaders::StandardSurface), StandardSurfaceFunctionsPath, StandardSurfaceSubstratePath},
 			{UMaterialXPipelineSettings::ToEnumKey(EInterchangeMaterialXShaders::StandardSurfaceTransmission), TransmissionSurfaceFunctionsPath, TransmissionSurfaceSubstratePath},
 			{UMaterialXPipelineSettings::ToEnumKey(EInterchangeMaterialXShaders::SurfaceUnlit), SurfaceUnlitFunctionsPath, SurfaceUnlitSubstratePath},
+			{UMaterialXPipelineSettings::ToEnumKey(EInterchangeMaterialXShaders::GltfPBR), GltfPBRFunctionsPath, GltfPBRSubstratePath},
 			{UMaterialXPipelineSettings::ToEnumKey(EInterchangeMaterialXShaders::UsdPreviewSurface), UsdPreviewSurfaceFunctionsPath, UsdPreviewSurfaceSubstratePath},
 			{UMaterialXPipelineSettings::ToEnumKey(EInterchangeMaterialXShaders::Displacement), DisplacementFunctionsPath, DisplacementSubstratePath},
 
