Fixed flaw in logic to select ParentMaterial
- The value of the boolean bCreateMaterialInstanceForParent waas not taken in account

#rnx
#Synced-CL 52060097

diff --git a/Engine/Plugins/Interchange/Runtime/Source/Pipelines/Private/InterchangeGenericMaterialPipeline.cpp b/Engine/Plugins/Interchange/Runtime/Source/Pipelines/Private/InterchangeGenericMaterialPipeline.cpp
--- a/Engine/Plugins/Interchange/Runtime/Source/Pipelines/Private/InterchangeGenericMaterialPipeline.cpp
+++ b/Engine/Plugins/Interchange/Runtime/Source/Pipelines/Private/InterchangeGenericMaterialPipeline.cpp
@@ -992,8 +992,9 @@ void UInterchangeGenericMaterialPipeline::ExecutePipeline(UInterchangeBaseNodeCo
 	for (UInterchangeMaterialInstanceNode* MaterialNode : MaterialInstanceNodes)
 	{
 		FString ParentPathOrMaterialUid;
-		// Override any existing parent material if one has been specified
-		if (UMaterialInterface* ParentMaterialObj = Cast<UMaterialInterface>(ParentMaterial.TryLoad()))
+		// Override any existing parent material if one has been specified and bCreateMaterialInstanceForParent is true
+		UMaterialInterface* ParentMaterialObj = bCreateMaterialInstanceForParent ? Cast<UMaterialInterface>(ParentMaterial.TryLoad()) : nullptr;
+		if (ParentMaterialObj)
 		{
 			ParentPathOrMaterialUid = ParentMaterialObj->GetPathName();
 			MaterialNode->SetCustomParent(ParentPathOrMaterialUid);
