Fixed flaw in logic to add face to group

#rnx
#Synced-CL 52060097

diff --git a/Engine/Plugins/Interchange/Runtime/Source/Import/Private/Mesh/InterchangeOBJTranslator.cpp b/Engine/Plugins/Interchange/Runtime/Source/Import/Private/Mesh/InterchangeOBJTranslator.cpp
--- a/Engine/Plugins/Interchange/Runtime/Source/Import/Private/Mesh/InterchangeOBJTranslator.cpp
+++ b/Engine/Plugins/Interchange/Runtime/Source/Import/Private/Mesh/InterchangeOBJTranslator.cpp
@@ -694,9 +694,9 @@ namespace ObjParserUtils
 		{
 			FStringView FaceToken = GetToken(Line);
 
-			auto GetFaceIndex = [](FStringView& Token, int32 NumElements, int32& Index)
+			auto GetFaceIndex = [](FStringView& Token, int32 NumElements)
 			{
-				Index = INDEX_NONE;
+				int32 Index = INDEX_NONE;
 				if (GetInteger(Token, Index, TEXT('/')))
 				{
 					if (Index < 0)
@@ -709,22 +709,21 @@ namespace ObjParserUtils
 						// Indices start at one, so adjust for zero-base
 						--Index;
 					}
+
+					if (Index >= 0 && Index < NumElements)
+					{
+						return Index;
+					}
 				}
 
-				return Index >= 0 && Index < NumElements;
+				return static_cast<int32>(INDEX_NONE);
 			};
 
-			int32 VertexIndex = INDEX_NONE;
-			int32 UVIndex = INDEX_NONE;
-			int32 NormalIndex = INDEX_NONE;
-			const bool bCanAddVertex = GetFaceIndex(FaceToken, ObjData.Positions.Num(), VertexIndex)
-				&& GetFaceIndex(FaceToken, ObjData.UVs.Num(), UVIndex)
-				&& GetFaceIndex(FaceToken, ObjData.Normals.Num(), NormalIndex);
+			int32 VertexIndex = GetFaceIndex(FaceToken, ObjData.Positions.Num());
+			int32 UVIndex = GetFaceIndex(FaceToken, ObjData.UVs.Num());
+			int32 NormalIndex = GetFaceIndex(FaceToken, ObjData.Normals.Num());
 
-			if (bCanAddVertex)
-			{
-				FaceData.Vertices.Emplace(VertexIndex, UVIndex, NormalIndex);
-			}
+			FaceData.Vertices.Emplace(VertexIndex, UVIndex, NormalIndex);
 		}
 
 		if (FaceData.Vertices.Num() > 2)
