Renderer: Fix compile errors

#rnx
#Synced-CL 52148822

diff --git a/Engine/Plugins/Experimental/InterchangeAxF/Runtime/Source/InterchangeAxF/Private/AxFInterchangePipeline.cpp b/Engine/Plugins/Experimental/InterchangeAxF/Runtime/Source/InterchangeAxF/Private/AxFInterchangePipeline.cpp
--- a/Engine/Plugins/Experimental/InterchangeAxF/Runtime/Source/InterchangeAxF/Private/AxFInterchangePipeline.cpp
+++ b/Engine/Plugins/Experimental/InterchangeAxF/Runtime/Source/InterchangeAxF/Private/AxFInterchangePipeline.cpp
@@ -450,7 +450,7 @@ void UAxFInterchangePipeline::AdjustSettingsForContext(FInterchangePipelineConte
 	Super::AdjustSettingsForContext(ContextParams);
 }
 
-inline void UAxFInterchangePipeline::PreDialogCleanup(FName const PipelineStackName)
+void UAxFInterchangePipeline::PreDialogCleanup(FName const PipelineStackName)
 {
 	if (TexturePipeline)
 	{
@@ -468,7 +468,7 @@ bool UAxFInterchangePipeline::IsSettingsAreValid(TOptional<FText>& OutInvalidRea
 	return Super::IsSettingsAreValid(OutInvalidReason);
 }
 
-inline void UAxFInterchangePipeline::FilterPropertiesFromTranslatedData(
+void UAxFInterchangePipeline::FilterPropertiesFromTranslatedData(
 	UInterchangeBaseNodeContainer* InBaseNodeContainer)
 {
 	if (TexturePipeline)
diff --git a/Engine/Source/Runtime/Renderer/Private/CustomDepthRendering.cpp b/Engine/Source/Runtime/Renderer/Private/CustomDepthRendering.cpp
--- a/Engine/Source/Runtime/Renderer/Private/CustomDepthRendering.cpp
+++ b/Engine/Source/Runtime/Renderer/Private/CustomDepthRendering.cpp
@@ -22,7 +22,7 @@
 #include "MeshBatch.h"
 #include "MeshDrawCommands.h"
 #include "MeshMaterialShader.h"
-#include "MeshPassProcessor.h"
+#include "MeshPassProcessor.inl"
 #include "Misc/LargeWorldCoordinates.h"
 #include "Nanite/NaniteComposition.h"
 #include "Nanite/NaniteCullRaster.h"
