Interchange: Fix crash when provided morph target curve payload is empty.

#Synced-CL 52058314

diff --git a/Engine/Plugins/Interchange/Runtime/Source/Import/Private/Animation/InterchangeAnimSequenceFactory.cpp b/Engine/Plugins/Interchange/Runtime/Source/Import/Private/Animation/InterchangeAnimSequenceFactory.cpp
--- a/Engine/Plugins/Interchange/Runtime/Source/Import/Private/Animation/InterchangeAnimSequenceFactory.cpp
+++ b/Engine/Plugins/Interchange/Runtime/Source/Import/Private/Animation/InterchangeAnimSequenceFactory.cpp
@@ -107,6 +107,11 @@ namespace UE::Interchange::Private
 		for (int32 CurveIndex = 0; CurveIndex < Curves.Num(); ++CurveIndex)
 		{
 			FRichCurve& Curve = Curves[CurveIndex];
+			if (Curve.IsEmpty())
+			{
+				continue;
+			}
+
 			FName Name(UE::Interchange::MakeSanitizedName(CurveNames[CurveIndex]));
 
 			if (bDoNotImportCurveWithZero)
