[Interchange] ufbx - fix skinned animation when bones not marked 'skeleton'

- the testcase in the mentioned old FBX jira - just import skinned with default settings and ufbx enabled
- Currently test file can be found on Lucid - UE_Editor_Content_Pipeline\FBX\Performance\Samurai_boy_multianims.fbx

#Synced-CL 51759865

diff --git a/Engine/Plugins/Interchange/Runtime/Source/Parsers/Fbx/Private/Ufbx/UfbxAnimation.cpp b/Engine/Plugins/Interchange/Runtime/Source/Parsers/Fbx/Private/Ufbx/UfbxAnimation.cpp
--- a/Engine/Plugins/Interchange/Runtime/Source/Parsers/Fbx/Private/Ufbx/UfbxAnimation.cpp
+++ b/Engine/Plugins/Interchange/Runtime/Source/Parsers/Fbx/Private/Ufbx/UfbxAnimation.cpp
@@ -402,7 +402,7 @@ namespace UE::Interchange::Private
 				}
 
 				ufbx_node* Node = FbxNodes[BakedNode.typed_id];
-				if (Node->bone)
+				if (Node && (Node->bone || Scene.IsInSkeleton(*Node)))
 				{
 					FString NodeUid = Parser.GetNodeUid(*Node);
 
