[Interchange] ufbx - mesh smoothing more solid

- basing mesh smoothing on precomputed normals in the fbx file are not always correct
- when edge or face smoothing found used those first
- face smoothing only used when all faces are with hard edges - due to bug in ufbx library which incorrectly interprets PyPolygon smoothing data in fbx
- also, for face smoothing, skinned has taken into consideration as done with FBX parser

#Synced-CL 51759865

diff --git a/Engine/Plugins/Interchange/Runtime/Source/Parsers/Fbx/Private/Ufbx/UfbxMesh.cpp b/Engine/Plugins/Interchange/Runtime/Source/Parsers/Fbx/Private/Ufbx/UfbxMesh.cpp
--- a/Engine/Plugins/Interchange/Runtime/Source/Parsers/Fbx/Private/Ufbx/UfbxMesh.cpp
+++ b/Engine/Plugins/Interchange/Runtime/Source/Parsers/Fbx/Private/Ufbx/UfbxMesh.cpp
@@ -13,6 +13,7 @@
 #include "MeshDescriptionBuilder.h"
 #endif 
 
+#include "Algo/NoneOf.h"
 #include "SkeletalMeshAttributes.h"
 #include "StaticMeshAttributes.h"
 #include "StaticMeshOperations.h"
@@ -230,7 +231,7 @@ namespace UE::Interchange::Private
 			bool bMeshHasVertexNormal = Mesh->vertex_normal.exists;
 			bool bMeshHasVertexBinormal = Mesh->vertex_bitangent.exists && Mesh->vertex_tangent.exists;
 			bool bMeshHasVertexTangent = bMeshHasVertexBinormal;
-			bool bMeshHasSmoothGroup = true;
+			bool bMeshHasSmoothGroup = Mesh->face_smoothing.count > 0 || Mesh->edge_smoothing.count > 0;
 			bool bMeshHasVertexColor = Mesh->vertex_color.exists;
 			int32 MeshUVCount = Mesh->uv_sets.count;
 			
@@ -432,6 +433,7 @@ namespace UE::Interchange::Private
 		Attributes.Register(bKeepExistingAttribute);
 
 		// #ufbx_todo Revisit how FMeshDescription is filled up
+		bool bAppliedSmoothingFromData = false;
 #ifdef WITH_MESH_DESCRIPTION_BUILDER
 		FMeshDescriptionBuilder Builder;
 		Builder.SetMeshDescription(&MeshDescription);
@@ -499,9 +501,61 @@ namespace UE::Interchange::Private
 		HandleMesh(FaceCollector);
 
 		Builder.ResumeMeshDescriptionIndexing();
+
+		// Apply edge hardness from ufbx smoothing data if available.
+		// ByEdge: direct 1:1 mapping from ufbx edge flags to MeshDescription edges via vertex-pair lookup.
+		// ByPolygon (face_smoothing bool): edge is hard when ALL faces have face_smoothing==false. This handles the common all-zero case correctly. 
+		// But:
+		// #ufbx_todo: face_smoothing is bool in ufbx and this is wrong since in source data in fbx it is group BIT MASK(32 bit for 32 groups, why 3ds Max has 32 group limit btw) 
+		// and ufbx converts it to bool. So different smoothing groups both appear true and get soft edges here 
+		// Correct fix requires changing ufbx.c and storing the raw int32 bitmask. 
+		if (Mesh.edge_smoothing.count > 0)
+		{
+			bAppliedSmoothingFromData = true;
+			TEdgeAttributesRef<bool> EdgeHardnesses = Attributes.GetEdgeHardnesses();
+
+			for (int32 EdgeIdx = 0; EdgeIdx < (int32)Mesh.edge_smoothing.count; ++EdgeIdx)
+			{
+				const ufbx_edge& Edge = Mesh.edges[EdgeIdx];
+				const FVertexID VertexA = FaceCollector.VertexMapping[Mesh.vertex_position.indices[Edge.a]];
+				const FVertexID VertexB = FaceCollector.VertexMapping[Mesh.vertex_position.indices[Edge.b]];
+				FEdgeID EdgeID = MeshDescription.GetVertexPairEdge(VertexA, VertexB);
+				if (EdgeID == INDEX_NONE)
+				{
+					EdgeID = MeshDescription.CreateEdge(VertexA, VertexB);
+				}
+				// ufbx softness to UE hardness
+				EdgeHardnesses[EdgeID] = !Mesh.edge_smoothing.data[EdgeIdx];
+			}
+		}
+		else if (Mesh.face_smoothing.count > 0)
+		{
+			// ufbx stores ByPolygon smoothing groups as bool (non-zero becomes true), losing the
+			// 32-bitmask value. Only safe case  when ALL groups are zero: that means every face is hard. 
+			// If any group is non-zero we cannot identify smoothing groups without the original bitmask, 
+			// so fallthrough to DetermineEdgeHardnessesFromVertexInstanceNormals to avoid regression.
+			bool bAllZero = Algo::NoneOf(Mesh.face_smoothing);
+
+			if (bAllZero && Mesh.skin_deformers.count == 0)
+			{
+				// All smoothing groups are 0 and mesh is NOT skinned assume every edge is hard (fully faceted).
+				// Skinned meshes with all-zero groups have their smoothing expressed through the baked
+				// normals (matches FBX SDK behaviour: ComputeEdgeSmoothingFromNormals for skinned).
+				bAppliedSmoothingFromData = true;
+				TEdgeAttributesRef<bool> EdgeHardnesses = Attributes.GetEdgeHardnesses();
+				for (const FEdgeID EdgeID : MeshDescription.Edges().GetElementIDs())
+				{
+					EdgeHardnesses[EdgeID] = true;
+				}
+			}
+		}
 #endif
 
-		FStaticMeshOperations::DetermineEdgeHardnessesFromVertexInstanceNormals(MeshDescription);
+		if (!bAppliedSmoothingFromData)
+		{
+			FStaticMeshOperations::DetermineEdgeHardnessesFromVertexInstanceNormals(MeshDescription);
+		}
+
 		FStaticMeshOperations::ApplyTransform(MeshDescription, MeshGlobalTransform, true);
 		return true;
 	}
