[Interchange] ufbx  - attempt of rigid anim curve reduction similar to FBX parser

#Synced-CL 51759865

diff --git a/Engine/Plugins/Interchange/Runtime/Source/Parsers/Fbx/Private/Ufbx/UfbxAnimation.cpp b/Engine/Plugins/Interchange/Runtime/Source/Parsers/Fbx/Private/Ufbx/UfbxAnimation.cpp
--- a/Engine/Plugins/Interchange/Runtime/Source/Parsers/Fbx/Private/Ufbx/UfbxAnimation.cpp
+++ b/Engine/Plugins/Interchange/Runtime/Source/Parsers/Fbx/Private/Ufbx/UfbxAnimation.cpp
@@ -200,6 +200,76 @@ namespace UE::Interchange::Private
 		return false;
 	}
 
+	void ReduceConstantKeys(const ufbx_anim_curve& Curve, TArray<FInterchangeCurveKey>& Keys, double Tolerance)
+	{
+		if (Keys.IsEmpty())
+		{
+			return;
+		}
+
+		// might use 1 instead of 0 if want to leave some keys even in all-constant curve?
+		constexpr int32 FirstKeyIndex = 0;
+		for (int32 KeyIndex = Keys.Num() - 1; KeyIndex >= FirstKeyIndex;)
+		{
+			const double CurrentValue = Keys[KeyIndex].Value;
+			FDoubleInterval ToleranceInterval(CurrentValue - Tolerance, CurrentValue + Tolerance);
+
+			// find all consecutive constant keys
+			int32 ConstantLeftIndex = KeyIndex;
+			int32 ConstantRightIndex = KeyIndex;
+			while (KeyIndex >= FirstKeyIndex)
+			{
+				const FInterchangeCurveKey& Key = Keys[KeyIndex];
+
+				const double KeyValue = Key.Value;
+				const bool bValueInBand = ToleranceInterval.Contains(KeyValue);
+
+				const bool bIsEffectivelyConstant = bValueInBand && (
+					Key.InterpMode == EInterchangeCurveInterpMode::Constant
+					// Linear within epsilon range is constant
+					|| Key.InterpMode == EInterchangeCurveInterpMode::Linear
+					// Cubic but both tangents are about zero
+					|| (Key.InterpMode == EInterchangeCurveInterpMode::Cubic
+						&& FMath::IsNearlyZero(Key.LeaveTangent, Tolerance)
+						&& FMath::IsNearlyZero(Key.ArriveTangent, Tolerance))
+					);
+
+				// also, might check constantness by evaluating curve between keys at desired fbx to see that values are somewhat constant
+				// instead of just testing tangent for constantness
+				// ufbx_evaluate_curve(&Curve, Key.Time, Key.Value);
+
+				if (bIsEffectivelyConstant)
+				{
+					// Constant: mark for deletion, step left
+					ConstantLeftIndex = KeyIndex;
+					--KeyIndex;
+				}
+				else
+				{
+					--KeyIndex;
+					break;
+				}
+			}
+			
+			if (ConstantLeftIndex != ConstantRightIndex)
+			{
+				int32 ConstantKeysNum = ConstantRightIndex - ConstantLeftIndex + 1;
+				if (ConstantKeysNum < Keys.Num())
+				{
+					// Remove constant interval except one(left-most) key
+					Keys.RemoveAt(ConstantLeftIndex + 1, ConstantKeysNum - 1);
+				}
+				else
+				{
+					// All keys are constant
+					// maybe, should check if curve is distinct from the default anim value to keep something here?
+					Keys.Empty();
+					return;
+				}
+			}
+		}
+	}
+
 	bool FUfbxAnimation::FetchRigidAnimation(const FUfbxParser& Parser, const FRigidAnimationPayloadContext& Animation,
 		TArray<FInterchangeCurve>& OutCurves)
 	{
@@ -208,7 +278,12 @@ namespace UE::Interchange::Private
 			const ufbx_anim_curve* Curve =  Prop ? Prop->anim_value->curves[CurveIndex] : nullptr;
 			if (Curve)
 			{
-				AnimUtils::ConvertCurve(*Curve, Scale, OutCurves.AddDefaulted_GetRef().Keys);
+				FInterchangeCurve& OutCurve = OutCurves.AddDefaulted_GetRef();
+				TArray<FInterchangeCurveKey>& Keys = OutCurve.Keys;
+				AnimUtils::ConvertCurve(*Curve, Scale, Keys);
+
+				// ConvertPivotAnimationRecursive in FBX Parser does key reduction, trying repeat it 
+				ReduceConstantKeys(*Curve, Keys, 0.009);
 			}
 			else
 			{
