Gltf:
- Use Unlit parent material

#Synced-CL 52245732

diff --git a/Engine/Plugins/Interchange/Runtime/Source/Pipelines/Private/InterchangeglTFPipeline.cpp b/Engine/Plugins/Interchange/Runtime/Source/Pipelines/Private/InterchangeglTFPipeline.cpp
--- a/Engine/Plugins/Interchange/Runtime/Source/Pipelines/Private/InterchangeglTFPipeline.cpp
+++ b/Engine/Plugins/Interchange/Runtime/Source/Pipelines/Private/InterchangeglTFPipeline.cpp
@@ -318,7 +318,8 @@ UInterchangeMaterialInstanceFactoryNode* UGLTFPipelineSettings::BuildMaterialIns
 	bool bSpecGloss = ParentIdentifier.Contains(TEXT("SpecularGlossiness"), ESearchCase::CaseSensitive) ? true : false;
 
 	FString Parent;
-	if(bSubstrate)
+	bool bUseSubstrateMaterial = bSubstrate && !ParentIdentifier.Contains(TEXT("Unlit"));
+	if (bUseSubstrateMaterial)
 	{
 		ParentIdentifier = TEXT("M_GLTF");
 	}
@@ -363,7 +364,7 @@ UInterchangeMaterialInstanceFactoryNode* UGLTFPipelineSettings::BuildMaterialIns
 		const UClass* MaterialClass = (EditorUtilities && EditorUtilities->IsRuntimeOrPIE()) ? UMaterialInstanceDynamic::StaticClass() : UMaterialInstanceConstant::StaticClass();
 		MaterialInstanceFactoryNode->SetCustomInstanceClassName(MaterialClass->GetPathName());
 		
-		if (bSubstrate)
+		if (bUseSubstrateMaterial)
 		{
 			MaterialInstanceFactoryNode->SetCustomBlendMode(BlendMode);
 			MaterialInstanceFactoryNode->SetCustomTwoSided(bTwoSided);
