Fixed issue with missing textures from FBX file

#rnx
#Synced-CL 52255923

diff --git a/Engine/Plugins/Interchange/Runtime/Source/Parsers/Fbx/Private/FbxMaterial.cpp b/Engine/Plugins/Interchange/Runtime/Source/Parsers/Fbx/Private/FbxMaterial.cpp
--- a/Engine/Plugins/Interchange/Runtime/Source/Parsers/Fbx/Private/FbxMaterial.cpp
+++ b/Engine/Plugins/Interchange/Runtime/Source/Parsers/Fbx/Private/FbxMaterial.cpp
@@ -25,6 +25,44 @@ namespace UE
 	{
 		namespace Private
 		{
+			FString FindTextureFilePath(const FFbxParser& Parser, FbxFileTexture& FbxTexture)
+			{
+				// try opening from absolute path
+				FString AbsoluteTextureFilepath = UTF8_TO_TCHAR(FbxTexture.GetFileName());
+				FPaths::NormalizeFilename(AbsoluteTextureFilepath);
+				if (FPaths::FileExists(AbsoluteTextureFilepath))
+				{
+					return AbsoluteTextureFilepath;
+				}
+
+				FString FileBasePath = FPaths::GetPath(Parser.GetSourceFilename());
+
+				FString RelativeToFileBasePath = FileBasePath / FPaths::GetCleanFilename(AbsoluteTextureFilepath);
+				FPaths::NormalizeFilename(RelativeToFileBasePath);
+				if (FPaths::FileExists(RelativeToFileBasePath))
+				{
+					return RelativeToFileBasePath;
+				}
+
+				// try fbx file base path + relative path
+				FString RelativeToFBXTextureFilepath = FileBasePath / UTF8_TO_TCHAR(FbxTexture.GetRelativeFileName());
+				FPaths::NormalizeFilename(RelativeToFBXTextureFilepath);
+				if (FPaths::FileExists(RelativeToFBXTextureFilepath))
+				{
+					return RelativeToFBXTextureFilepath;
+				}
+
+				// Some fbx files do not store the actual absolute filename as absolute and it is actually relative.  Try to get it relative to the FBX file we are importing
+				FString AbosluteAsRelativeToFBXTextureFilepath = FileBasePath / UTF8_TO_TCHAR(FbxTexture.GetFileName());
+				FPaths::NormalizeFilename(AbosluteAsRelativeToFBXTextureFilepath);
+				if (FPaths::FileExists(AbosluteAsRelativeToFBXTextureFilepath))
+				{
+					return AbosluteAsRelativeToFBXTextureFilepath;
+				}
+
+				return FString();
+			}
+
 			UInterchangeMaterialInstanceNode* FFbxMaterial::CreateMaterialInstanceNode(UInterchangeBaseNodeContainer& NodeContainer, const FString& NodeUID, const FString& NodeName)
 			{
 				UInterchangeMaterialInstanceNode* MaterialInstanceNode = NewObject<UInterchangeMaterialInstanceNode>(&NodeContainer);
@@ -71,35 +109,8 @@ namespace UE
 					return nullptr;
 				}
 
-				const FString TextureFilename = [this, FbxTexture]
-				{
-					// try opening from absolute path
-					FString AbsoluteTextureFilepath = UTF8_TO_TCHAR(FbxTexture->GetFileName());
-					FPaths::NormalizeFilename(AbsoluteTextureFilepath);
-					if (FPaths::FileExists(AbsoluteTextureFilepath))
-					{
-						return AbsoluteTextureFilepath;
-					}
-
-					FString FileBasePath = FPaths::GetPath(Parser.GetSourceFilename());
-					// try fbx file base path + relative path
-					FString RelativeToFBXTextureFilepath = FileBasePath / UTF8_TO_TCHAR(FbxTexture->GetRelativeFileName());
-					FPaths::NormalizeFilename(RelativeToFBXTextureFilepath);
-					if (FPaths::FileExists(RelativeToFBXTextureFilepath))
-					{
-						return RelativeToFBXTextureFilepath;
-					}
-
-					// Some fbx files do not store the actual absolute filename as absolute and it is actually relative.  Try to get it relative to the FBX file we are importing
-					FString AbosluteAsRelativeToFBXTextureFilepath = FileBasePath / UTF8_TO_TCHAR(FbxTexture->GetFileName());
-					FPaths::NormalizeFilename(AbosluteAsRelativeToFBXTextureFilepath);
-					if (FPaths::FileExists(AbosluteAsRelativeToFBXTextureFilepath))
-					{
-						return AbosluteAsRelativeToFBXTextureFilepath;
-					}
 
-					return FString();
-				}();
+				const FString TextureFilename = FindTextureFilePath(Parser, *FbxTexture);
 
 				// Return incomplete texture sampler if texture file does not exist
 				if (TextureFilename.IsEmpty() || !FPaths::FileExists(TextureFilename))
@@ -427,25 +438,29 @@ namespace UE
 				int32 TextureCount = SDKScene->GetSrcObjectCount<FbxFileTexture>();
 				for (int32 TextureIndex = 0; TextureIndex < TextureCount; ++TextureIndex)
 				{
-					FbxFileTexture* Texture = SDKScene->GetSrcObject<FbxFileTexture>(TextureIndex);
-					FString TextureFilename = UTF8_TO_TCHAR(Texture->GetFileName());
-					//Only import texture that exist on disk
-					if (!FPaths::FileExists(TextureFilename))
+					if (FbxFileTexture* Texture = SDKScene->GetSrcObject<FbxFileTexture>(TextureIndex))
 					{
-						if (!GIsAutomationTesting)
+						const FString TextureFilename = FindTextureFilePath(Parser, *Texture);
+						
+						//Only import texture that exist on disk
+						if (TextureFilename.IsEmpty() || !FPaths::FileExists(TextureFilename))
 						{
-							UInterchangeResultTextureDisplay_TextureFileDoNotExist* Message = Parser.AddMessage<UInterchangeResultTextureDisplay_TextureFileDoNotExist>();
-							Message->TextureName = TextureFilename;
-							Message->MaterialName.Empty();
+							if (!GIsAutomationTesting)
+							{
+								UInterchangeResultTextureDisplay_TextureFileDoNotExist* Message = Parser.AddMessage<UInterchangeResultTextureDisplay_TextureFileDoNotExist>();
+								Message->TextureName = TextureFilename;
+								Message->MaterialName.Empty();
+							}
+							continue;
+						}
+
+						//Create a texture node and make it child of the material node
+						const FString TextureName = FPaths::GetBaseFilename(TextureFilename);
+						const UInterchangeTexture2DNode* TextureNode = Cast<UInterchangeTexture2DNode>(NodeContainer.GetNode(UInterchangeTextureNode::MakeNodeUid(TextureName)));
+						if (!TextureNode)
+						{
+							CreateTexture2DNode(NodeContainer, TextureFilename);
 						}
-						continue;
-					}
-					//Create a texture node and make it child of the material node
-					const FString TextureName = FPaths::GetBaseFilename(TextureFilename);
-					const UInterchangeTexture2DNode* TextureNode = Cast<UInterchangeTexture2DNode>(NodeContainer.GetNode(UInterchangeTextureNode::MakeNodeUid(TextureName)));
-					if (!TextureNode)
-					{
-						CreateTexture2DNode(NodeContainer, TextureFilename);
 					}
 				}
 			}
diff --git a/Engine/Plugins/Interchange/Runtime/Source/Parsers/Fbx/Private/Ufbx/UfbxMaterial.cpp b/Engine/Plugins/Interchange/Runtime/Source/Parsers/Fbx/Private/Ufbx/UfbxMaterial.cpp
--- a/Engine/Plugins/Interchange/Runtime/Source/Parsers/Fbx/Private/Ufbx/UfbxMaterial.cpp
+++ b/Engine/Plugins/Interchange/Runtime/Source/Parsers/Fbx/Private/Ufbx/UfbxMaterial.cpp
@@ -70,8 +70,16 @@ namespace UE::Interchange::Private
 			{
 				return TextureFilepath;
 			}
+
 			FString FileBasePath = FPaths::GetPath(Parser.SourceFilename);
 
+			FString RelativeToFileBasePath = FileBasePath / FPaths::GetCleanFilename(TextureFilepath);
+			FPaths::NormalizeFilename(RelativeToFileBasePath);
+			if (FPaths::FileExists(RelativeToFileBasePath))
+			{
+				return RelativeToFileBasePath;
+			}
+
 			FString RelativeToFBXTextureFilepath = FileBasePath / Convert::ToUnrealString(Texture.relative_filename);
 			FPaths::NormalizeFilename(RelativeToFBXTextureFilepath);
 			if( FPaths::FileExists(RelativeToFBXTextureFilepath) )
