[Interchange] Displaying icons for Texture and Material nodes inside the Graph Inspector

#rnx

#Synced-CL 52310122

diff --git a/Engine/Plugins/Interchange/Runtime/Source/Nodes/Private/InterchangeMaterialInstanceNode.cpp b/Engine/Plugins/Interchange/Runtime/Source/Nodes/Private/InterchangeMaterialInstanceNode.cpp
--- a/Engine/Plugins/Interchange/Runtime/Source/Nodes/Private/InterchangeMaterialInstanceNode.cpp
+++ b/Engine/Plugins/Interchange/Runtime/Source/Nodes/Private/InterchangeMaterialInstanceNode.cpp
@@ -11,6 +11,11 @@ FString UInterchangeMaterialInstanceNode::GetTypeName() const
     return TEXT("MaterialInstanceNode");
 }
 
+FName UInterchangeMaterialInstanceNode::GetIconName() const
+{
+	return FName(TEXT("ClassIcon.MaterialInstanceConstant"));
+}
+
 FString UInterchangeMaterialInstanceNode::MakeNodeUid(const FStringView NodeName, const FStringView ParentNodeUid)
 {
 	FString MaterialInstanceNodeUid = FString(UInterchangeBaseNode::HierarchySeparator) + TEXT("MaterialInstances") + FString(UInterchangeBaseNode::HierarchySeparator);
diff --git a/Engine/Plugins/Interchange/Runtime/Source/Nodes/Public/InterchangeMaterialInstanceNode.h b/Engine/Plugins/Interchange/Runtime/Source/Nodes/Public/InterchangeMaterialInstanceNode.h
--- a/Engine/Plugins/Interchange/Runtime/Source/Nodes/Public/InterchangeMaterialInstanceNode.h
+++ b/Engine/Plugins/Interchange/Runtime/Source/Nodes/Public/InterchangeMaterialInstanceNode.h
@@ -23,6 +23,7 @@ public:
      */
 	UE_API virtual FString GetTypeName() const override;
 
+	UE_API virtual FName GetIconName() const override;
 
 	/**
 	 * Build and return a UID name for a material instance node.
diff --git a/Engine/Plugins/Interchange/Runtime/Source/Nodes/Public/InterchangeTextureNode.h b/Engine/Plugins/Interchange/Runtime/Source/Nodes/Public/InterchangeTextureNode.h
--- a/Engine/Plugins/Interchange/Runtime/Source/Nodes/Public/InterchangeTextureNode.h
+++ b/Engine/Plugins/Interchange/Runtime/Source/Nodes/Public/InterchangeTextureNode.h
@@ -89,6 +89,12 @@ public:
 		const FString TypeName = TEXT("TextureNode");
 		return TypeName;
 	}
+
+	virtual FName GetIconName() const override
+	{
+		return FName(TEXT("ClassIcon.Texture2D"));
+	}
+
 #if WITH_EDITOR
 	virtual FString GetKeyDisplayName(const UE::Interchange::FAttributeKey& NodeAttributeKey) const override
 	{
