[Interchange] [Fbx] Prevent ensure when importing FBX that contains clashes in Morph Target names

Partially related (but not a fix) for UE-371274

#Synced-CL 52310122

diff --git a/Engine/Plugins/Interchange/Runtime/Source/Parsers/Fbx/Private/FbxMesh.cpp b/Engine/Plugins/Interchange/Runtime/Source/Parsers/Fbx/Private/FbxMesh.cpp
--- a/Engine/Plugins/Interchange/Runtime/Source/Parsers/Fbx/Private/FbxMesh.cpp
+++ b/Engine/Plugins/Interchange/Runtime/Source/Parsers/Fbx/Private/FbxMesh.cpp
@@ -9,6 +9,7 @@
 #include "FbxHelper.h"
 #include "FbxInclude.h"
 #include "Fbx/InterchangeFbxMessages.h"
+#include "InterchangeFbxParserLog.h"
 #include "InterchangeMeshNode.h"
 #include "InterchangeResultsContainer.h"
 #include "MeshDescription.h"
@@ -1868,7 +1869,26 @@ void FFbxMesh::AddAllMeshes(FbxScene* SDKScene, FbxGeometryConverter* SDKGeometr
 						ShapeName = TEXT("MorphTarget_") + FString::FromInt(ChannelIndex) + TEXT("_") + FString::FromInt(ChannelMorphTargetIndex);
 					}
 
+					static const FString NameClash = "_ncl_";
+
+					if (ShapeNameToFbxShape.Contains(ShapeName))
+					{
+						int32 UniqueId = 1;
+						FString UniqueShapeName;
+						do
+						{
+							UniqueShapeName = ShapeName + NameClash + FString::FromInt(UniqueId++);
+						}
+						while (ShapeNameToFbxShape.Contains(UniqueShapeName));
+
+						UE_LOG(LogInterchangeFbxParser, Warning,
+							TEXT("Duplicate morph target channel name '%s' found in FBX file. It will be imported as '%s'."),
+							*ShapeName, *UniqueShapeName);
+
+						ShapeName = UniqueShapeName;
+					}
 					ensure(!ShapeNameToFbxShape.Contains(ShapeName));
+
 					ShapeNameToFbxShape.Add(ShapeName, Shape);
 					FString MorphTargetAttributeName = Parser.GetFbxHelper()->GetMeshName(Shape);
 					FString MorphTargetUniqueID = Parser.GetFbxHelper()->GetMeshUniqueID(Shape);
@@ -1876,7 +1896,6 @@ void FFbxMesh::AddAllMeshes(FbxScene* SDKScene, FbxGeometryConverter* SDKGeometr
 					if (ExistingMorphTargetNode)
 					{
 						int32 UniqueId = 1;
-						FString NameClash = "_ncl_";
 						while (ExistingMorphTargetNode)
 						{
 							MorphTargetUniqueID = Parser.GetFbxHelper()->GetMeshUniqueID(Shape) + NameClash + FString::FromInt(UniqueId++);
