AxF:
- Improvements to Interchange AxF plugin
[On behalf of Stefan Ivanovski from XRite]

#14572

#Synced-CL 52534341

diff --git a/Engine/Plugins/Experimental/InterchangeAxF/Runtime/Source/InterchangeAxF/Private/AxFInterchangePipeline.cpp b/Engine/Plugins/Experimental/InterchangeAxF/Runtime/Source/InterchangeAxF/Private/AxFInterchangePipeline.cpp
--- a/Engine/Plugins/Experimental/InterchangeAxF/Runtime/Source/InterchangeAxF/Private/AxFInterchangePipeline.cpp
+++ b/Engine/Plugins/Experimental/InterchangeAxF/Runtime/Source/InterchangeAxF/Private/AxFInterchangePipeline.cpp
@@ -45,16 +45,16 @@ void UAxFInterchangePipeline::ExecutePipeline(UInterchangeBaseNodeContainer* Bas
 											  TArray<UInterchangeSourceData*> const& SourceDatas,
 											  FString const& ContentBasePath)
 {
+	using namespace UE::Interchange::Materials;
+
 	if (TexturePipeline)
 	{
 		TexturePipeline->ScriptedExecutePipeline(BaseNodeContainer, SourceDatas, ContentBasePath);
 	}
 
-	TArray<UInterchangeSpecularProfileFactoryNode*> SpecularProfileFactoryNodes;
-
 	// Find all translated and factory nodes we need for this pipeline
 	BaseNodeContainer->IterateNodes(
-		[this, BaseNodeContainer, &SpecularProfileFactoryNodes](FString const& NodeUid, UInterchangeBaseNode* Node) {
+		[this](FString const& NodeUid, UInterchangeBaseNode* Node) {
 			switch (Node->GetNodeContainerType())
 			{
 			case EInterchangeNodeContainerType::TranslatedAsset:
@@ -73,29 +73,6 @@ void UAxFInterchangePipeline::ExecutePipeline(UInterchangeBaseNodeContainer* Bas
 					FAxFTextureMap& Pair = TexturesMap.FindOrAdd(ParameterName);
 					Pair.SizeMM = FVector2f{ WidthMM, HeightMM };
 				}
-				else if (UInterchangeSpecularProfileNode* SpecularProfileNode = Cast<UInterchangeSpecularProfileNode>(Node))
-				{
-					FString SpecProfileUID = SpecularProfileNode->GetUniqueID();
-					SpecProfileUID.ReplaceInline(TEXT("SPN"), TEXT("SP"));
-
-					UInterchangeSpecularProfileFactoryNode* FactoryNode = NewObject<UInterchangeSpecularProfileFactoryNode>(BaseNodeContainer, FName(SpecProfileUID));
-					FactoryNode->InitializeNode(SpecProfileUID, SpecProfileUID, EInterchangeNodeContainerType::FactoryData);
-
-					uint8 SpecularProfileFormat{};
-					if (SpecularProfileNode->GetCustomFormat(SpecularProfileFormat))
-					{
-						FactoryNode->SetCustomFormat(ESpecularProfileFormat(SpecularProfileFormat));
-					}
-
-					FString TextureUID{};
-					if (SpecularProfileNode->GetCustomTexture(TextureUID))
-					{
-						FactoryNode->SetCustomTexture(TextureUID);
-						FactoryNode->AddFactoryDependencyUid(TEXT("Factory_") + TextureUID);
-					}
-
-					SpecularProfileFactoryNodes.Emplace(FactoryNode);
-				}
 			}
 			break;
 			case EInterchangeNodeContainerType::FactoryData:
@@ -110,149 +87,163 @@ void UAxFInterchangePipeline::ExecutePipeline(UInterchangeBaseNodeContainer* Bas
 			}
 		});
 
-	// Defer adding the newly created specular profile factory node to the base node container
-	// because adding it immediately results in a "added during iteration of container" warning.
-	for (UInterchangeSpecularProfileFactoryNode* SpecularProfileFactoryNode : SpecularProfileFactoryNodes)
+	if (AxFObjectNodes.IsEmpty())
 	{
-		BaseNodeContainer->AddNode(SpecularProfileFactoryNode);
+		return;
 	}
-
-	if (AxFObjectNodes.IsEmpty())
+	TObjectPtr<UAxFMaterialObjectNode> AxFNode = AxFObjectNodes[0];
+	if (!AxFNode)
 	{
+		UE_LOG(LogAxFInterchangePipeline, Error, TEXT("AxF Interchange Node is not properly constructed!"));
 		return;
 	}
 
-	for (const UAxFMaterialObjectNode* AxFNode : AxFObjectNodes)
+	// Get Specular Profile node using string attribute set in AxFTranslator
+	const UInterchangeSpecularProfileNode* SpecularProfileNode = nullptr;
 	{
-		if (!AxFNode)
-		{
-			continue;
-		}
-
-		const TArray<EAxFFeature>& Features = AxFNode->PayloadData.UsedFeatures;
-		const TMap<FString, FLinearColor>& ValuesMap = AxFNode->PayloadData.ValuesMap;
-
-		FString BaseMaterialName;
-		if (Features.Contains(EAxFFeature::CPA2))
-		{
-			BaseMaterialName += TEXT("CPA2");
-			BaseMaterialName += (Features.Contains(EAxFFeature::Flakes) ? TEXT("_Flakes") : TEXT(""));
-		}
-		else
+		if (FString SpecularProfileUID;
+			AxFNode->GetStringAttribute(SubstrateMaterial::SpecularProfile.ToString(), SpecularProfileUID))
 		{
-			BaseMaterialName += "SVBRDF";
+			SpecularProfileNode = Cast<UInterchangeSpecularProfileNode>(BaseNodeContainer->GetNode(SpecularProfileUID));
+			if (SpecularProfileNode)
 			{
-				BaseMaterialName += (Features.Contains(EAxFFeature::Alpha) ? TEXT("_Alpha") : TEXT(""));
-				BaseMaterialName += (Features.Contains(EAxFFeature::Heightmap) ? TEXT("_POM") : TEXT(""));
+				CreateSpecularProfileFactoryNode(SpecularProfileNode, BaseNodeContainer, AxFNode);
+			}
+			else
+			{
+				UE_LOG(LogAxFInterchangePipeline, Error, TEXT("Couldn't find the SpecularProfileNode: %s"), *SpecularProfileUID);
+				return;
 			}
 		}
+	}
 
-		FString FullBaseMaterialPath =
-			TEXT("/InterchangeAxFAssets/BaseMaterials/") + BaseMaterialName + TEXT(".") + BaseMaterialName;
+	const TArray<EAxFFeature>& Features = AxFNode->PayloadData.UsedFeatures;
+	const TMap<FString, FLinearColor>& ValuesMap = AxFNode->PayloadData.ValuesMap;
 
-		if (!ensure(ParameterTypesPerMaterial.Contains(BaseMaterialName)))
+	FString BaseMaterialName;
+	if (Features.Contains(EAxFFeature::CPA2))
+	{
+		BaseMaterialName += TEXT("CPA2");
+		BaseMaterialName += (Features.Contains(EAxFFeature::Flakes) ? TEXT("_Flakes") : TEXT(""));
+	}
+	else
+	{
+		BaseMaterialName += "SVBRDF";
 		{
-			UE_LOG(LogAxFInterchangePipeline, Error, TEXT("Failed to create material instance %s: %s was not found"), *AxFNode->GetDisplayLabel(), *FullBaseMaterialPath);
-			continue;
+			BaseMaterialName += (Features.Contains(EAxFFeature::Alpha) ? TEXT("_Alpha") : TEXT(""));
+			BaseMaterialName += (Features.Contains(EAxFFeature::Heightmap) ? TEXT("_POM") : TEXT(""));
 		}
+	}
 
-		const FMaterialParameterTypes& ParameterTypes = ParameterTypesPerMaterial[BaseMaterialName];
+	FString FullBaseMaterialPath =
+		TEXT("/InterchangeAxFAssets/BaseMaterials/") + BaseMaterialName + TEXT(".") + BaseMaterialName;
 
-		FString MaterialInstanceUID = TEXT("MI_") + AxFNode->GetDisplayLabel();
+	if (!ensure(ParameterTypesPerMaterial.Contains(BaseMaterialName)))
+	{
+		UE_LOG(LogAxFInterchangePipeline, Error, TEXT("Failed to create material instance %s: %s was not found"), *AxFNode->GetDisplayLabel(), *FullBaseMaterialPath);
+		return;
+	}
 
-		UInterchangeMaterialInstanceFactoryNode* MaterialInstance = NewObject<UInterchangeMaterialInstanceFactoryNode>(BaseNodeContainer, FName(MaterialInstanceUID));
-		BaseNodeContainer->SetupNode(MaterialInstance, MaterialInstanceUID, AxFNode->GetDisplayLabel(), EInterchangeNodeContainerType::FactoryData);
+	const FMaterialParameterTypes& ParameterTypes = ParameterTypesPerMaterial[BaseMaterialName];
 
-		for (const EAxFFeature Feature : Features)
-		{
-			const FString Value = UEnum::GetValueAsString(Feature);
-			const FString AttributeKey = UInterchangeShaderPortsAPI::MakeInputValueKey(TEXT("Has") + Value);
+	FString MaterialInstanceUID = TEXT("MI_") + AxFNode->GetDisplayLabel();
 
-			MaterialInstance->AddBooleanAttribute(AttributeKey, true);
-		}
+	UInterchangeMaterialInstanceFactoryNode* MaterialInstance = NewObject<UInterchangeMaterialInstanceFactoryNode>(BaseNodeContainer, FName(MaterialInstanceUID));
+	BaseNodeContainer->SetupNode(MaterialInstance, MaterialInstanceUID, AxFNode->GetDisplayLabel(), EInterchangeNodeContainerType::FactoryData);
 
-		TArray<FString> ParameterNames;
-		ValuesMap.GetKeys(ParameterNames);
-		constexpr float Epsilon = 1e-6;
-		for (FString const& Param : ParameterNames)
-		{
-			if (Param == TEXT("FlakePatchSizeMM"))
-			{
-				const FString XAttributeKey = UInterchangeShaderPortsAPI::MakeInputValueKey(TEXT("WidthMM"));
-				const FString YAttributeKey = UInterchangeShaderPortsAPI::MakeInputValueKey(TEXT("HeightMM"));
-				const FLinearColor& FlakePatchSizeMM = ValuesMap[Param];
-				MaterialInstance->AddFloatAttribute(XAttributeKey, FlakePatchSizeMM.R + Epsilon);
-				MaterialInstance->AddFloatAttribute(YAttributeKey, FlakePatchSizeMM.G + Epsilon);
-			}
+	for (const EAxFFeature Feature : Features)
+	{
+		const FString Value = UEnum::GetValueAsString(Feature);
+		const FString AttributeKey = UInterchangeShaderPortsAPI::MakeInputValueKey(TEXT("Has") + Value);
 
-			// NOTE: Disable texture path in the case that we have a value for this parameter
-			FString TextureToggleAttributeKey = UInterchangeShaderPortsAPI::MakeInputValueKey(TEXT("Has") + Param + TEXT("Texture"));
-			MaterialInstance->AddBooleanAttribute(TextureToggleAttributeKey, false);
+		MaterialInstance->AddBooleanAttribute(AttributeKey, true);
+	}
 
-			const FString AttributeKey = UInterchangeShaderPortsAPI::MakeInputValueKey(Param);
-			const FLinearColor& Value = ValuesMap[Param];
-			if (ParameterTypes.ScalarParameters.Contains(Param))
-			{
-				MaterialInstance->AddFloatAttribute(AttributeKey, Value.R);
-			}
-			else
-			{
-				MaterialInstance->AddLinearColorAttribute(AttributeKey, Value);
-			}
+	TArray<FString> ParameterNames;
+	ValuesMap.GetKeys(ParameterNames);
+	constexpr float Epsilon = 1e-6;
+	for (FString const& Param : ParameterNames)
+	{
+		if (Param == TEXT("FlakePatchSizeMM"))
+		{
+			const FString XAttributeKey = UInterchangeShaderPortsAPI::MakeInputValueKey(TEXT("WidthMM"));
+			const FString YAttributeKey = UInterchangeShaderPortsAPI::MakeInputValueKey(TEXT("HeightMM"));
+			const FLinearColor& FlakePatchSizeMM = ValuesMap[Param];
+			MaterialInstance->AddFloatAttribute(XAttributeKey, FlakePatchSizeMM.R + Epsilon);
+			MaterialInstance->AddFloatAttribute(YAttributeKey, FlakePatchSizeMM.G + Epsilon);
 		}
 
-		if (bUseTriplanarMappingByDefault)
+		// NOTE: Disable texture path in the case that we have a value for this parameter
+		FString TextureToggleAttributeKey = UInterchangeShaderPortsAPI::MakeInputValueKey(TEXT("Has") + Param + TEXT("Texture"));
+		MaterialInstance->AddBooleanAttribute(TextureToggleAttributeKey, false);
+
+		const FString AttributeKey = UInterchangeShaderPortsAPI::MakeInputValueKey(Param);
+		const FLinearColor& Value = ValuesMap[Param];
+		if (ParameterTypes.ScalarParameters.Contains(Param))
 		{
-			FString GeoUVsToggleAttributeKey = UInterchangeShaderPortsAPI::MakeInputValueKey(TEXT("Use Geo UVs"));
-			MaterialInstance->AddBooleanAttribute(GeoUVsToggleAttributeKey, false);
+			MaterialInstance->AddFloatAttribute(AttributeKey, Value.R);
 		}
-
-		for (UInterchangeTexture2DFactoryNode* Texture2DFactoryNode : Texture2DFactoryNodes)
+		else
 		{
-			Texture2DFactoryNode->SetCustomSubPath(AxFNode->GetDisplayLabel());
+			MaterialInstance->AddLinearColorAttribute(AttributeKey, Value);
+		}
+	}
 
-			const FString Param = GetParameterNameFromUID(Texture2DFactoryNode->GetUniqueID());
+	if (bUseTriplanarMappingByDefault)
+	{
+		FString GeoUVsToggleAttributeKey = UInterchangeShaderPortsAPI::MakeInputValueKey(TEXT("Use Geo UVs"));
+		MaterialInstance->AddBooleanAttribute(GeoUVsToggleAttributeKey, false);
+	}
 
-			{
-				FString AttributeKey = UInterchangeShaderPortsAPI::MakeInputValueKey(TEXT("Has") + Param + TEXT("Texture"));
-				MaterialInstance->AddBooleanAttribute(AttributeKey, true);
-			}
+	for (UInterchangeTexture2DFactoryNode* Texture2DFactoryNode : Texture2DFactoryNodes)
+	{
+		Texture2DFactoryNode->SetCustomSubPath(AxFNode->GetDisplayLabel());
 
-			{
-				FString AttributeKey = UInterchangeShaderPortsAPI::MakeInputValueKey(Param + TEXT("Texture"));
-				MaterialInstance->AddStringAttribute(AttributeKey, Texture2DFactoryNode->GetUniqueID());
-			}
+		const FString Param = GetParameterNameFromUID(Texture2DFactoryNode->GetUniqueID());
 
-			if (Param == TEXT("HeightMap"))
-			{
-				FString AttributeKey = UInterchangeShaderPortsAPI::MakeInputValueKey(Param);
-				MaterialInstance->AddStringAttribute(AttributeKey, Texture2DFactoryNode->GetUniqueID());
-			}
+		{
+			FString AttributeKey = UInterchangeShaderPortsAPI::MakeInputValueKey(TEXT("Has") + Param + TEXT("Texture"));
+			MaterialInstance->AddBooleanAttribute(AttributeKey, true);
+		}
 
-			if (Param != TEXT("CPA2ColorTable"))
-			{
-				const FVector2f UVScale = TexturesMap.FindOrAdd(Param).SizeMM;
-				FString UVScaleXAttributeKey = UInterchangeShaderPortsAPI::MakeInputValueKey(TEXT("WidthMM"));
-				FString UVScaleYAttributeKey = UInterchangeShaderPortsAPI::MakeInputValueKey(TEXT("HeightMM"));
-				MaterialInstance->AddFloatAttribute(UVScaleXAttributeKey, UVScale.X);
-				MaterialInstance->AddFloatAttribute(UVScaleYAttributeKey, UVScale.Y);
-			}
+		{
+			FString AttributeKey = UInterchangeShaderPortsAPI::MakeInputValueKey(Param + TEXT("Texture"));
+			MaterialInstance->AddStringAttribute(AttributeKey, Texture2DFactoryNode->GetUniqueID());
+		}
 
-			MaterialInstance->AddFactoryDependencyUid(Texture2DFactoryNode->GetUniqueID());
+		if (Param == TEXT("HeightMap"))
+		{
+			FString AttributeKey = UInterchangeShaderPortsAPI::MakeInputValueKey(Param);
+			MaterialInstance->AddStringAttribute(AttributeKey, Texture2DFactoryNode->GetUniqueID());
 		}
-		Texture2DFactoryNodes.Empty();
 
-		for (UInterchangeSpecularProfileFactoryNode* SpecularProfileFactoryNode : SpecularProfileFactoryNodes)
+		if (Param != TEXT("CPA2ColorTable"))
 		{
-			SpecularProfileFactoryNode->SetCustomSubPath(AxFNode->GetDisplayLabel());
-			MaterialInstance->AddFactoryDependencyUid(SpecularProfileFactoryNode->GetUniqueID());
+			const FVector2f UVScale = TexturesMap.FindOrAdd(Param).SizeMM;
+			FString UVScaleXAttributeKey = UInterchangeShaderPortsAPI::MakeInputValueKey(TEXT("WidthMM"));
+			FString UVScaleYAttributeKey = UInterchangeShaderPortsAPI::MakeInputValueKey(TEXT("HeightMM"));
+			MaterialInstance->AddFloatAttribute(UVScaleXAttributeKey, UVScale.X);
+			MaterialInstance->AddFloatAttribute(UVScaleYAttributeKey, UVScale.Y);
+			FString UVOffsetYAttributeKey = UInterchangeShaderPortsAPI::MakeInputValueKey(TEXT("OffsetOnY"));
+			const float UVOffsetY = FMath::IsNearlyZero(UVScale.Y) ? 0.f : 1.0f - 100.0f / UVScale.Y;
+			MaterialInstance->AddFloatAttribute(UVOffsetYAttributeKey, UVOffsetY);
 		}
 
-		MaterialInstance->SetCustomParent(FullBaseMaterialPath);
-		MaterialInstance->MarkPackageDirty();
-		MaterialInstance->PostEditChange();
+		MaterialInstance->AddFactoryDependencyUid(Texture2DFactoryNode->GetUniqueID());
+	}
+	Texture2DFactoryNodes.Empty();
+
+	if (SpecularProfileNode != nullptr)
+	{
+		const FString SpecularProfileUid = SpecularProfileNode->GetUniqueID();
+		MaterialInstance->AddStringAttribute(SubstrateMaterial::SpecularProfile.ToString(), SpecularProfileUid);
+		MaterialInstance->AddFactoryDependencyUid(UInterchangeFactoryBaseNode::BuildFactoryNodeUid(SpecularProfileNode->GetUniqueID()));
 	}
 
+	MaterialInstance->SetCustomParent(FullBaseMaterialPath);
+	MaterialInstance->MarkPackageDirty();
+	MaterialInstance->PostEditChange();
+
 	BaseNodeContainer->IterateNodesOfType<UInterchangeFactoryBaseNode>(
 		[ReimportStrategyClosure = ReimportStrategy](const FString& NodeUid, UInterchangeFactoryBaseNode* FactoryNode)
 		{
@@ -450,7 +441,36 @@ void UAxFInterchangePipeline::AdjustSettingsForContext(FInterchangePipelineConte
 	Super::AdjustSettingsForContext(ContextParams);
 }
 
-void UAxFInterchangePipeline::PreDialogCleanup(FName const PipelineStackName)
+void UAxFInterchangePipeline::CreateSpecularProfileFactoryNode(const UInterchangeSpecularProfileNode* SpecularProfileNode, UInterchangeBaseNodeContainer* BaseNodeContainer, const TObjectPtr<UAxFMaterialObjectNode> AxFNode)
+{
+	const FString FactoryNodeUid = UInterchangeFactoryBaseNode::BuildFactoryNodeUid(SpecularProfileNode->GetUniqueID());
+
+	UInterchangeSpecularProfileFactoryNode* FactoryNode = NewObject<UInterchangeSpecularProfileFactoryNode>(BaseNodeContainer, NAME_None);
+
+	BaseNodeContainer->SetupNode(FactoryNode, FactoryNodeUid, SpecularProfileNode->GetDisplayLabel(), EInterchangeNodeContainerType::FactoryData);
+
+	FactoryNode->SetEnabled(true);
+	FactoryNode->SetCustomSubPath(AxFNode->GetDisplayLabel());
+
+	if (uint8 Format; SpecularProfileNode->GetCustomFormat(Format))
+	{
+		FactoryNode->SetCustomFormat(ESpecularProfileFormat{ Format });
+	}
+
+	if (FString TextureUid; SpecularProfileNode->GetCustomTexture(TextureUid))
+	{
+		if (BaseNodeContainer->GetNode(TextureUid))
+		{
+			FactoryNode->SetCustomTexture(TextureUid);
+			FactoryNode->AddFactoryDependencyUid(UInterchangeFactoryBaseNode::BuildFactoryNodeUid(TextureUid));
+		}
+	}
+
+	FactoryNode->AddTargetNodeUid(SpecularProfileNode->GetUniqueID());
+	SpecularProfileNode->AddTargetNodeUid(FactoryNode->GetUniqueID());
+}
+
+inline void UAxFInterchangePipeline::PreDialogCleanup(FName const PipelineStackName)
 {
 	if (TexturePipeline)
 	{
@@ -468,7 +488,7 @@ bool UAxFInterchangePipeline::IsSettingsAreValid(TOptional<FText>& OutInvalidRea
 	return Super::IsSettingsAreValid(OutInvalidReason);
 }
 
-void UAxFInterchangePipeline::FilterPropertiesFromTranslatedData(
+inline void UAxFInterchangePipeline::FilterPropertiesFromTranslatedData(
 	UInterchangeBaseNodeContainer* InBaseNodeContainer)
 {
 	if (TexturePipeline)
diff --git a/Engine/Plugins/Experimental/InterchangeAxF/Runtime/Source/InterchangeAxF/Private/AxFTranslator.cpp b/Engine/Plugins/Experimental/InterchangeAxF/Runtime/Source/InterchangeAxF/Private/AxFTranslator.cpp
--- a/Engine/Plugins/Experimental/InterchangeAxF/Runtime/Source/InterchangeAxF/Private/AxFTranslator.cpp
+++ b/Engine/Plugins/Experimental/InterchangeAxF/Runtime/Source/InterchangeAxF/Private/AxFTranslator.cpp
@@ -9,6 +9,7 @@
 #include "InterchangeTexture2DNode.h"
 #include "Modules/ModuleManager.h"
 #include "InterchangeImportModule.h"
+#include "InterchangeMaterialDefinitions.h"
 
 #include "Serialization/JsonSerializer.h"
 
@@ -195,6 +196,7 @@ bool UAxFTranslator::Translate(UInterchangeBaseNodeContainer& BaseNodeContainer)
 				BaseNodeContainer.SetupNode(SpecularProfileNode, SpecProfileUID, SpecProfileUID, EInterchangeNodeContainerType::TranslatedAsset);
 				SpecularProfileNode->SetCustomFormat((uint8)ESpecularProfileFormat::HalfVector);
 				SpecularProfileNode->SetCustomTexture(TextureUID);
+				AxFMaterialObjectNode->AddStringAttribute(UE::Interchange::Materials::SubstrateMaterial::SpecularProfile.ToString(), SpecProfileUID);
 			}
 		}
 		else if (bool ValueBool; JSON->TryGetBoolField(ParameterName, ValueBool))
diff --git a/Engine/Plugins/Experimental/InterchangeAxF/Runtime/Source/InterchangeAxF/Private/InterchangeAxF.cpp b/Engine/Plugins/Experimental/InterchangeAxF/Runtime/Source/InterchangeAxF/Private/InterchangeAxF.cpp
--- a/Engine/Plugins/Experimental/InterchangeAxF/Runtime/Source/InterchangeAxF/Private/InterchangeAxF.cpp
+++ b/Engine/Plugins/Experimental/InterchangeAxF/Runtime/Source/InterchangeAxF/Private/InterchangeAxF.cpp
@@ -17,63 +17,69 @@ DEFINE_LOG_CATEGORY(LogInterchangeAxF);
 // This code will execute after your module is loaded into memory; the exact timing is specified in the .uplugin file per-module
 void FInterchangeAxFModule::StartupModule()
 {
-    auto RegisterItems = []()
-    {
-        UInterchangeManager &InterchangeManager = UInterchangeManager::GetInterchangeManager();
+	bool IsSubstrateEnabled = GetDefault<URendererSettings>() ? GetDefault<URendererSettings>()->bEnableSubstrate : false;
+	if (!IsSubstrateEnabled)
+	{
+		UE_LOG(LogInterchangeAxF, Warning, TEXT("Interchange AxF requires Substrate to be enabled. Enable \"Substrate materials\" in Project Settings under Engine > Rendering."));
+	}
 
-        // Register the translators
-        InterchangeManager.RegisterTranslator(UAxFTranslator::StaticClass());
-
-        // Add X-Rite's AxF translator to the Interchange engine's pipeline set
-		if (UInterchangeProjectSettings* ProjectSettings = GetMutableDefault<UInterchangeProjectSettings>())
+	auto RegisterItems = []()
 		{
-			bool bPipelinesWereAdded = false;
-			FInterchangeContentImportSettings& AssetImportSettings = ProjectSettings->ContentImportSettings;
+			UInterchangeManager& InterchangeManager = UInterchangeManager::GetInterchangeManager();
+
+			// Register the translators
+			InterchangeManager.RegisterTranslator(UAxFTranslator::StaticClass());
 
-			if (AssetImportSettings.DefaultPipelineStackOverride.Contains(EInterchangeTranslatorAssetType::Materials))
+			// Add X-Rite's AxF translator to the Interchange engine's pipeline set
+			if (UInterchangeProjectSettings* ProjectSettings = GetMutableDefault<UInterchangeProjectSettings>())
 			{
-				const FName DefaultMaterialsStack = AssetImportSettings.DefaultPipelineStackOverride[EInterchangeTranslatorAssetType::Materials];
+				bool bPipelinesWereAdded = false;
+				FInterchangeContentImportSettings& AssetImportSettings = ProjectSettings->ContentImportSettings;
 
-				if (AssetImportSettings.PipelineStacks.Contains(DefaultMaterialsStack))
+				if (AssetImportSettings.DefaultPipelineStackOverride.Contains(EInterchangeTranslatorAssetType::Materials))
 				{
-					FInterchangeTranslatorPipelines TranslatorPipeplines;
-					TranslatorPipeplines.Translator = TSoftClassPtr<UInterchangeTranslatorBase>(UAxFTranslator::StaticClass());
-					TranslatorPipeplines.Pipelines.Add(FSoftObjectPath(TEXT("/InterchangeAxF/AxFInterchangePipeline.AxFInterchangePipeline")));
+					const FName DefaultMaterialsStack = AssetImportSettings.DefaultPipelineStackOverride[EInterchangeTranslatorAssetType::Materials];
+
+					if (AssetImportSettings.PipelineStacks.Contains(DefaultMaterialsStack))
 					{
-						FInterchangePipelineStack& PipelineStack = AssetImportSettings.PipelineStacks[DefaultMaterialsStack];
-						PipelineStack.PerTranslatorPipelines.Add(TranslatorPipeplines);
+						FInterchangeTranslatorPipelines TranslatorPipeplines;
+						TranslatorPipeplines.Translator = TSoftClassPtr<UInterchangeTranslatorBase>(UAxFTranslator::StaticClass());
+						TranslatorPipeplines.Pipelines.Add(FSoftObjectPath(TEXT("/InterchangeAxF/AxFInterchangePipeline.AxFInterchangePipeline")));
+						{
+							FInterchangePipelineStack& PipelineStack = AssetImportSettings.PipelineStacks[DefaultMaterialsStack];
+							PipelineStack.PerTranslatorPipelines.Add(TranslatorPipeplines);
+						}
+
+						bPipelinesWereAdded = true;
+						UE_LOG(LogInterchangeAxF, Log, TEXT("Interchange AxF's pipelines added to default Material pipeline stack, %s"), *DefaultMaterialsStack.ToString());
 					}
+				}
 
-					bPipelinesWereAdded = true;
-					UE_LOG(LogInterchangeAxF, Log, TEXT("Interchange AxF's pipelines added to default Material pipeline stack, %s"), *DefaultMaterialsStack.ToString());
+				if (!AssetImportSettings.ShowImportDialogOverride.Contains(EInterchangeTranslatorAssetType::Materials))
+				{
+					FInterchangePerTranslatorDialogOverride AxFTranslatorImportDialogOverride{};
+					AxFTranslatorImportDialogOverride.Translator = TSoftClassPtr<UInterchangeTranslatorBase>(UAxFTranslator::StaticClass());
+					AxFTranslatorImportDialogOverride.bShowImportDialog = false;
+					AxFTranslatorImportDialogOverride.bShowReimportDialog = false;
+					FInterchangeDialogOverride& DialogOverride = AssetImportSettings.ShowImportDialogOverride.FindOrAdd(EInterchangeTranslatorAssetType::Materials);
+					DialogOverride.PerTranslatorImportDialogOverride.Add(AxFTranslatorImportDialogOverride);
 				}
-			}
 
-			if (!AssetImportSettings.ShowImportDialogOverride.Contains(EInterchangeTranslatorAssetType::Materials))
-			{
-				FInterchangePerTranslatorDialogOverride AxFTranslatorImportDialogOverride{};
-				AxFTranslatorImportDialogOverride.Translator = TSoftClassPtr<UInterchangeTranslatorBase>(UAxFTranslator::StaticClass());
-				AxFTranslatorImportDialogOverride.bShowImportDialog = false;
-				AxFTranslatorImportDialogOverride.bShowReimportDialog = false;
-				FInterchangeDialogOverride& DialogOverride = AssetImportSettings.ShowImportDialogOverride.FindOrAdd(EInterchangeTranslatorAssetType::Materials); 
-				DialogOverride.PerTranslatorImportDialogOverride.Add(AxFTranslatorImportDialogOverride);
+				if (!bPipelinesWereAdded)
+				{
+					UE_LOG(LogInterchangeAxF, Warning, TEXT("Cannot find Interchange default Materials' stack to add pipelines. Interchange AxF material import may not work properly."));
+				}
 			}
+		};
 
-			if(!bPipelinesWereAdded)
-			{
-				UE_LOG(LogInterchangeAxF, Warning, TEXT("Cannot find Interchange default Materials' stack to add pipelines. Interchange AxF material import may not work properly."));
-			}
-		}
-	};
-
-    if (GEngine)
-    {
-        RegisterItems();
-    }
-    else
-    {
-        FCoreDelegates::GetOnPostEngineInit().AddLambda(RegisterItems);
-    }
+	if (GEngine)
+	{
+		RegisterItems();
+	}
+	else
+	{
+		FCoreDelegates::GetOnPostEngineInit().AddLambda(RegisterItems);
+	}
 }
 
 // This function may be called during shutdown to clean up your module.  For modules that support dynamic reloading,
diff --git a/Engine/Plugins/Experimental/InterchangeAxF/Runtime/Source/InterchangeAxF/Public/AxFInterchangePipeline.h b/Engine/Plugins/Experimental/InterchangeAxF/Runtime/Source/InterchangeAxF/Public/AxFInterchangePipeline.h
--- a/Engine/Plugins/Experimental/InterchangeAxF/Runtime/Source/InterchangeAxF/Public/AxFInterchangePipeline.h
+++ b/Engine/Plugins/Experimental/InterchangeAxF/Runtime/Source/InterchangeAxF/Public/AxFInterchangePipeline.h
@@ -6,6 +6,7 @@
 #include "Engine/SpecularProfile.h"
 #include "InterchangeGenericTexturePipeline.h"
 #include "InterchangePipelineBase.h"
+#include "InterchangeSpecularProfileNode.h"
 
 #include "AxFInterchangePipeline.generated.h"
 
@@ -68,6 +69,12 @@ protected:
 
 	virtual void AdjustSettingsForContext(FInterchangePipelineContextParams const& ContextParams) override;
 
+	// NOTE: This is a copied and modified version of UInterchangeGenericMaterialPipeline::CreateSpecularProfileFactoryNode().
+	// This would be handled automatically, but for now the compromise is to just do what the Generic Material Pipeline does, 
+	// in order to remove it as a dependency in the Interchange Pipeline Stack.
+	void CreateSpecularProfileFactoryNode(const UInterchangeSpecularProfileNode* SpecularProfileNode,
+										  UInterchangeBaseNodeContainer* BaseNodeContainer, TObjectPtr<UAxFMaterialObjectNode> AxFNode);
+
 	virtual void PreDialogCleanup(FName const PipelineStackName) override;
 
 	virtual bool IsSettingsAreValid(TOptional<FText>& OutInvalidReason) const override;
