[Interchange] Reworking Static and Skeletal Mesh import settings

Main changes :

- `bImportMeshesInBoneHierachy` has been renamed to `bConvertStaticsInBoneHierarchyToSkeletals`

- Conversion options has been moved into `Common Meshes/Conversion` section:
- Force All Mesh Types
- Convert Static to Skeletals (Morph Targets, Animations, In Bone Hierarchy)
- Use Single bone

- Common Meshes section has been removed from Skeletal and Static meshes sections and now have it's own dedicated one.
- This setup will make the Skeletal and Static mesh section disapear if no Skeletal/Static mesh is generated

- Added a custom version to support backward compatibility

Other fixes and adjustments in this CL :

- Fixed a recent regression where importing a single SM that have animations will import both a SM and a SKM (a test has been added to cover this case)
- Sparse Volume options have moved to `Textures/Volumes` fixing an issue where the Texture section will still be visible even if empty

[RN] Reorganized mesh import settings for clarity and consistency
#tests All Interchange tests passing except test_skinned_collapsing already failing on Main

#virtualized
#Synced-CL 52448386

#Synced-CL 52448386

diff --git a/Engine/Plugins/Experimental/Enterprise/DatasmithInterchange/Source/DatasmithInterchange/Private/InterchangeDatasmithPipeline.cpp b/Engine/Plugins/Experimental/Enterprise/DatasmithInterchange/Source/DatasmithInterchange/Private/InterchangeDatasmithPipeline.cpp
--- a/Engine/Plugins/Experimental/Enterprise/DatasmithInterchange/Source/DatasmithInterchange/Private/InterchangeDatasmithPipeline.cpp
+++ b/Engine/Plugins/Experimental/Enterprise/DatasmithInterchange/Source/DatasmithInterchange/Private/InterchangeDatasmithPipeline.cpp
@@ -95,9 +95,9 @@ namespace UE::Interchange::StaticMeshUtils
 UInterchangeDatasmithPipeline::UInterchangeDatasmithPipeline()
 {
 	CommonSkeletalMeshesAndAnimationsProperties = CreateDefaultSubobject<UInterchangeGenericCommonSkeletalMeshesAndAnimationsProperties>("CommonSkeletalMeshesAndAnimationsProperties");
-	CommonSkeletalMeshesAndAnimationsProperties->bImportMeshesInBoneHierarchy = false;
 	
 	CommonMeshesProperties = CreateDefaultSubobject<UInterchangeGenericCommonMeshesProperties>("CommonMeshesProperties");
+	CommonMeshesProperties->bConvertStaticsInBoneHierarchyToSkeletals = false;
 	CommonMeshesProperties->bBakeMeshes = false;
 
 	MaterialPipeline = CreateDefaultSubobject<UInterchangeDatasmithMaterialPipeline>("DatasmithMaterialPipeline");
diff --git a/Engine/Plugins/Interchange/Editor/Source/InterchangeEditor/Private/InterchangeFbxAssetImportDataConverter.cpp b/Engine/Plugins/Interchange/Editor/Source/InterchangeEditor/Private/InterchangeFbxAssetImportDataConverter.cpp
--- a/Engine/Plugins/Interchange/Editor/Source/InterchangeEditor/Private/InterchangeFbxAssetImportDataConverter.cpp
+++ b/Engine/Plugins/Interchange/Editor/Source/InterchangeEditor/Private/InterchangeFbxAssetImportDataConverter.cpp
@@ -151,7 +151,7 @@ namespace UE::Interchange::Private
 
 	void FillFbxSkeletalMeshImportData(const UInterchangeGenericAssetsPipeline* GenericAssetPipeline, UFbxSkeletalMeshImportData* DestinationSkeletalMeshImportData)
 	{
-		DestinationSkeletalMeshImportData->bImportMeshesInBoneHierarchy = GenericAssetPipeline->CommonSkeletalMeshesAndAnimationsProperties->bImportMeshesInBoneHierarchy;
+		DestinationSkeletalMeshImportData->bImportMeshesInBoneHierarchy = GenericAssetPipeline->CommonMeshesProperties->bConvertStaticsInBoneHierarchyToSkeletals;
 		DestinationSkeletalMeshImportData->bImportMorphTargets = GenericAssetPipeline->MeshPipeline->bImportMorphTargets;
 		DestinationSkeletalMeshImportData->bImportVertexAttributes = GenericAssetPipeline->MeshPipeline->bImportVertexAttributes;
 		DestinationSkeletalMeshImportData->bKeepSectionsSeparate = GenericAssetPipeline->CommonMeshesProperties->bKeepSectionsSeparate;
@@ -231,7 +231,7 @@ namespace UE::Interchange::Private
 		DestinationAnimSequenceImportData->bDoNotImportCurveWithZero = GenericAssetPipeline->AnimationPipeline->bDoNotImportCurveWithZero;
 		DestinationAnimSequenceImportData->bImportBoneTracks = GenericAssetPipeline->AnimationPipeline->bImportBoneTracks;
 		DestinationAnimSequenceImportData->bImportCustomAttribute = GenericAssetPipeline->AnimationPipeline->bImportCustomAttribute;
-		DestinationAnimSequenceImportData->bImportMeshesInBoneHierarchy = GenericAssetPipeline->CommonSkeletalMeshesAndAnimationsProperties->bImportMeshesInBoneHierarchy;
+		DestinationAnimSequenceImportData->bImportMeshesInBoneHierarchy = GenericAssetPipeline->CommonMeshesProperties->bConvertStaticsInBoneHierarchyToSkeletals;
 		DestinationAnimSequenceImportData->bPreserveLocalTransform = false;
 		DestinationAnimSequenceImportData->bRemoveRedundantKeys = GenericAssetPipeline->AnimationPipeline->bRemoveCurveRedundantKeys;
 		DestinationAnimSequenceImportData->bSetMaterialDriveParameterOnCustomAttribute = GenericAssetPipeline->AnimationPipeline->bSetMaterialDriveParameterOnCustomAttribute;
@@ -353,7 +353,7 @@ namespace UE::Interchange::Private
 			FillInterchangeGenericAssetsPipelineFromFbxMeshImportData(GenericAssetPipeline, SkeletalMeshImportData);
 		}
 
-		GenericAssetPipeline->CommonSkeletalMeshesAndAnimationsProperties->bImportMeshesInBoneHierarchy = SkeletalMeshImportData->bImportMeshesInBoneHierarchy;
+		GenericAssetPipeline->CommonMeshesProperties->bConvertStaticsInBoneHierarchyToSkeletals = SkeletalMeshImportData->bImportMeshesInBoneHierarchy;
 		GenericAssetPipeline->CommonMeshesProperties->bKeepSectionsSeparate = SkeletalMeshImportData->bKeepSectionsSeparate;
 		GenericAssetPipeline->MeshPipeline->bCreatePhysicsAsset = false;
 		GenericAssetPipeline->MeshPipeline->bImportMorphTargets = SkeletalMeshImportData->bImportMorphTargets;
@@ -439,7 +439,7 @@ namespace UE::Interchange::Private
 		GenericAssetPipeline->AnimationPipeline->bDoNotImportCurveWithZero = AnimSequenceImportData->bDoNotImportCurveWithZero;
 		GenericAssetPipeline->AnimationPipeline->bImportBoneTracks = AnimSequenceImportData->bImportBoneTracks;
 		GenericAssetPipeline->AnimationPipeline->bImportCustomAttribute = AnimSequenceImportData->bImportCustomAttribute;
-		GenericAssetPipeline->CommonSkeletalMeshesAndAnimationsProperties->bImportMeshesInBoneHierarchy = AnimSequenceImportData->bImportMeshesInBoneHierarchy;
+		GenericAssetPipeline->CommonMeshesProperties->bConvertStaticsInBoneHierarchyToSkeletals = AnimSequenceImportData->bImportMeshesInBoneHierarchy;
 		GenericAssetPipeline->AnimationPipeline->bRemoveCurveRedundantKeys = AnimSequenceImportData->bRemoveRedundantKeys;
 		GenericAssetPipeline->AnimationPipeline->bSetMaterialDriveParameterOnCustomAttribute = AnimSequenceImportData->bSetMaterialDriveParameterOnCustomAttribute;
 		GenericAssetPipeline->AnimationPipeline->bSnapToClosestFrameBoundary = AnimSequenceImportData->bSnapToClosestFrameBoundary;
diff --git a/Engine/Plugins/Interchange/Editor/Source/Pipelines/Private/InterchangeEditorPipelinesModule.cpp b/Engine/Plugins/Interchange/Editor/Source/Pipelines/Private/InterchangeEditorPipelinesModule.cpp
--- a/Engine/Plugins/Interchange/Editor/Source/Pipelines/Private/InterchangeEditorPipelinesModule.cpp
+++ b/Engine/Plugins/Interchange/Editor/Source/Pipelines/Private/InterchangeEditorPipelinesModule.cpp
@@ -166,17 +166,21 @@ void FInterchangeEditorPipelinesModule::RegisterPropertySectionMappings()
 		Section->AddCategory("Conflicts");
 	}
 
+	// Meshes
+	{
+		TSharedRef<FPropertySection> Section = RegisterPropertySection(PropertyModule, "InterchangeGenericAssetsPipeline", "0 Meshes", LOCTEXT("Meshes", "Meshes"));
+		Section->AddCategory("Common Meshes");
+	}
+
 	// Static Meshes
 	{
 		TSharedRef<FPropertySection> Section = RegisterPropertySection(PropertyModule, "InterchangeGenericMeshPipeline", "1 StaticMeshes", LOCTEXT("Static Meshes", "Static Meshes"));
-		Section->AddCategory("Common Meshes");
 		Section->AddCategory("Static Meshes");
 	}
 
 	// Skeletal Meshes
 	{
 		TSharedRef<FPropertySection> Section = RegisterPropertySection(PropertyModule, "InterchangeGenericMeshPipeline", "2 SkeletalMeshes", LOCTEXT("Skeletal Meshes", "Skeletal Meshes"));
-		Section->AddCategory("Common Meshes");
 		Section->AddCategory("Common Skeletal Meshes and Animations");
 		Section->AddCategory("Skeletal Meshes");
 	}
@@ -198,6 +202,7 @@ void FInterchangeEditorPipelinesModule::RegisterPropertySectionMappings()
 	{
 		TSharedRef<FPropertySection> Section = RegisterPropertySection(PropertyModule, "InterchangeGenericTexturePipeline", "5 Textures", LOCTEXT("Textures", "Textures"));
 		Section->AddCategory("Textures");
+		Section->AddCategory("Volumes");
 	}
 }
 
diff --git a/Engine/Plugins/Interchange/Runtime/Source/Pipelines/Private/InterchangeGenericAnimationPipeline.cpp b/Engine/Plugins/Interchange/Runtime/Source/Pipelines/Private/InterchangeGenericAnimationPipeline.cpp
--- a/Engine/Plugins/Interchange/Runtime/Source/Pipelines/Private/InterchangeGenericAnimationPipeline.cpp
+++ b/Engine/Plugins/Interchange/Runtime/Source/Pipelines/Private/InterchangeGenericAnimationPipeline.cpp
@@ -33,7 +33,7 @@ namespace UE::Interchange::Private
 {
 	const FString ConvertedFromRigidAnimationPrefixIdentifier = TEXT("\\SkeletalAnimation\\ConvertedFromRigidAnimation\\");
 
-	bool IsTranslatedDataContainOnlyJointAnimation(const UInterchangeBaseNodeContainer* InBaseNodeContainer, bool bConvertStaticsWithMorphTargetsToSkeletals, bool bImportMeshesInBoneHierarchy)
+	bool IsTranslatedDataContainOnlyJointAnimation(const UInterchangeBaseNodeContainer* InBaseNodeContainer, bool bConvertStaticsWithMorphTargetsToSkeletals, bool bConvertStaticsInBoneHierarchyToSkeletals)
 	{
 		//Its valid to call GetMeshesInformationFromTranslatedData with a null container
 		if (!InBaseNodeContainer)
@@ -50,11 +50,11 @@ namespace UE::Interchange::Private
 		if (bContainJointAnimation)
 		{
 			bool bContainSkinnedMeshNode = false;
-			InBaseNodeContainer->BreakableIterateNodesOfType<UInterchangeMeshNode>([&bContainSkinnedMeshNode, &bConvertStaticsWithMorphTargetsToSkeletals, &bImportMeshesInBoneHierarchy](const FString& NodeUid, UInterchangeMeshNode* MeshNode)
+			InBaseNodeContainer->BreakableIterateNodesOfType<UInterchangeMeshNode>([&bContainSkinnedMeshNode, &bConvertStaticsWithMorphTargetsToSkeletals, &bConvertStaticsInBoneHierarchyToSkeletals](const FString& NodeUid, UInterchangeMeshNode* MeshNode)
 				{
 					if (!MeshNode->IsMorphTarget())
 					{
-						if (MeshNode->IsSkinnedMesh() || bImportMeshesInBoneHierarchy)
+						if (MeshNode->IsSkinnedMesh() || bConvertStaticsInBoneHierarchyToSkeletals)
 						{
 							bContainSkinnedMeshNode = true;
 						}
@@ -213,8 +213,8 @@ void UInterchangeGenericAnimationPipeline::AdjustSettingsForContext(const FInter
 	{
 		if(UE::Interchange::Private::IsTranslatedDataContainOnlyJointAnimation(
 			ContextParams.BaseNodeContainer, 
-			CommonSkeletalMeshesAndAnimationsProperties->bConvertStaticsWithMorphTargetsToSkeletals,
-			CommonSkeletalMeshesAndAnimationsProperties->bImportMeshesInBoneHierarchy))
+			CommonMeshesProperties->bConvertStaticsWithMorphTargetsToSkeletals,
+			CommonMeshesProperties->bConvertStaticsInBoneHierarchyToSkeletals))
 		{
 			bImportAnimations = true;
 			CommonSkeletalMeshesAndAnimationsProperties->bImportOnlyAnimations = true;
@@ -949,7 +949,7 @@ void UInterchangeGenericAnimationPipeline::CreateAnimSequenceFactoryNode(UInterc
 			bSkeletonCompatible = bSkeletonCompatible && UE::Interchange::Private::FSkeletonHelper::IsCompatibleSkeleton(Skeleton
 				, RootJointUid
 				, BaseNodeContainer
-				, CommonSkeletalMeshesAndAnimationsProperties->bConvertStaticsWithMorphTargetsToSkeletals || CommonMeshesProperties->ForceAllMeshAsType == EInterchangeForceMeshType::IFMT_SkeletalMesh
+				, CommonMeshesProperties->bConvertStaticsWithMorphTargetsToSkeletals || CommonMeshesProperties->ForceAllMeshAsType == EInterchangeForceMeshType::IFMT_SkeletalMesh
 				, false /*bCheckForIdenticalSkeleton*/
 				, CommonMeshesProperties->bImportSockets);
 #endif
diff --git a/Engine/Plugins/Interchange/Runtime/Source/Pipelines/Private/InterchangeGenericAssetsPipeline.cpp b/Engine/Plugins/Interchange/Runtime/Source/Pipelines/Private/InterchangeGenericAssetsPipeline.cpp
--- a/Engine/Plugins/Interchange/Runtime/Source/Pipelines/Private/InterchangeGenericAssetsPipeline.cpp
+++ b/Engine/Plugins/Interchange/Runtime/Source/Pipelines/Private/InterchangeGenericAssetsPipeline.cpp
@@ -115,6 +115,31 @@ void UInterchangeGenericAssetsPipeline::PreDialogCleanup(const FName PipelineSta
 	}
 }
 
+void UInterchangeGenericAssetsPipeline::Serialize(FArchive& Ar)
+{
+	Super::Serialize(Ar);
+	Ar.UsingCustomVersion(FInterchangeCustomVersion::GUID);
+}
+
+void UInterchangeGenericAssetsPipeline::PostLoad()
+{
+	Super::PostLoad();
+
+	if (GetLinkerCustomVersion(FInterchangeCustomVersion::GUID) < FInterchangeCustomVersion::SkeletalAndStaticMeshOptionsCleanUp)
+	{
+		if (CommonSkeletalMeshesAndAnimationsProperties && CommonMeshesProperties)
+		{
+			PRAGMA_DISABLE_DEPRECATION_WARNINGS
+			CommonMeshesProperties->bConvertStaticsWithAnimatedTransformToSkeletals = CommonMeshesProperties->bAutoDetectMeshType_DEPRECATED;
+			CommonMeshesProperties->bConvertStaticsInBoneHierarchyToSkeletals = CommonSkeletalMeshesAndAnimationsProperties->bImportMeshesInBoneHierarchy_DEPRECATED;
+			CommonMeshesProperties->bConvertStaticsWithMorphTargetsToSkeletals = CommonSkeletalMeshesAndAnimationsProperties->bConvertStaticsWithMorphTargetsToSkeletals_DEPRECATED;
+			PRAGMA_ENABLE_DEPRECATION_WARNINGS
+		}
+	}
+}
+
+
+
 bool UInterchangeGenericAssetsPipeline::IsSettingsAreValid(TOptional<FText>& OutInvalidReason) const
 {
 	if (MaterialPipeline && !MaterialPipeline->IsSettingsAreValid(OutInvalidReason))
@@ -273,29 +298,6 @@ void UInterchangeGenericAssetsPipeline::FilterPropertiesFromTranslatedData(UInte
 				RawAnimationNode++;
 			});
 
-		UInterchangePipelineBase* OuterMostPipeline = GetMostPipelineOuter();
-		auto HideFullCategory = [this, &OuterMostPipeline](const FString& Category, UInterchangePipelineBase* Pipeline)
-		{
-			TArray<FString> HideCategories;
-			//Filter out all Textures properties
-			HideCategories.Add(Category);
-			if (OuterMostPipeline)
-			{
-				for (const FString& HideCategoryName : HideCategories)
-				{
-					HidePropertiesOfCategory(OuterMostPipeline, Pipeline, HideCategoryName);
-				}
-			}
-		};
-
-		auto LocalHideProperty = [this, &OuterMostPipeline](UInterchangePipelineBase* Pipeline, FName PropertyName)
-		{
-			if (OuterMostPipeline)
-			{
-				HideProperty(OuterMostPipeline, Pipeline, PropertyName);
-			}
-		};
-
 		//Found which categories to hide
 		bool bHideStaticMeshes = false;
 		bool bHideSkeletalMeshes = false;
@@ -327,21 +329,23 @@ void UInterchangeGenericAssetsPipeline::FilterPropertiesFromTranslatedData(UInte
 			bHideStaticMeshes = true;
 		}
 
+		UInterchangePipelineBase* OuterMostPipeline = GetMostPipelineOuter();
+
 		if (SkeletalMeshes.Num() > 0)
 		{
 			if (MeshPipeline->SkeletalMeshImportContentType == EInterchangeSkeletalMeshContentType::SkinningWeights)
 			{
-				LocalHideProperty(CommonMeshesProperties, GET_MEMBER_NAME_CHECKED(UInterchangeGenericCommonMeshesProperties, VertexOverrideColor));
-				LocalHideProperty(CommonMeshesProperties, GET_MEMBER_NAME_CHECKED(UInterchangeGenericCommonMeshesProperties, VertexColorImportOption));
-				LocalHideProperty(MeshPipeline, GET_MEMBER_NAME_CHECKED(UInterchangeGenericMeshPipeline, bImportMorphTargets));
-				LocalHideProperty(MeshPipeline, GET_MEMBER_NAME_CHECKED(UInterchangeGenericMeshPipeline, ThresholdPosition));
-				LocalHideProperty(MeshPipeline, GET_MEMBER_NAME_CHECKED(UInterchangeGenericMeshPipeline, ThresholdTangentNormal));
-				LocalHideProperty(MeshPipeline, GET_MEMBER_NAME_CHECKED(UInterchangeGenericMeshPipeline, ThresholdUV));
-				LocalHideProperty(MeshPipeline, GET_MEMBER_NAME_CHECKED(UInterchangeGenericMeshPipeline, MorphThresholdPosition));
+				HideProperty(OuterMostPipeline, CommonMeshesProperties, GET_MEMBER_NAME_CHECKED(UInterchangeGenericCommonMeshesProperties, VertexOverrideColor));
+				HideProperty(OuterMostPipeline, CommonMeshesProperties, GET_MEMBER_NAME_CHECKED(UInterchangeGenericCommonMeshesProperties, VertexColorImportOption));
+				HideProperty(OuterMostPipeline, MeshPipeline, GET_MEMBER_NAME_CHECKED(UInterchangeGenericMeshPipeline, bImportMorphTargets));
+				HideProperty(OuterMostPipeline, MeshPipeline, GET_MEMBER_NAME_CHECKED(UInterchangeGenericMeshPipeline, ThresholdPosition));
+				HideProperty(OuterMostPipeline, MeshPipeline, GET_MEMBER_NAME_CHECKED(UInterchangeGenericMeshPipeline, ThresholdTangentNormal));
+				HideProperty(OuterMostPipeline, MeshPipeline, GET_MEMBER_NAME_CHECKED(UInterchangeGenericMeshPipeline, ThresholdUV));
+				HideProperty(OuterMostPipeline, MeshPipeline, GET_MEMBER_NAME_CHECKED(UInterchangeGenericMeshPipeline, MorphThresholdPosition));
 			}
 			else if (MeshPipeline->SkeletalMeshImportContentType == EInterchangeSkeletalMeshContentType::Geometry)
 			{
-				LocalHideProperty(MeshPipeline, GET_MEMBER_NAME_CHECKED(UInterchangeGenericMeshPipeline, bUpdateSkeletonReferencePose));
+				HideProperty(OuterMostPipeline, MeshPipeline, GET_MEMBER_NAME_CHECKED(UInterchangeGenericMeshPipeline, bUpdateSkeletonReferencePose));
 			}
 		}
 
@@ -360,10 +364,10 @@ void UInterchangeGenericAssetsPipeline::FilterPropertiesFromTranslatedData(UInte
 		{
 			if (MaterialPipeline)
 			{
-				HideFullCategory(UInterchangeGenericMaterialPipeline::GetPipelineCategory(nullptr), MaterialPipeline);
+				HidePropertiesOfCategory(OuterMostPipeline, MaterialPipeline, UInterchangeGenericMaterialPipeline::GetPipelineCategory(nullptr));
 				if (MaterialPipeline->TexturePipeline)
 				{
-					HideFullCategory(UInterchangeGenericTexturePipeline::GetPipelineCategory(nullptr), MaterialPipeline->TexturePipeline);
+					HidePropertiesOfCategory(OuterMostPipeline, MaterialPipeline->TexturePipeline, UInterchangeGenericTexturePipeline::GetPipelineCategory(nullptr));
 				}
 			}
 		}
@@ -371,35 +375,35 @@ void UInterchangeGenericAssetsPipeline::FilterPropertiesFromTranslatedData(UInte
 		//Hide the categories
 		if (bHideStaticMeshes)
 		{
-			HideFullCategory(UInterchangeGenericMeshPipeline::GetPipelineCategory(UStaticMesh::StaticClass()), MeshPipeline);
+			HidePropertiesOfCategory(OuterMostPipeline, MeshPipeline, UInterchangeGenericMeshPipeline::GetPipelineCategory(UStaticMesh::StaticClass()));
 		}
 		if (bHideSkeletalMeshes)
 		{
-			HideFullCategory(UInterchangeGenericMeshPipeline::GetPipelineCategory(USkeletalMesh::StaticClass()), MeshPipeline);
+			HidePropertiesOfCategory(OuterMostPipeline, MeshPipeline, UInterchangeGenericMeshPipeline::GetPipelineCategory(USkeletalMesh::StaticClass()));
 		}
 		if (bHideGeometryCaches)
 		{
-			HideFullCategory(UInterchangeGenericMeshPipeline::GetPipelineCategory(UGeometryCache::StaticClass()), MeshPipeline);
+			HidePropertiesOfCategory(OuterMostPipeline, MeshPipeline, UInterchangeGenericMeshPipeline::GetPipelineCategory(UGeometryCache::StaticClass()));
 		}
-		if(bHideCommonMeshes)
+		if (bHideCommonMeshes)
 		{
-			HideFullCategory(UInterchangeGenericCommonMeshesProperties::GetPipelineCategory(nullptr), CommonMeshesProperties);
+			HidePropertiesOfCategory(OuterMostPipeline, CommonMeshesProperties, UInterchangeGenericCommonMeshesProperties::GetPipelineCategory(nullptr));
 		}
 		if (bHideCommonSkeletalMeshesAndAnimations)
 		{
 			//We cannot hide common sk and anim if import only animation is true.
 			if (!CommonSkeletalMeshesAndAnimationsProperties->bImportOnlyAnimations)
 			{
-				HideFullCategory(UInterchangeGenericCommonSkeletalMeshesAndAnimationsProperties::GetPipelineCategory(nullptr), CommonSkeletalMeshesAndAnimationsProperties);
+				HidePropertiesOfCategory(OuterMostPipeline, CommonSkeletalMeshesAndAnimationsProperties, UInterchangeGenericCommonSkeletalMeshesAndAnimationsProperties::GetPipelineCategory(nullptr));
 			}
 		}
 		if (bHideCommonSkeletalMeshesAndAnimations_StaticMesh)
 		{
-			HideFullCategory(UInterchangeGenericMeshPipeline::GetPipelineCategory(UStaticMesh::StaticClass()), CommonSkeletalMeshesAndAnimationsProperties);
+			HidePropertiesOfCategory(OuterMostPipeline, CommonSkeletalMeshesAndAnimationsProperties, UInterchangeGenericMeshPipeline::GetPipelineCategory(UStaticMesh::StaticClass()));
 		}
 		if (bHideAnimations)
 		{
-			HideFullCategory(UInterchangeGenericAnimationPipeline::GetPipelineCategory(nullptr), AnimationPipeline);
+			HidePropertiesOfCategory(OuterMostPipeline, AnimationPipeline, UInterchangeGenericAnimationPipeline::GetPipelineCategory(nullptr));
 		}
 	}
 
@@ -976,7 +980,7 @@ void UInterchangeGenericAssetsPipeline::CreateSkeletonConflict(USkeleton* Specif
 			return;
 		}
 
-		const bool bConvertStaticToSkeletalActive = CommonSkeletalMeshesAndAnimationsProperties->bConvertStaticsWithMorphTargetsToSkeletals || CommonMeshesProperties->ForceAllMeshAsType == EInterchangeForceMeshType::IFMT_SkeletalMesh;
+		const bool bConvertStaticToSkeletalActive = CommonMeshesProperties->bConvertStaticsWithMorphTargetsToSkeletals || CommonMeshesProperties->ForceAllMeshAsType == EInterchangeForceMeshType::IFMT_SkeletalMesh;
 		//If we have a compatible skeleton we do not need to create a conflict, for stricter comparisons we check for identical skeletons to report any deviations.
 		constexpr bool bCheckForIdenticalSkeleton = true;
 		if (UE::Interchange::Private::FSkeletonHelper::IsCompatibleSkeleton(SpecifiedSkeleton ? SpecifiedSkeleton : SkeletalMesh->GetSkeleton()
diff --git a/Engine/Plugins/Interchange/Runtime/Source/Pipelines/Private/InterchangeGenericGeometryCachePipeline.cpp b/Engine/Plugins/Interchange/Runtime/Source/Pipelines/Private/InterchangeGenericGeometryCachePipeline.cpp
--- a/Engine/Plugins/Interchange/Runtime/Source/Pipelines/Private/InterchangeGenericGeometryCachePipeline.cpp
+++ b/Engine/Plugins/Interchange/Runtime/Source/Pipelines/Private/InterchangeGenericGeometryCachePipeline.cpp
@@ -95,7 +95,7 @@ void UInterchangeGenericMeshPipeline::ExecutePreImportPipelineGeometryCache()
 		// The bConvertStaticsWithAnimatedTransformToSkeletals setting converts static meshes with animated transforms into skeletal meshes with AnimSequence
 		// For the purpose of geometry cache, both are considered as just meshes
 		TArray<FString> AnimMeshInstanceUids;
-		if (CommonSkeletalMeshesAndAnimationsProperties->bConvertStaticsWithAnimatedTransformToSkeletals && CommonMeshesProperties->ForceAllMeshAsType == EInterchangeForceMeshType::IFMT_None)
+		if (CommonMeshesProperties->bConvertStaticsWithAnimatedTransformToSkeletals && CommonMeshesProperties->ForceAllMeshAsType == EInterchangeForceMeshType::IFMT_None)
 		{
 			PipelineMeshesUtilities->GetAllSkinnedMeshInstance(AnimMeshInstanceUids);
 		}
diff --git a/Engine/Plugins/Interchange/Runtime/Source/Pipelines/Private/InterchangeGenericMaterialPipeline.cpp b/Engine/Plugins/Interchange/Runtime/Source/Pipelines/Private/InterchangeGenericMaterialPipeline.cpp
--- a/Engine/Plugins/Interchange/Runtime/Source/Pipelines/Private/InterchangeGenericMaterialPipeline.cpp
+++ b/Engine/Plugins/Interchange/Runtime/Source/Pipelines/Private/InterchangeGenericMaterialPipeline.cpp
@@ -735,25 +735,20 @@ void UInterchangeGenericMaterialPipeline::FilterPropertiesFromTranslatedData(UIn
 	uint32 MaterialCount = TmpMaterialNodes.Num();
 	InBaseNodeContainer->GetNodes(UInterchangeMaterialInstanceNode::StaticClass(), TmpMaterialNodes);
 	MaterialCount += TmpMaterialNodes.Num();
-	if(MaterialCount == 0)
+	if (MaterialCount == 0)
 	{
-		TArray<FString> HideCategories;
-		//Filter out all material properties
-		HideCategories.Add(TEXT("Materials"));
 		if (UInterchangePipelineBase* OuterMostPipeline = GetMostPipelineOuter())
 		{
-			for (const FString& HideCategoryName : HideCategories)
-			{
-				HidePropertiesOfCategory(OuterMostPipeline, this, HideCategoryName);
-			}
+			HidePropertiesOfCategory(OuterMostPipeline, this, UInterchangeGenericMaterialPipeline::GetPipelineCategory(nullptr));
 		}
 	}
 
-	if(TexturePipeline)
+	if (TexturePipeline)
 	{
 		TexturePipeline->FilterPropertiesFromTranslatedData(InBaseNodeContainer);
 	}
-	if(SparseVolumeTexturePipeline)
+
+	if (SparseVolumeTexturePipeline)
 	{
 		SparseVolumeTexturePipeline->FilterPropertiesFromTranslatedData(InBaseNodeContainer);
 	}
diff --git a/Engine/Plugins/Interchange/Runtime/Source/Pipelines/Private/InterchangeGenericMeshPipeline.cpp b/Engine/Plugins/Interchange/Runtime/Source/Pipelines/Private/InterchangeGenericMeshPipeline.cpp
--- a/Engine/Plugins/Interchange/Runtime/Source/Pipelines/Private/InterchangeGenericMeshPipeline.cpp
+++ b/Engine/Plugins/Interchange/Runtime/Source/Pipelines/Private/InterchangeGenericMeshPipeline.cpp
@@ -124,7 +124,7 @@ void UInterchangeGenericMeshPipeline::AdjustSettingsForContext(const FInterchang
 			SkeletalMeshImportContentType = EInterchangeSkeletalMeshContentType::All;
 			CommonSkeletalMeshesAndAnimationsProperties->Skeleton.Reset();
 			CommonSkeletalMeshesAndAnimationsProperties->bImportOnlyAnimations = false;
-			CommonSkeletalMeshesAndAnimationsProperties->bConvertStaticsWithAnimatedTransformToSkeletals = false;
+			CommonMeshesProperties->bConvertStaticsWithAnimatedTransformToSkeletals = false;
 		}
 		else if (ContextParams.ContextType == EInterchangePipelineContext::AssetCustomMorphTargetImport
 			|| ContextParams.ContextType == EInterchangePipelineContext::AssetCustomMorphTargetReImport)
@@ -136,7 +136,7 @@ void UInterchangeGenericMeshPipeline::AdjustSettingsForContext(const FInterchang
 			CommonMeshesProperties->bImportLods = true;
 			CommonMeshesProperties->bKeepSectionsSeparate = false;
 			CommonMeshesProperties->VertexColorImportOption = EInterchangeVertexColorImportOption::IVCIO_Ignore;
-			CommonSkeletalMeshesAndAnimationsProperties->bConvertStaticsWithAnimatedTransformToSkeletals = false;
+			CommonMeshesProperties->bConvertStaticsWithAnimatedTransformToSkeletals = false;
 			bImportSkeletalMeshes = false;
 			bImportStaticMeshes = true;
 			CombineStaticMeshesBehavior = EInterchangeCombineStaticMeshesBehavior::All;
@@ -160,7 +160,7 @@ void UInterchangeGenericMeshPipeline::AdjustSettingsForContext(const FInterchang
 				{
 					bImportStaticMeshes = true;
 					CommonMeshesProperties->ForceAllMeshAsType = EInterchangeForceMeshType::IFMT_StaticMesh;
-					CommonSkeletalMeshesAndAnimationsProperties->bConvertStaticsWithAnimatedTransformToSkeletals = false;
+					CommonMeshesProperties->bConvertStaticsWithAnimatedTransformToSkeletals = false;
 					bImportSkeletalMeshes = false;
 					CombineStaticMeshesBehavior = EInterchangeCombineStaticMeshesBehavior::All;
 					LodGroup = NAME_None;
@@ -175,13 +175,13 @@ void UInterchangeGenericMeshPipeline::AdjustSettingsForContext(const FInterchang
 				{
 					bImportSkeletalMeshes = true;
 					CommonMeshesProperties->ForceAllMeshAsType = EInterchangeForceMeshType::IFMT_SkeletalMesh;
-					CommonSkeletalMeshesAndAnimationsProperties->bConvertStaticsWithAnimatedTransformToSkeletals = false;
+					CommonMeshesProperties->bConvertStaticsWithAnimatedTransformToSkeletals = false;
 					bCreatePhysicsAsset = false;
 					bImportStaticMeshes = false;
 				}
 				else
 				{
-					CommonSkeletalMeshesAndAnimationsProperties->bConvertStaticsWithAnimatedTransformToSkeletals = true;
+					CommonMeshesProperties->bConvertStaticsWithAnimatedTransformToSkeletals = true;
 				}
 			}
 		}
@@ -202,7 +202,7 @@ void UInterchangeGenericMeshPipeline::AdjustSettingsForContext(const FInterchang
 
 		GetMeshesInformationFromTranslatedData(ContextParams.BaseNodeContainer, bContainStaticMesh, bContainSkeletalMesh, bContainGeometryCache);
 
-		CommonSkeletalMeshesAndAnimationsProperties->bConvertStaticsWithAnimatedTransformToSkeletals = false;
+		CommonMeshesProperties->bConvertStaticsWithAnimatedTransformToSkeletals = false;
 
 		HideSubCategories.Add(TEXT("Build"));
 		if (USkeletalMesh* SkeletalMesh = Cast<USkeletalMesh>(ContextParams.ReimportAsset))
@@ -421,9 +421,9 @@ UInterchangePipelineMeshesUtilities* UInterchangeGenericMeshPipeline::CreateMesh
 
 	//Set the context option to use when querying the pipeline mesh utilities
 	DataContext.ForceMeshType = Pipeline->CommonMeshesProperties->ForceAllMeshAsType;
-	DataContext.bConvertStaticsWithAnimatedTransformToSkeletals = Pipeline->CommonSkeletalMeshesAndAnimationsProperties->bConvertStaticsWithAnimatedTransformToSkeletals;
-	DataContext.bConvertStaticsWithMorphTargetsToSkeletals = Pipeline->CommonSkeletalMeshesAndAnimationsProperties->bConvertStaticsWithMorphTargetsToSkeletals;
-	DataContext.bImportMeshesInBoneHierarchy = Pipeline->CommonSkeletalMeshesAndAnimationsProperties->bImportMeshesInBoneHierarchy;
+	DataContext.bConvertStaticsWithAnimatedTransformToSkeletals = Pipeline->CommonMeshesProperties->bConvertStaticsWithAnimatedTransformToSkeletals;
+	DataContext.bConvertStaticsWithMorphTargetsToSkeletals = Pipeline->CommonMeshesProperties->bConvertStaticsWithMorphTargetsToSkeletals;
+	DataContext.bConvertStaticsInBoneHierarchyToSkeletals = Pipeline->CommonMeshesProperties->bConvertStaticsInBoneHierarchyToSkeletals;
 	DataContext.bQueryGeometryOnlyIfNoInstance = Pipeline->CommonMeshesProperties->bBakeMeshes || Pipeline->CommonMeshesProperties->bBakePivotMeshes;
 
 	UInterchangePipelineMeshesUtilities* CreatedPipelineMeshesUtilities = UInterchangePipelineMeshesUtilities::CreateInterchangePipelineMeshesUtilities(InBaseNodeContainer);
diff --git a/Engine/Plugins/Interchange/Runtime/Source/Pipelines/Private/InterchangeGenericSkeletalMeshPipeline.cpp b/Engine/Plugins/Interchange/Runtime/Source/Pipelines/Private/InterchangeGenericSkeletalMeshPipeline.cpp
--- a/Engine/Plugins/Interchange/Runtime/Source/Pipelines/Private/InterchangeGenericSkeletalMeshPipeline.cpp
+++ b/Engine/Plugins/Interchange/Runtime/Source/Pipelines/Private/InterchangeGenericSkeletalMeshPipeline.cpp
@@ -305,7 +305,7 @@ UInterchangeSkeletalMeshFactoryNode* UInterchangeGenericMeshPipeline::CreateSkel
 			bSkeletonCompatible = UE::Interchange::Private::FSkeletonHelper::IsCompatibleSkeleton(Skeleton
 				, RootJointUid
 				, BaseNodeContainer
-				, CommonSkeletalMeshesAndAnimationsProperties->bConvertStaticsWithMorphTargetsToSkeletals || CommonMeshesProperties->ForceAllMeshAsType == EInterchangeForceMeshType::IFMT_SkeletalMesh
+				, CommonMeshesProperties->bConvertStaticsWithMorphTargetsToSkeletals || CommonMeshesProperties->ForceAllMeshAsType == EInterchangeForceMeshType::IFMT_SkeletalMesh
 				, false /*bCheckForIdenticalSkeleton*/
 				, CommonMeshesProperties->bImportSockets
 			);
diff --git a/Engine/Plugins/Interchange/Runtime/Source/Pipelines/Private/InterchangePipelineMeshesUtilities.cpp b/Engine/Plugins/Interchange/Runtime/Source/Pipelines/Private/InterchangePipelineMeshesUtilities.cpp
--- a/Engine/Plugins/Interchange/Runtime/Source/Pipelines/Private/InterchangePipelineMeshesUtilities.cpp
+++ b/Engine/Plugins/Interchange/Runtime/Source/Pipelines/Private/InterchangePipelineMeshesUtilities.cpp
@@ -376,7 +376,7 @@ bool FInterchangePipelineMeshesUtilitiesContext::IsSkeletalMeshInstance(const FI
 			return true;
 		}
 	}
-	if (bImportMeshesInBoneHierarchy && MeshInstance.IsNestedInSkeleton())
+	if (bConvertStaticsInBoneHierarchyToSkeletals && MeshInstance.IsNestedInSkeleton())
 	{
 		return true;
 	}
@@ -410,7 +410,7 @@ bool FInterchangePipelineMeshesUtilitiesContext::IsStaticMeshGeometry(const FInt
 	{
 		return true;
 	}
-	if (bImportMeshesInBoneHierarchy && MeshGeometry.bIsReferencedBySkeleton)
+	if (bConvertStaticsInBoneHierarchyToSkeletals && MeshGeometry.bIsReferencedBySkeleton)
 	{
 		return false;
 	}
@@ -431,10 +431,23 @@ bool FInterchangePipelineMeshesUtilitiesContext::IsSkeletalMeshGeometry(const FI
 	{
 		return false;
 	}
-	if (!MeshGeometry.BaseNode->IsA<UInterchangeGeometryCacheNode>()
-		&& (ForceMeshType == EInterchangeForceMeshType::IFMT_SkeletalMesh
-			|| MeshGeometry.IsSkinnedMesh()
-			|| (bConvertStaticsWithMorphTargetsToSkeletals && MeshGeometry.HasMorphTargetDependencies())))
+	if (MeshGeometry.BaseNode->IsA<UInterchangeGeometryCacheNode>())
+	{
+		return false;
+	}
+	if (ForceMeshType == EInterchangeForceMeshType::IFMT_SkeletalMesh)
+	{
+		return true;
+	}
+	if (MeshGeometry.IsSkinnedMesh())
+	{
+		return true;
+	}
+	if (bConvertStaticsWithMorphTargetsToSkeletals && MeshGeometry.HasMorphTargetDependencies())
+	{
+		return true;
+	}
+	if (bConvertStaticsWithAnimatedTransformToSkeletals && MeshGeometry.bHasTransformAnimationTrack)
 	{
 		return true;
 	}
@@ -614,12 +627,17 @@ UInterchangePipelineMeshesUtilities* UInterchangePipelineMeshesUtilities::Create
 			if (AssetNode && (AssetNode->IsA<UInterchangeMeshNode>() || AssetNode->IsA<UInterchangeMeshLODContainerNode>()))
 			{
 				const FString& StaticMeshNodeUid = AssetNode->GetUniqueID();
+				FInterchangeMeshGeometry* MeshGeometry = PipelineMeshesUtilities->MeshGeometriesPerMeshUid.Find(StaticMeshNodeUid);
 
 				for (const UInterchangeSceneNode* AnimatedSceneNode : SceneNodesWithTransformAnimationTrackNode)
 				{
 					if (IsImpactingMeshRecursive(AnimatedSceneNode, BaseNodeContainer, StaticMeshNodeUid))
 					{
 						MeshInstancePtr->bHasTransformAnimationTrack = true;
+						if ensure(MeshGeometry) //We could have used FindChecked above but the since the semantic is not clear let's avoid a crash.
+						{
+							MeshGeometry->bHasTransformAnimationTrack = true;
+						}
 						break;
 					}
 				}
@@ -1088,7 +1106,7 @@ void UInterchangePipelineMeshesUtilities::GetCombinedSkinnedMeshInstances(TMap<F
 						}
 
 						//If we want to use a single bone skeleton for the static mesh, we need to create or reuse a dedicated scene node for that
-						//Make sure to ignore nested static meshes so we don't affect "Import Meshes In Bone Hierarchy" option behavior
+						//By design, only Static Meshes that are not nested in a joint hierarchy can use a single bone skeleton.
 						if (!bIsNestedIntoSkeleton && bUseSingleBoneForConvertedMeshes)
 						{
 							//Make sure to use the SkeletonRootUid found above so all meshes under the same root joint will use the same single bone skeleton too.
diff --git a/Engine/Plugins/Interchange/Runtime/Source/Pipelines/Private/InterchangeSparseVolumeTexturePipeline.cpp b/Engine/Plugins/Interchange/Runtime/Source/Pipelines/Private/InterchangeSparseVolumeTexturePipeline.cpp
--- a/Engine/Plugins/Interchange/Runtime/Source/Pipelines/Private/InterchangeSparseVolumeTexturePipeline.cpp
+++ b/Engine/Plugins/Interchange/Runtime/Source/Pipelines/Private/InterchangeSparseVolumeTexturePipeline.cpp
@@ -256,8 +256,7 @@ namespace UE::Interchange::Private
 
 FString UInterchangeSparseVolumeTexturePipeline::GetPipelineCategory(UClass* AssetClass)
 {
-	// Ideally we'd be in a "Volumes" one, but these seem to be somewhat hard-coded?
-	return TEXT("Textures");
+	return TEXT("Volumes");
 }
 
 void UInterchangeSparseVolumeTexturePipeline::AdjustSettingsForContext(const FInterchangePipelineContextParams& ContextParams)
diff --git a/Engine/Plugins/Interchange/Runtime/Source/Pipelines/Public/InterchangeGenericAssetsPipeline.h b/Engine/Plugins/Interchange/Runtime/Source/Pipelines/Public/InterchangeGenericAssetsPipeline.h
--- a/Engine/Plugins/Interchange/Runtime/Source/Pipelines/Public/InterchangeGenericAssetsPipeline.h
+++ b/Engine/Plugins/Interchange/Runtime/Source/Pipelines/Public/InterchangeGenericAssetsPipeline.h
@@ -114,6 +114,10 @@ public:
 
 	UE_API virtual void PreDialogCleanup(const FName PipelineStackName) override;
 
+	UE_API virtual void Serialize(FArchive& Ar) override;
+
+	UE_API virtual void PostLoad() override;
+
 	UE_API virtual bool IsSettingsAreValid(TOptional<FText>& OutInvalidReason) const override;
 
 
diff --git a/Engine/Plugins/Interchange/Runtime/Source/Pipelines/Public/InterchangeGenericAssetsPipelineSharedSettings.h b/Engine/Plugins/Interchange/Runtime/Source/Pipelines/Public/InterchangeGenericAssetsPipelineSharedSettings.h
--- a/Engine/Plugins/Interchange/Runtime/Source/Pipelines/Public/InterchangeGenericAssetsPipelineSharedSettings.h
+++ b/Engine/Plugins/Interchange/Runtime/Source/Pipelines/Public/InterchangeGenericAssetsPipelineSharedSettings.h
@@ -67,15 +67,27 @@ public:
 	 * For skeletal meshes the conversion will happen only if there is no skinned meshes.
 	 * Mixing rigid skeletal mesh with skinned mesh is not good and will result in multiple skeletal meshes.
 	 */
- 	UPROPERTY(EditAnywhere, BlueprintReadWrite, Category = "Common Meshes")
- 	EInterchangeForceMeshType ForceAllMeshAsType = EInterchangeForceMeshType::IFMT_None;
+ 	UPROPERTY(EditAnywhere, BlueprintReadWrite, Category = "Common Meshes", meta = (SubCategory = "Conversion"))
+	EInterchangeForceMeshType ForceAllMeshAsType = EInterchangeForceMeshType::IFMT_None;
+
+	/** If enabled, and some static mesh transforms are animated, the pipeline will convert the static mesh into a rigid skeletal mesh. */
+	UPROPERTY(EditAnywhere, BlueprintReadWrite, Category = "Common Meshes", meta = (SubCategory = "Conversion", editcondition = "ForceAllMeshAsType == EInterchangeForceMeshType::IFMT_None"))
+	bool bConvertStaticsWithAnimatedTransformToSkeletals = true;
+
+	/** If enabled, all static meshes that have morph targets will be imported as skeletal meshes instead. */
+	UPROPERTY(EditAnywhere, BlueprintReadWrite, Category = "Common Meshes", meta = (SubCategory = "Conversion", editcondition = "ForceAllMeshAsType == EInterchangeForceMeshType::IFMT_None"))
+	bool bConvertStaticsWithMorphTargetsToSkeletals = false;
+
+	/** If enabled, all static meshes nested in bone hierarchies will be imported as skeletal meshes instead. */
+	UPROPERTY(EditAnywhere, BlueprintReadWrite, Category = "Common Meshes", meta = (SubCategory = "Conversion", editcondition = "ForceAllMeshAsType == EInterchangeForceMeshType::IFMT_None"))
+	bool bConvertStaticsInBoneHierarchyToSkeletals = true;
 
 	/**
 	* If enabled, converted static meshes to skeletal meshes will use a skeleton with a unique bone at the geometry origin.
 	* Otherwise, the bone structure is determined by the model hierarchy.
 	* This setting is only used if Force All Meshes As Type setting is set to "Skeletal Mesh".
 	*/
-	UPROPERTY(EditAnywhere, BlueprintReadWrite, Category = "Common Meshes", meta = (editcondition = "ForceAllMeshAsType == EInterchangeForceMeshType::IFMT_SkeletalMesh"))
+	UPROPERTY(EditAnywhere, BlueprintReadWrite, Category = "Common Meshes", meta = (SubCategory = "Conversion", editcondition = "ForceAllMeshAsType == EInterchangeForceMeshType::IFMT_SkeletalMesh"))
 	bool bSingleBoneSkeleton = false;
 
 	/** If enabled, any existing LODs for meshes are imported. This setting is only used if the Bake Meshes setting is also enabled. */
@@ -142,8 +154,8 @@ public:
 	UPROPERTY(EditAnywhere, BlueprintReadWrite, Category = "Common Meshes", meta = (SubCategory = "Build"))
 	bool bRemoveDegenerates = false;
 
-	UE_DEPRECATED(5.8, "Use UInterchangeGenericCommonSkeletalMeshesAndAnimationsProperties::bConvertStaticsWithAnimatedTransformToSkeletals instead")
-	UPROPERTY(meta = (DeprecatedProperty, DeprecationMessage = "bAutoDetectMeshType is deprecated. Use UInterchangeGenericCommonSkeletalMeshesAndAnimationsProperties::bConvertStaticsWithAnimatedTransformToSkeletals instead."))
+	UE_DEPRECATED(5.8, "Use bConvertStaticsWithAnimatedTransformToSkeletals instead")
+	UPROPERTY(meta = (DeprecatedProperty, DeprecationMessage = "bAutoDetectMeshType is deprecated. Use bConvertStaticsWithAnimatedTransformToSkeletals instead."))
 	bool bAutoDetectMeshType_DEPRECATED = true;
 
 #if WITH_EDITOR
@@ -152,6 +164,9 @@ public:
 		static const TSet<FName> NeedRefreshProperties =
 		{
 			GET_MEMBER_NAME_CHECKED(UInterchangeGenericCommonMeshesProperties, ForceAllMeshAsType),
+			GET_MEMBER_NAME_CHECKED(UInterchangeGenericCommonMeshesProperties, bConvertStaticsWithAnimatedTransformToSkeletals),
+			GET_MEMBER_NAME_CHECKED(UInterchangeGenericCommonMeshesProperties, bConvertStaticsWithMorphTargetsToSkeletals),
+			GET_MEMBER_NAME_CHECKED(UInterchangeGenericCommonMeshesProperties, bConvertStaticsInBoneHierarchyToSkeletals),
 			GET_MEMBER_NAME_CHECKED(UInterchangeGenericCommonMeshesProperties, bSingleBoneSkeleton)
 		};
 
@@ -186,13 +201,6 @@ public:
 	/** Skeleton to use for imported asset. When importing a skeletal mesh, leaving this as "None" will create a new skeleton. When importing an animation, this must be specified. */
 	UPROPERTY(EditAnywhere, BlueprintReadWrite, Category = "Common Skeletal Meshes and Animations", meta = (EditCondition = "!bTryAutoSelectSkeleton"))
 	TSoftObjectPtr<USkeleton> Skeleton;
-	
-	/** 
-	 * If enabled, meshes nested in bone hierarchies will be imported as meshes instead of being converted to bones. If the meshes are not skinned, they are
-	 * added to the skeletal mesh and removed from the list of static meshes.
-	 */
-	UPROPERTY(EditAnywhere, BlueprintReadWrite, Category = "Common Skeletal Meshes and Animations")
-	bool bImportMeshesInBoneHierarchy = true;
 
 	/** If enabled, skinned meshes will be rebind using the joints transform pose at frame 0 instead of being import with the bind pose. */
 	UPROPERTY(EditAnywhere, BlueprintReadWrite, Category = "Common Skeletal Meshes and Animations", meta=(DisplayName="Use T0 As Ref Pose"))
@@ -202,22 +210,19 @@ public:
 	UPROPERTY(EditAnywhere, BlueprintReadWrite, Category = "Common Skeletal Meshes and Animations")
 	bool bAddCurveMetadataToSkeleton = true;
 
-	/** If enabled, and some static mesh transforms are animated, the pipeline will convert the static mesh into a rigid skeletal mesh. */
-	UPROPERTY(EditAnywhere, BlueprintReadWrite, Category = "Static Meshes")
-	bool bConvertStaticsWithAnimatedTransformToSkeletals = true;
+	UE_DEPRECATED(5.8, "Use UInterchangeGenericCommonMeshesProperties::bConvertStaticsInBoneHierarchyToSkeletals instead")
+	UPROPERTY(meta = (DeprecatedProperty, DeprecationMessage = "bImportMeshesInBoneHierarchy is deprecated. Use UInterchangeGenericCommonMeshesProperties::bConvertStaticsInBoneHierarchyToSkeletals instead."))
+	bool bImportMeshesInBoneHierarchy_DEPRECATED = true;
 
-	/** If enabled, all static meshes that have morph targets will be imported as skeletal meshes instead. */
-	UPROPERTY(EditAnywhere, BlueprintReadWrite, Category = "Static Meshes")
-	bool bConvertStaticsWithMorphTargetsToSkeletals = false;
+	UE_DEPRECATED(5.8, "Use UInterchangeGenericCommonMeshesProperties::bConvertStaticsWithMorphTargetsToSkeletals instead")
+	UPROPERTY(meta = (DeprecatedProperty, DeprecationMessage = "bConvertStaticsWithMorphTargetsToSkeletals is deprecated. Use UInterchangeGenericCommonMeshesProperties::bConvertStaticsWithMorphTargetsToSkeletals instead."))
+	bool bConvertStaticsWithMorphTargetsToSkeletals_DEPRECATED = false;
 
 #if WITH_EDITOR
 	virtual bool IsPropertyChangeNeedRefresh(const FPropertyChangedEvent& PropertyChangedEvent) const override
 	{
 		static const TSet<FName> NeedRefreshProperties =
 		{
-			GET_MEMBER_NAME_CHECKED(UInterchangeGenericCommonSkeletalMeshesAndAnimationsProperties, bConvertStaticsWithAnimatedTransformToSkeletals),
-			GET_MEMBER_NAME_CHECKED(UInterchangeGenericCommonSkeletalMeshesAndAnimationsProperties, bConvertStaticsWithMorphTargetsToSkeletals),
-			GET_MEMBER_NAME_CHECKED(UInterchangeGenericCommonSkeletalMeshesAndAnimationsProperties, bImportMeshesInBoneHierarchy),
 			GET_MEMBER_NAME_CHECKED(UInterchangeGenericCommonSkeletalMeshesAndAnimationsProperties, Skeleton),
 			GET_MEMBER_NAME_CHECKED(UInterchangeGenericCommonSkeletalMeshesAndAnimationsProperties, bImportOnlyAnimations)
 		};
diff --git a/Engine/Plugins/Interchange/Runtime/Source/Pipelines/Public/InterchangeGenericMaterialPipeline.h b/Engine/Plugins/Interchange/Runtime/Source/Pipelines/Public/InterchangeGenericMaterialPipeline.h
--- a/Engine/Plugins/Interchange/Runtime/Source/Pipelines/Public/InterchangeGenericMaterialPipeline.h
+++ b/Engine/Plugins/Interchange/Runtime/Source/Pipelines/Public/InterchangeGenericMaterialPipeline.h
@@ -124,7 +124,7 @@ public:
 	UPROPERTY(VisibleAnywhere, BlueprintReadWrite, Instanced, Category = "Textures")
 	TObjectPtr<UInterchangeGenericTexturePipeline> TexturePipeline;
 
-	UPROPERTY(VisibleAnywhere, BlueprintReadWrite, Instanced, Category = "Textures")
+	UPROPERTY(VisibleAnywhere, BlueprintReadWrite, Instanced, Category = "Volumes")
 	TObjectPtr<UInterchangeSparseVolumeTexturePipeline> SparseVolumeTexturePipeline;
 
 	/**
diff --git a/Engine/Plugins/Interchange/Runtime/Source/Pipelines/Public/InterchangePipelineMeshesUtilities.h b/Engine/Plugins/Interchange/Runtime/Source/Pipelines/Public/InterchangePipelineMeshesUtilities.h
--- a/Engine/Plugins/Interchange/Runtime/Source/Pipelines/Public/InterchangePipelineMeshesUtilities.h
+++ b/Engine/Plugins/Interchange/Runtime/Source/Pipelines/Public/InterchangePipelineMeshesUtilities.h
@@ -88,6 +88,7 @@ struct FInterchangeMeshGeometry
 	FInterchangeMeshGeometry()
 		: BaseNode(nullptr)
 		, bIsReferencedBySkeleton(false)
+		, bHasTransformAnimationTrack(false)
 		, bIsMorphTarget(false)
 		, bIsSkinnedMesh(false)
 		, bHasMorphTargetDependencies(false)
@@ -122,6 +123,12 @@ struct FInterchangeMeshGeometry
 	UPROPERTY(EditAnywhere, Category = "Interchange | Pipeline | MeshesGeometry")
 	bool bIsReferencedBySkeleton;
 
+	/*
+	* Indicates that at least 1 of the Instancing Node Uids has a transform animation.
+	*/
+	UPROPERTY(EditAnywhere, Category = "Interchange | Pipeline | MeshesGeometry")
+	bool bHasTransformAnimationTrack;
+
 	void SetupMeshGeometry(const UInterchangeBaseNodeContainer& BaseNodeContainer, const UInterchangeBaseNode* InBaseNode);
 
 	bool IsMorphTarget() const;
@@ -202,11 +209,10 @@ struct FInterchangePipelineMeshesUtilitiesContext
 	bool bConvertStaticsWithMorphTargetsToSkeletals = false;
 
 	/**
-	 * If enabled, meshes nested in bone hierarchies are imported as meshes instead of being converted to bones. If the meshes are not skinned, they are
-	 * added to the skeletal mesh and removed from the list of static meshes.
-	 */
+	* If enabled, all static meshes nested in bone hierarchies will be imported as skeletal meshes instead.
+	*/
 	UPROPERTY(EditAnywhere, BlueprintReadWrite, Category = "Interchange | Pipeline | MeshesContext")
-	bool bImportMeshesInBoneHierarchy = true;
+	bool bConvertStaticsInBoneHierarchyToSkeletals = true;
 
 	/**
 	 * When querying geometry, this flag will not add MeshGeometry if there is a scene node that points to a geometry.
@@ -220,6 +226,10 @@ struct FInterchangePipelineMeshesUtilitiesContext
 	UPROPERTY(EditAnywhere, BlueprintReadWrite, Category = "Interchange | Pipeline | MeshesContext")
 	bool bIgnoreGeometryCaches = false;
 
+	UE_DEPRECATED(5.8, "Use bConvertStaticsInBoneHierarchyToSkeletals instead.")
+	UPROPERTY(meta = (DeprecatedProperty, DeprecationMessage = "Use bConvertStaticsInBoneHierarchyToSkeletals instead."))
+	bool bImportMeshesInBoneHierarchy_DEPRECATED = true;
+
 	UE_DEPRECATED(5.8, "Use ForceMeshType instead.")
 	UPROPERTY(meta = (DeprecatedProperty, DeprecationMessage = "Use ForceMeshType instead."))
 	bool bConvertStaticMeshToSkeletalMesh_DEPRECATED = false;
diff --git a/Engine/Plugins/Interchange/Runtime/Source/Pipelines/Public/InterchangeSparseVolumeTexturePipeline.h b/Engine/Plugins/Interchange/Runtime/Source/Pipelines/Public/InterchangeSparseVolumeTexturePipeline.h
--- a/Engine/Plugins/Interchange/Runtime/Source/Pipelines/Public/InterchangeSparseVolumeTexturePipeline.h
+++ b/Engine/Plugins/Interchange/Runtime/Source/Pipelines/Public/InterchangeSparseVolumeTexturePipeline.h
@@ -26,24 +26,24 @@ public:
 	UPROPERTY(
 		EditAnywhere,
 		BlueprintReadWrite,
-		Category = "Textures",
+		Category = "Volumes",
 		meta = (StandAlonePipelineProperty = "True", PipelineInternalEditionData = "True")
 	)
 	FString PipelineDisplayName;
 
 	/** If enabled, imports all sparse volume texture assets found in the source. */
-	UPROPERTY(EditAnywhere, BlueprintReadWrite, Category = "Textures")
+	UPROPERTY(EditAnywhere, BlueprintReadWrite, Category = "Volumes")
 	bool bImportSparseVolumeTextures = true;
 
 	/**
 	 * If enabled, will attempt to import volume nodes corresponding to numbered files in the same folder as individual frames
 	 * of an animated SparseVolumeTextures.
 	 */
-	UPROPERTY(EditAnywhere, BlueprintReadWrite, Category = "Textures")
+	UPROPERTY(EditAnywhere, BlueprintReadWrite, Category = "Volumes")
 	bool bImportAnimatedSparseVolumeTextures = true;
 
 	/** If set, and there is only one asset and one source, the imported asset will be given this name. */
-	UPROPERTY(EditAnywhere, BlueprintReadWrite, Category = "Textures", meta = (StandAlonePipelineProperty = "True"))
+	UPROPERTY(EditAnywhere, BlueprintReadWrite, Category = "Volumes", meta = (StandAlonePipelineProperty = "True"))
 	FString AssetName;
 
 public:
diff --git a/Engine/Source/Runtime/Interchange/Core/Private/InterchangePipelineBase.cpp b/Engine/Source/Runtime/Interchange/Core/Private/InterchangePipelineBase.cpp
--- a/Engine/Source/Runtime/Interchange/Core/Private/InterchangePipelineBase.cpp
+++ b/Engine/Source/Runtime/Interchange/Core/Private/InterchangePipelineBase.cpp
@@ -565,6 +565,11 @@ FString UInterchangePipelineBase::GetPipelineDisplayName() const
 
 void UInterchangePipelineBase::InternalToggleVisibilityPropertiesOfMetaDataValue(UInterchangePipelineBase* OuterMostPipeline, UInterchangePipelineBase* Pipeline, bool bDoTransientSubPipeline, const FString& MetaDataKey, const FString& MetaDataValue, const bool VisibilityState)
 {
+	if (!OuterMostPipeline)
+	{
+		return;
+	}
+
 	UClass* PipelineClass = Pipeline->GetClass();
 	for (FProperty* Property = PipelineClass->PropertyLink; Property; Property = Property->PropertyLinkNext)
 	{
@@ -610,6 +615,11 @@ void UInterchangePipelineBase::HidePropertiesOfSubCategory(UInterchangePipelineB
 
 void UInterchangePipelineBase::HideProperty(UInterchangePipelineBase* OuterMostPipeline, UInterchangePipelineBase* Pipeline, const FName& HidePropertyName)
 {
+	if (!OuterMostPipeline)
+	{
+		return;
+	}
+
 	constexpr bool bVisibilityState = false;
 	UClass* PipelineClass = Pipeline->GetClass();
 	for (FProperty* Property = PipelineClass->PropertyLink; Property; Property = Property->PropertyLinkNext)
diff --git a/Engine/Source/Runtime/Interchange/Core/Public/InterchangeCustomVersion.h b/Engine/Source/Runtime/Interchange/Core/Public/InterchangeCustomVersion.h
--- a/Engine/Source/Runtime/Interchange/Core/Public/InterchangeCustomVersion.h
+++ b/Engine/Source/Runtime/Interchange/Core/Public/InterchangeCustomVersion.h
@@ -22,6 +22,9 @@ struct FInterchangeCustomVersion
 		// Replaced 'DoNotSearch' Material's search location enum value with a boolean property
 		AdditionOfReuseExistingMaterialsBool,
 
+		// Moved mesh conversion setting from UInterchangeGenericCommonSkeletalMeshesAndAnimationsProperties to UInterchangeGenericCommonMeshesProperties
+		SkeletalAndStaticMeshOptionsCleanUp,
+
 		// -----<new versions can be added above this line>-------------------------------------------------
 		VersionPlusOne,
 		LatestVersion = VersionPlusOne - 1
