Parentheses missing from if-ensure statement

#rnx

#Synced-CL 52672165

diff --git a/Engine/Plugins/Interchange/Runtime/Source/Pipelines/Private/InterchangePipelineMeshesUtilities.cpp b/Engine/Plugins/Interchange/Runtime/Source/Pipelines/Private/InterchangePipelineMeshesUtilities.cpp
--- a/Engine/Plugins/Interchange/Runtime/Source/Pipelines/Private/InterchangePipelineMeshesUtilities.cpp
+++ b/Engine/Plugins/Interchange/Runtime/Source/Pipelines/Private/InterchangePipelineMeshesUtilities.cpp
@@ -634,7 +634,7 @@ UInterchangePipelineMeshesUtilities* UInterchangePipelineMeshesUtilities::Create
 					if (IsImpactingMeshRecursive(AnimatedSceneNode, BaseNodeContainer, StaticMeshNodeUid))
 					{
 						MeshInstancePtr->bHasTransformAnimationTrack = true;
-						if ensure(MeshGeometry) //We could have used FindChecked above but the since the semantic is not clear let's avoid a crash.
+						if (ensure(MeshGeometry)) //We could have used FindChecked above but the since the semantic is not clear let's avoid a crash.
 						{
 							MeshGeometry->bHasTransformAnimationTrack = true;
 						}
