Interchange - Fixed Global Offset Transform not applied properly with time zero binding skeletal meshes

Thank you [REDACTED_HANDLE] for finding the fix.

[RN] Fixed Global Offset Transform not applied properly with time zero binding skeletal meshes
#tests Interchange test
#Synced-CL 52680506

diff --git a/Engine/Plugins/Interchange/Runtime/Source/Import/Private/Mesh/InterchangeSkeletalMeshFactory.cpp b/Engine/Plugins/Interchange/Runtime/Source/Import/Private/Mesh/InterchangeSkeletalMeshFactory.cpp
--- a/Engine/Plugins/Interchange/Runtime/Source/Import/Private/Mesh/InterchangeSkeletalMeshFactory.cpp
+++ b/Engine/Plugins/Interchange/Runtime/Source/Import/Private/Mesh/InterchangeSkeletalMeshFactory.cpp
@@ -265,19 +265,20 @@ namespace UE
 				TArray<double> SkinWeights;
 				SkinWeights.AddZeroed(VertexCount);
 
-				FTransform GlobalOffsetTransform = FTransform::Identity;
-				if (UInterchangeCommonPipelineDataFactoryNode* CommonPipelineDataFactoryNode = UInterchangeCommonPipelineDataFactoryNode::GetUniqueInstance(NodeContainer))
-				{
-					CommonPipelineDataFactoryNode->GetCustomGlobalOffsetTransform(GlobalOffsetTransform);
-				}
+				FMatrix SceneNodeMatrixInverse = SceneNodeTransform.ToInverseMatrixWithScale();
+				FMatrix SceneNodeMatrixInverseWithoutGlobalOffset = SceneNodeMatrixInverse;
 
-				//If we use transform, some scaling are not well supported, so we must do all the math with matrix
-				FMatrix SceneNodeMatrixInverse = SceneNodeTransform.ToMatrixWithScale().Inverse();
-				FMatrix GeometricMatrix = FMatrix::Identity;
-
-				//Scope any transform variable to be sure they are not use to create the vertex transformation matrix
+				if (UInterchangeCommonPipelineDataFactoryNode* CommonPipelineDataFactoryNode = UInterchangeCommonPipelineDataFactoryNode::GetUniqueInstance(NodeContainer))
 				{
-					GeometricMatrix = GeometricTransform.ToMatrixWithScale();
+					FTransform GlobalOffsetTransform = FTransform::Identity;
+					if (CommonPipelineDataFactoryNode->GetCustomGlobalOffsetTransform(GlobalOffsetTransform))
+					{
+						bool bBakeMeshes = false;
+						if (CommonPipelineDataFactoryNode->GetBakeMeshes(bBakeMeshes) && bBakeMeshes)
+						{
+							SceneNodeMatrixInverseWithoutGlobalOffset = (SceneNodeTransform * GlobalOffsetTransform.Inverse()).ToInverseMatrixWithScale();
+						}
+					}
 				}
 
 				const int32 JointCount = JointNames.Num();
@@ -329,13 +330,16 @@ namespace UE
 						}
 
 						//We must add the geometric node transform
-						JointReferenceMatrix = GeometricMatrix * JointReferenceMatrix;
+						JointReferenceMatrix = GeometricTransform.ToMatrixWithScale() * JointReferenceMatrix;
 					}
 
-					FMatrix ReferenceMatrixRelativeToBindPose = JointReferenceMatrix * JointBindPoseMatrix.Inverse();
-					FMatrix jointTimeZeroRelativeToMeshNodeMatrix = JointT0Matrix * SceneNodeMatrixInverse;
-					//Compute the transform to apply to vertex to change the bind pose from the existing one to use time zero pose instead
-					FMatrix VertexTransformMatrix = SceneNodeMatrixInverse * (ReferenceMatrixRelativeToBindPose * jointTimeZeroRelativeToMeshNodeMatrix);
+					FMatrix VertexTransformMatrix;
+					{
+						FMatrix ReferenceMatrixRelativeToBindPose = JointReferenceMatrix * JointBindPoseMatrix.Inverse();
+						FMatrix JointTimeZeroRelativeToMeshNodeMatrix = JointT0Matrix * SceneNodeMatrixInverseWithoutGlobalOffset;
+						//Compute the transform to apply to vertex to change the bind pose from the existing one to use time zero pose instead
+						VertexTransformMatrix = SceneNodeMatrixInverse * (ReferenceMatrixRelativeToBindPose * JointTimeZeroRelativeToMeshNodeMatrix);
+					}
 
 					//Iterate all bone vertices
 					for (FVertexID VertexID : MeshDescription.Vertices().GetElementIDs())
@@ -727,9 +731,7 @@ namespace UE
 			bool BuildMeshReferences(const UInterchangeFactoryBase::FImportAssetObjectParams& Arguments
 				, const FString& RootJointNodeId
 				, UInterchangeSkeletalMeshFactory::FImportAssetObjectLODData& ImportAssetObjectLODDataRef
-				, const FTransform& GlobalOffsetTransform
 				, const UInterchangeSkeletalMeshLodDataNode* LodDataNode
-				, const bool bBakeMeshes
 				, bool bImportSockets)
 			{
 				if (!LodDataNode)
@@ -948,9 +950,7 @@ void UInterchangeSkeletalMeshFactory::CreatePayloadTasks(const FImportAssetObjec
 			Arguments
 			, RootJointNodeId
 			, ImportAssetObjectLODData
-			, GlobalOffsetTransform
 			, LodDataNode
-			, bBakeMeshes
 			, bImportSockets
 		);
 
@@ -1332,15 +1332,6 @@ UInterchangeFactoryBase::FImportAssetResult UInterchangeSkeletalMeshFactory::Imp
 	TMap<FString, TArray<FLODUtilities::FMorphTargetLodBackupData>> BackupImportedMorphTargetData;
 	FLODUtilities::BackupCustomImportedMorphTargetData(SkeletalMesh, BackupImportedMorphTargetData);
 
-
-	FTransform GlobalOffsetTransform = FTransform::Identity;
-	bool bBakeMeshes = false;
-	if (UInterchangeCommonPipelineDataFactoryNode* CommonPipelineDataFactoryNode = UInterchangeCommonPipelineDataFactoryNode::GetUniqueInstance(Arguments.NodeContainer))
-	{
-		CommonPipelineDataFactoryNode->GetCustomGlobalOffsetTransform(GlobalOffsetTransform);
-		CommonPipelineDataFactoryNode->GetBakeMeshes(bBakeMeshes);
-	}
-
 	USkeleton* SkeletonReference = ImportAssetObjectData.SkeletonReference;
 		
 	FSkeletalMeshModel* ImportedResource = SkeletalMesh->GetImportedModel();
